AggregateTicking and Pipe merged into Engine/Plugins/Runtime/MassEntity

diff --git a/Engine/Plugins/Runtime/MassEntity/Source/MassEntityTestSuite/Public/MassEntityTestTypes.h b/Engine/Plugins/Runtime/MassEntity/Source/MassEntityTestSuite/Public/MassEntityTestTypes.h
new file mode 100644
--- /dev/null
+++ b/Engine/Plugins/Runtime/MassEntity/Source/MassEntityTestSuite/Public/MassEntityTestTypes.h
@@ -0,0 +1,177 @@
+// Copyright Epic Games, Inc. All Rights Reserved.
+
+#pragma once
+
+#include "MassProcessor.h"
+#include "LWComponentTypes.h"
+#include "EntitySubsystem.h"
+#include "AITestsCommon.h"
+#include "MassEntityTestTypes.generated.h"
+
+
+class UWorld;
+class UPipeEntitySubsystem;
+
+USTRUCT()
+struct FTestFragment_Float : public FLWComponentData
+{
+	GENERATED_BODY()
+	float Value = 0;
+};
+
+USTRUCT()
+struct FTestFragment_Int : public FLWComponentData
+{
+	GENERATED_BODY()
+	int32 Value = 0;
+};
+
+USTRUCT()
+struct FTestFragment_Bool : public FLWComponentData
+{
+	GENERATED_BODY()
+	bool bValue = false;
+};
+
+/** @todo rename to FTestTag */
+USTRUCT()
+struct FTestFragment_Tag : public FComponentTag
+{
+	GENERATED_BODY()
+};
+
+USTRUCT()
+struct FTestTag_A : public FComponentTag
+{
+	GENERATED_BODY()
+};
+
+USTRUCT()
+struct FTestTag_B : public FComponentTag
+{
+	GENERATED_BODY()
+};
+
+USTRUCT()
+struct FTestTag_C : public FComponentTag
+{
+	GENERATED_BODY()
+};
+
+USTRUCT()
+struct FTestTag_D : public FComponentTag
+{
+	GENERATED_BODY()
+};
+
+
+UCLASS()
+class UPipeTestProcessorBase : public UPipeProcessor
+{
+	GENERATED_BODY()
+public:
+	UPipeTestProcessorBase();
+	FPipeProcessorExecutionOrder& GetMutableExecutionOrder() { return ExecutionOrder; }
+	virtual void Execute(UEntitySubsystem& EntitySubsystem, FLWComponentSystemExecutionContext& Context) override 
+	{
+		ExecutionFunction(EntitySubsystem, Context);
+	}
+	virtual void ConfigureQueries() override
+	{
+		RequirementsFunction(EntityQuery);
+	}
+
+	TFunction<void(UEntitySubsystem& EntitySubsystem, FLWComponentSystemExecutionContext& Context)> ExecutionFunction;
+	TFunction<void(FLWComponentQuery& Query)> RequirementsFunction;
+
+	FLWComponentQuery& TestGetQuery() { return EntityQuery; }
+
+protected:
+	FLWComponentQuery EntityQuery;
+};
+
+UCLASS()
+class UPipeTestProcessor_A : public UPipeTestProcessorBase
+{
+	GENERATED_BODY()
+};
+
+UCLASS()
+class UPipeTestProcessor_B : public UPipeTestProcessorBase
+{
+	GENERATED_BODY()
+};
+
+UCLASS()
+class UPipeTestProcessor_C : public UPipeTestProcessorBase
+{
+	GENERATED_BODY()
+};
+
+UCLASS()
+class UPipeTestProcessor_D : public UPipeTestProcessorBase
+{
+	GENERATED_BODY()
+};
+
+UCLASS()
+class UPipeTestProcessor_E : public UPipeTestProcessorBase
+{
+	GENERATED_BODY()
+};
+
+UCLASS()
+class UPipeTestProcessor_F : public UPipeTestProcessorBase
+{
+	GENERATED_BODY()
+};
+
+UCLASS()
+class UPipeTestProcessor_Floats : public UPipeTestProcessorBase
+{
+	GENERATED_BODY()
+public:
+	TArrayView<FTestFragment_Float> Floats;
+	UPipeTestProcessor_Floats();
+};
+
+UCLASS()
+class UPipeTestProcessor_Ints : public UPipeTestProcessorBase
+{
+	GENERATED_BODY()
+public:
+	TArrayView<FTestFragment_Int> Ints;
+	UPipeTestProcessor_Ints();
+};
+
+UCLASS()
+class UPipeTestProcessor_FloatsInts : public UPipeTestProcessorBase
+{
+	GENERATED_BODY()
+public:
+	TArrayView<FTestFragment_Float> Floats;
+	TArrayView<FTestFragment_Int> Ints;
+	UPipeTestProcessor_FloatsInts();
+};
+
+struct FExecutionTestBase : FAITestBase
+{
+	UPipeEntitySubsystem* EntitySubsystem = nullptr;
+	UWorld* World = nullptr;
+
+	virtual bool SetUp() override;
+};
+
+const int TestIntValue = 123456;
+
+struct FEntityTestBase : FExecutionTestBase
+{
+	FArchetypeHandle EmptyArchetype;
+	FArchetypeHandle FloatsArchetype;
+	FArchetypeHandle IntsArchetype;
+	FArchetypeHandle FloatsIntsArchetype;
+
+	FInstancedStruct InstanceInt;
+
+	virtual bool SetUp() override;
+};
diff --git a/Engine/Plugins/Runtime/MassGameplay/Source/MassActors/Private/MassActorHelper.cpp b/Engine/Plugins/Runtime/MassGameplay/Source/MassActors/Private/MassActorHelper.cpp
new file mode 100644
--- /dev/null
+++ b/Engine/Plugins/Runtime/MassGameplay/Source/MassActors/Private/MassActorHelper.cpp
@@ -0,0 +1,60 @@
+// Copyright Epic Games, Inc.All Rights Reserved.
+
+#include "MassActorHelper.h"
+#include "MassActorTypes.h"
+#include "MassActorSubsystem.h"
+#include "MassEntitySystem.h"
+#include "Engine/World.h"
+
+namespace UE::MassActor
+{
+	bool AddEntityTagToActor(const AActor& Actor, const UScriptStruct& TagType)
+	{
+		UWorld* World = Actor.GetWorld();
+		UPipeEntitySubsystem* EntitySystem = UPipeEntitySubsystem::GetCurrent(World);
+		UMassActorSubsystem* MassActorSubsystem = UWorld::GetSubsystem<UMassActorSubsystem>(World);
+		if (EntitySystem && MassActorSubsystem)
+		{
+			const FMassHandle AgentHandle = MassActorSubsystem->GetHandleFromActor(&Actor);
+			if (AgentHandle.IsValid())
+			{
+				EntitySystem->AddTagToEntity(AgentHandle.GetLWEntity(), &TagType);
+				return true;
+			}
+			UE_LOG(LogMassActor, Warning, TEXT("Failed to add tag %s to actor %s due to it not having an associated entity")
+				, *TagType.GetName(), *Actor.GetName());
+		}
+		else
+		{
+			UE_LOG(LogMassActor, Warning, TEXT("Failed to add tag %s to actor %s due to missing: %s%s")
+				, *TagType.GetName(), *Actor.GetName(), EntitySystem ? TEXT("EntitySubsystem, ") : TEXT("")
+				, MassActorSubsystem ? TEXT("MassActorSubsystem, ") : TEXT(""));
+		}
+		return false;
+	}
+
+	bool RemoveEntityTagFromActor(const AActor& Actor, const UScriptStruct& TagType)
+	{
+		UWorld* World = Actor.GetWorld();
+		UPipeEntitySubsystem* EntitySystem = UPipeEntitySubsystem::GetCurrent(World);
+		UMassActorSubsystem* MassActorSubsystem = UWorld::GetSubsystem<UMassActorSubsystem>(World);
+		if (EntitySystem && MassActorSubsystem)
+		{
+			const FMassHandle AgentHandle = MassActorSubsystem->GetHandleFromActor(&Actor);
+			if (AgentHandle.IsValid())
+			{
+				EntitySystem->RemoveTagFromEntity(AgentHandle.GetLWEntity(), &TagType);
+				return true;
+			}
+			UE_LOG(LogMassActor, Warning, TEXT("Failed to remove tag %s from actor %s due to it not having an associated entity")
+				, *TagType.GetName(), *Actor.GetName());
+		}
+		else
+		{
+			UE_LOG(LogMassActor, Warning, TEXT("Failed to remove tag %s from actor %s due to missing: %s%s")
+				, *TagType.GetName(), *Actor.GetName(), EntitySystem ? TEXT("EntitySubsystem, ") : TEXT("")
+				, MassActorSubsystem ? TEXT("MassActorSubsystem, ") : TEXT(""));
+		}
+		return false;
+	}
+} // namespace UE::MassActor
diff --git a/Engine/Plugins/Runtime/MassGameplay/Source/MassActors/Private/MassActorSpawnerSubsystem.cpp b/Engine/Plugins/Runtime/MassGameplay/Source/MassActors/Private/MassActorSpawnerSubsystem.cpp
new file mode 100644
--- /dev/null
+++ b/Engine/Plugins/Runtime/MassGameplay/Source/MassActors/Private/MassActorSpawnerSubsystem.cpp
@@ -0,0 +1,426 @@
+// Copyright Epic Games, Inc. All Rights Reserved.
+
+#include "MassActorSpawnerSubsystem.h"
+#include "MassAgentComponent.h"
+#include "MassSimulationSettings.h"
+#include "Engine/World.h"
+#include "MassActorTypes.h"
+#include "VisualLogger/VisualLogger.h"
+#include "MassActorPoolableInterface.h"
+#include "MassSimulationSubsystem.h"
+#include "Logging/LogScopedVerbosityOverride.h"
+
+namespace UE::MassActors
+{
+	int32 bUseActorPooling = 1;
+	FAutoConsoleVariableRef CVarDebugRepresentationLOD(TEXT("ai.mass.actorpooling"), bUseActorPooling, TEXT("Use Actor Pooling"), ECVF_Scalability);
+}
+
+//----------------------------------------------------------------------//
+// UMassActorSpawnerSubsystem 
+//----------------------------------------------------------------------//
+void UMassActorSpawnerSubsystem::RetryActorSpawnRequest(FMassHandle_ActorSpawnRequest SpawnRequestHandle)
+{
+	check(SpawnRequestHandleManager.IsValidHandle(SpawnRequestHandle));
+	const int32 Index = SpawnRequestHandle.GetIndex();
+	check(SpawnRequests.IsValidIndex(Index));
+	FMassActorSpawnRequest& SpawnRequest = SpawnRequests[SpawnRequestHandle.GetIndex()].GetMutable<FMassActorSpawnRequest>();
+	if (ensureMsgf(SpawnRequest.SpawnStatus == ESpawnRequestStatus::Failed, TEXT("Can only retry failed spawn requests")))
+	{
+		UWorld* World = GetWorld();
+		check(World);
+
+		SpawnRequest.SpawnStatus = ESpawnRequestStatus::RetryPending;
+		SpawnRequest.SerialNumber = RequestSerialNumberCounter.fetch_add(1);
+		SpawnRequest.RequestedTime = World->GetTimeSeconds();
+	}
+}
+
+bool UMassActorSpawnerSubsystem::RemoveActorSpawnRequest(FMassHandle_ActorSpawnRequest& SpawnRequestHandle)
+{
+	if (!ensureMsgf(SpawnRequestHandleManager.RemoveHandle(SpawnRequestHandle), TEXT("Invalid spawn request handle")))
+	{
+		return false;
+	}
+
+	check(SpawnRequests.IsValidIndex(SpawnRequestHandle.GetIndex()));
+	FMassActorSpawnRequest& SpawnRequest = SpawnRequests[SpawnRequestHandle.GetIndex()].GetMutable<FMassActorSpawnRequest>();
+	check(SpawnRequest.SpawnStatus != ESpawnRequestStatus::Processing);
+	SpawnRequestHandle.Invalidate();
+	SpawnRequest.Reset();
+	return true;
+}
+
+void UMassActorSpawnerSubsystem::DestroyActor(AActor* Actor, bool bImmediate /*= false*/)
+{
+	// We need to unregister immediately MassAgentComponent as it will become out of sync with mass
+	if (UMassAgentComponent* AgentComp = Actor->FindComponentByClass<UMassAgentComponent>())
+	{
+		if (AgentComp->IsRegistered())
+		{
+			AgentComp->UnregisterComponent();
+		}
+	}
+
+	if(bImmediate)
+	{
+		if (!ReleaseActorToPool(Actor))
+		{
+			UWorld* World = GetWorld();
+			check(World);
+
+			World->DestroyActor(Actor);
+		}
+	}
+	else
+	{
+		ActorsToDestroy.Add(Actor);
+	}
+}
+
+bool UMassActorSpawnerSubsystem::ReleaseActorToPool(AActor* Actor)
+{
+	if (!UE::MassActors::bUseActorPooling || !bActorPoolingEnabled)
+	{
+		return false;
+	}
+
+	const bool bIsPoolableActor = Actor->Implements<UMassActorPoolableInterface>();
+	if (bIsPoolableActor && IMassActorPoolableInterface::Execute_CanBePooled(Actor))
+	{
+		UMassAgentComponent* AgentComp = Actor->FindComponentByClass<UMassAgentComponent>();
+		if (!AgentComp || AgentComp->IsReadyForPooling())
+		{
+			IMassActorPoolableInterface::Execute_PrepareForPooling(Actor);
+			Actor->SetActorHiddenInGame(true);
+			if (AgentComp)
+			{
+				AgentComp->UnregisterWithAgentManager();
+			}
+
+			TArray<AActor*>& Pool = PooledActors.FindOrAdd(Actor->GetClass());
+			checkf(Pool.Find(Actor) == INDEX_NONE, TEXT("Actor%s is already in the pool"), *AActor::GetDebugName(Actor));
+			Pool.Add(Actor);
+			return true;
+		}
+	}
+	return false;
+}
+
+FMassHandle_ActorSpawnRequest UMassActorSpawnerSubsystem::RequestActorSpawnInternal(const FConstStructView SpawnRequestView)
+{
+	// The handle manager has a freelist of the release indexes, so it can return us a index that we previously used.
+	const FMassHandle_ActorSpawnRequest SpawnRequestHandle = SpawnRequestHandleManager.GetNextHandle();
+	const int32 Index = SpawnRequestHandle.GetIndex();
+
+	// Check if we need to grow the array, otherwise it is a previously released index that was returned.
+	if (!SpawnRequests.IsValidIndex(Index))
+	{
+		checkf(SpawnRequests.Num() == Index, TEXT("This case should only be when we need to grow the array of one element."));
+		SpawnRequests.Add(SpawnRequestView);
+	}
+	else
+	{
+		SpawnRequests[Index] = SpawnRequestView;
+	}
+
+	UWorld* World = GetWorld();
+	check(World);
+
+	// Initialize the spawn request status
+	FMassActorSpawnRequest& SpawnRequest = GetMutableSpawnRequest<FMassActorSpawnRequest>(SpawnRequestHandle);
+	SpawnRequest.SpawnStatus = ESpawnRequestStatus::Pending;
+	SpawnRequest.SerialNumber = RequestSerialNumberCounter.fetch_add(1);
+	SpawnRequest.RequestedTime = World->GetTimeSeconds();
+
+	return SpawnRequestHandle;
+}
+
+FMassHandle_ActorSpawnRequest UMassActorSpawnerSubsystem::GetNextRequestToSpawn() const
+{
+	FMassHandle_ActorSpawnRequest BestSpawnRequestHandle;
+	float BestPriority = MAX_FLT;
+	bool bBestIsPending = false;
+	uint32 BestSerialNumber = MAX_uint32;
+	for (const FMassHandle_ActorSpawnRequest SpawnRequestHandle : SpawnRequestHandleManager.GetHandles())
+	{
+		if (!SpawnRequestHandle.IsValid())
+		{
+			continue;
+		}
+		const FMassActorSpawnRequest& SpawnRequest = GetSpawnRequest<FMassActorSpawnRequest>(SpawnRequestHandle);
+		if (SpawnRequest.SpawnStatus == ESpawnRequestStatus::Pending)
+		{
+			if (!bBestIsPending ||
+				SpawnRequest.Priority < BestPriority ||
+				(SpawnRequest.Priority == BestPriority && SpawnRequest.SerialNumber < BestSerialNumber))
+			{
+				BestSpawnRequestHandle = SpawnRequestHandle;
+				BestSerialNumber = SpawnRequest.SerialNumber;
+				BestPriority = SpawnRequest.Priority;
+				bBestIsPending = true;
+			}
+		}
+		else if (!bBestIsPending && SpawnRequest.SpawnStatus == ESpawnRequestStatus::RetryPending)
+		{
+			// No priority on retries just FIFO
+			if (SpawnRequest.SerialNumber < BestSerialNumber)
+			{
+				BestSpawnRequestHandle = SpawnRequestHandle;
+				BestSerialNumber = SpawnRequest.SerialNumber;
+			}
+		}
+	}
+
+	return BestSpawnRequestHandle;
+}
+
+AActor* UMassActorSpawnerSubsystem::SpawnOrRetrieveFromPool(const FStructView& SpawnRequestView)
+{
+	if (UE::MassActors::bUseActorPooling != 0 && bActorPoolingEnabled)
+	{
+		const FMassActorSpawnRequest& SpawnRequest = SpawnRequestView.Get<FMassActorSpawnRequest>();
+		TArray<AActor*>* Pool = PooledActors.Find(SpawnRequest.Template);
+
+		if (Pool && Pool->Num() > 0)
+		{
+			AActor* PooledActor = (*Pool)[0];
+			Pool->RemoveAt(0);
+			PooledActor->SetActorHiddenInGame(false);
+			PooledActor->SetActorTransform(SpawnRequest.Transform, false, nullptr, ETeleportType::ResetPhysics);
+
+			IMassActorPoolableInterface::Execute_PrepareForGame(PooledActor);
+
+			if (UMassAgentComponent* AgentComp = PooledActor->FindComponentByClass<UMassAgentComponent>())
+			{
+				AgentComp->RegisterWithAgentManager();
+				AgentComp->SetPuppetHandle(SpawnRequest.MassAgent);
+			}
+
+			return PooledActor;
+		}
+	}
+
+	return SpawnActor(SpawnRequestView);
+}
+
+AActor* UMassActorSpawnerSubsystem::SpawnActor(const FStructView& SpawnRequestView) const
+{
+	TRACE_CPUPROFILER_EVENT_SCOPE(UMassActorSpawnerSubsystem::SpawnActor);
+
+	UWorld* World = GetWorld();
+	check(World);
+
+	const FMassActorSpawnRequest& SpawnRequest = SpawnRequestView.Get<FMassActorSpawnRequest>();
+	if (AActor* SpawnedActor = World->SpawnActorDeferred<AActor>(SpawnRequest.Template, SpawnRequest.Transform, nullptr, nullptr, ESpawnActorCollisionHandlingMethod::AlwaysSpawn))
+	{
+		// Add code here before construction script
+
+		SpawnedActor->FinishSpawning(SpawnRequest.Transform);
+		// The finish spawning might have failed and the spawned actor is destroyed.
+		if (IsValidChecked(SpawnedActor))
+		{
+			if (UMassAgentComponent* AgentComp = SpawnedActor->FindComponentByClass<UMassAgentComponent>())
+			{
+				AgentComp->SetPuppetHandle(SpawnRequest.MassAgent);
+			}
+			return SpawnedActor;
+		}
+	}
+
+	UE_VLOG_CAPSULE(this, LogMassActor, Error,
+					SpawnRequest.Transform.GetLocation(),
+					SpawnRequest.Template.GetDefaultObject()->GetSimpleCollisionHalfHeight(),
+					SpawnRequest.Template.GetDefaultObject()->GetSimpleCollisionRadius(),
+					SpawnRequest.Transform.GetRotation(),
+					FColor::Red,
+					TEXT("Unable to spawn actor for Mass entity [%s]"), *SpawnRequest.MassAgent.GetLWEntity().DebugGetDescription());
+	return nullptr;
+}
+
+void UMassActorSpawnerSubsystem::ProcessPendingSpawningRequest(const double MaxTimeSlicePerTick)
+{
+	TRACE_CPUPROFILER_EVENT_SCOPE(UMassActorSpawnerSubsystem::ProcessPendingSpawningRequest);
+	SpawnRequestHandleManager.ShrinkHandles();
+
+	const double TimeSliceEnd = FPlatformTime::Seconds() + MaxTimeSlicePerTick;
+
+	while (FPlatformTime::Seconds() < TimeSliceEnd)
+	{
+		FMassHandle_ActorSpawnRequest SpawnRequestHandle = GetNextRequestToSpawn();
+		if (!SpawnRequestHandle.IsValid() ||
+			!ensureMsgf(SpawnRequestHandleManager.IsValidHandle(SpawnRequestHandle), TEXT("GetNextRequestToSpawn returned an invalid handle, expecting an empty one or a valid one.")))
+		{
+			return;
+		}
+
+		FStructView SpawnRequestView = SpawnRequests[SpawnRequestHandle.GetIndex()];
+		FMassActorSpawnRequest& SpawnRequest = SpawnRequestView.GetMutable<FMassActorSpawnRequest>();
+
+		if (!ensureMsgf(SpawnRequest.SpawnStatus == ESpawnRequestStatus::Pending ||
+						SpawnRequest.SpawnStatus == ESpawnRequestStatus::RetryPending, TEXT("GetNextRequestToSpawn returned a request that was already processed, need to return only request with pending status.")))
+		{
+			return;
+		}
+
+		// Do the spawning
+		SpawnRequest.SpawnStatus = ESpawnRequestStatus::Processing;
+
+		// Call the pre spawn delegate on the spawn request
+		if (SpawnRequest.ActorPreSpawnDelegate.IsBound())
+		{
+			SpawnRequest.ActorPreSpawnDelegate.Execute(SpawnRequestHandle, SpawnRequestView);
+		}
+
+		SpawnRequest.SpawnedActor = SpawnOrRetrieveFromPool(SpawnRequestView);
+
+		SpawnRequest.SpawnStatus = SpawnRequest.SpawnedActor ? ESpawnRequestStatus::Succeeded : ESpawnRequestStatus::Failed;
+
+		// Call the post spawn delegate on the spawn request
+		if (SpawnRequest.ActorPostSpawnDelegate.IsBound())
+		{
+			if (SpawnRequest.ActorPostSpawnDelegate.Execute(SpawnRequestHandle, SpawnRequestView) == EMassActorSpawnRequestAction::Remove)
+			{
+				// If notified, remove the spawning request
+				ensureMsgf(SpawnRequestHandleManager.RemoveHandle(SpawnRequestHandle), TEXT("When providing a delegate, the spawn request gets automatically removed, no need to remove it on your side"));
+			}
+		}
+	}
+}
+
+void UMassActorSpawnerSubsystem::ProcessPendingDestruction(const double MaxTimeSlicePerTick)
+{
+	TRACE_CPUPROFILER_EVENT_SCOPE(UMassActorSpawnerSubsystem::ProcessPendingDestruction);
+
+	UWorld* World = GetWorld();
+	check(World);
+
+	const ENetMode CurrentWorldNetMode = World->GetNetMode();
+	const double HasToDestroyAllActorsOnServerSide = CurrentWorldNetMode != NM_Client && CurrentWorldNetMode != NM_Standalone;
+	const double TimeSliceEnd = FPlatformTime::Seconds() + MaxTimeSlicePerTick;
+
+	{
+		// Try release to pool actors or destroy them
+		TRACE_CPUPROFILER_EVENT_SCOPE(DestroyActors);
+		while ((DeactivatedActorsToDestroy.Num() || ActorsToDestroy.Num()) && 
+			   (HasToDestroyAllActorsOnServerSide || FPlatformTime::Seconds() <= TimeSliceEnd))
+		{
+			AActor* ActorToDestroy = DeactivatedActorsToDestroy.Num() ? DeactivatedActorsToDestroy.Pop(/*bAllowShrinking*/false) : ActorsToDestroy.Pop(/*bAllowShrinking*/false);
+			if (!ReleaseActorToPool(ActorToDestroy))
+			{
+				// Couldn't release actor back to pool, so destroy it
+				World->DestroyActor(ActorToDestroy);
+			}
+		}
+	}
+
+	if(ActorsToDestroy.Num())
+	{
+		// Try release to pool remaining actors or deactivate them
+		TRACE_CPUPROFILER_EVENT_SCOPE(DeactivateActors);
+		for (AActor* ActorToDestroy : ActorsToDestroy)
+		{
+			if (!ReleaseActorToPool(ActorToDestroy))
+			{
+				// Couldn't release actor back to pool, do simple deactivate instead
+				ActorToDestroy->SetActorEnableCollision(false);
+				ActorToDestroy->SetActorHiddenInGame(true);
+				ActorToDestroy->SetActorTickEnabled(false);
+				DeactivatedActorsToDestroy.Add(ActorToDestroy);
+			}
+		}
+		ActorsToDestroy.Reset();
+	}
+}
+
+void UMassActorSpawnerSubsystem::Initialize(FSubsystemCollectionBase& Collection)
+{
+	Collection.InitializeDependency(UMassSimulationSubsystem::StaticClass());
+	Super::Initialize(Collection);
+
+	if (const UWorld* World = GetWorld())
+	{
+		UMassSimulationSubsystem* SimSystem = UWorld::GetSubsystem<UMassSimulationSubsystem>(World);
+		check(SimSystem);
+		SimSystem->GetOnProcessingPhaseStarted(EPipeProcessingPhase::PrePhysics).AddUObject(this, &UMassActorSpawnerSubsystem::OnPrePhysicsPhaseStarted);
+		SimSystem->GetOnProcessingPhaseFinished(EPipeProcessingPhase::PrePhysics).AddUObject(this, &UMassActorSpawnerSubsystem::OnPrePhysicsPhaseFinished);
+	}
+}
+
+void UMassActorSpawnerSubsystem::Deinitialize()
+{
+	Super::Deinitialize();
+
+	if (const UWorld* World = GetWorld())
+	{
+		if (const UMassSimulationSubsystem* SimSystem = UWorld::GetSubsystem<UMassSimulationSubsystem>(World))
+		{
+			SimSystem->GetOnProcessingPhaseStarted(EPipeProcessingPhase::PrePhysics).RemoveAll(this);
+			SimSystem->GetOnProcessingPhaseFinished(EPipeProcessingPhase::PrePhysics).RemoveAll(this);
+		}
+	}
+}
+
+void UMassActorSpawnerSubsystem::OnPrePhysicsPhaseStarted(const float DeltaSeconds)
+{
+	// Spawn actors before processors run so they can operate on deterministic actor state for the frame
+	// 
+	// Note: MassRepresentationProcessor relies on actor spawns being processed before it runs so it can confirm the
+	// spawn and clean up the previous representation  
+	ProcessPendingSpawningRequest(GET_MASS_CONFIG_VALUE(DesiredActorSpawningTimeSlicePerTick));
+}
+
+void UMassActorSpawnerSubsystem::OnPrePhysicsPhaseFinished(const float DeltaSeconds)
+{
+	// Destroy any actors queued for destruction this frame and hide any we didn't get to within the max processing time
+	// 
+	// Note: MassRepresentationProcessor relies on actor destruction processing after it runs so it can clean up
+	// unwanted actor representations that it has replaced for this frame. It also relies on this running before physics
+	// so unwanted representations don't interfere with new physics enabled actors 
+	ProcessPendingDestruction(GET_MASS_CONFIG_VALUE(DesiredActorDestructionTimeSlicePerTick));
+}
+
+void UMassActorSpawnerSubsystem::AddReferencedObjects(UObject* InThis, FReferenceCollector& Collector)
+{
+	UMassActorSpawnerSubsystem* MASS = Cast<UMassActorSpawnerSubsystem>(InThis);
+	if (MASS)
+	{
+		for (auto It = MASS->PooledActors.CreateIterator(); It; ++It)
+		{
+			TArray<AActor*>& Value = It.Value();
+
+			Collector.AddReferencedObjects<AActor>(Value);
+		}
+	}
+
+	Super::AddReferencedObjects(InThis, Collector);
+}
+
+void UMassActorSpawnerSubsystem::EnableActorPooling() 
+{ 
+	bActorPoolingEnabled = true; 
+}
+
+void UMassActorSpawnerSubsystem::DisableActorPooling() 
+{
+	bActorPoolingEnabled = false;
+
+	ReleaseAllResources();
+
+}
+
+void UMassActorSpawnerSubsystem::ReleaseAllResources()
+{
+	if (UWorld* World = GetWorld())
+	{
+		for (auto It = PooledActors.CreateIterator(); It; ++It)
+		{
+			TArray<AActor*>& ActorArray = It.Value();
+			for (int i = 0; i < ActorArray.Num(); i++)
+			{
+				World->DestroyActor(ActorArray[i]);
+			}
+		}
+	}
+	PooledActors.Empty();
+}
diff --git a/Engine/Plugins/Runtime/MassGameplay/Source/MassActors/Private/MassActorSubsystem.cpp b/Engine/Plugins/Runtime/MassGameplay/Source/MassActors/Private/MassActorSubsystem.cpp
new file mode 100644
--- /dev/null
+++ b/Engine/Plugins/Runtime/MassGameplay/Source/MassActors/Private/MassActorSubsystem.cpp
@@ -0,0 +1,134 @@
+// Copyright Epic Games, Inc. All Rights Reserved.
+
+#include "MassActorSubsystem.h"
+#include "MassCommonTypes.h"
+#include "Engine/World.h"
+#include "MassEntitySystem.h"
+#include "MassSimulationSubsystem.h"
+#include "VisualLogger/VisualLogger.h"
+
+//----------------------------------------------------------------------//
+//  FDataFragment_Actor 
+//----------------------------------------------------------------------//
+
+void FDataFragment_Actor::SetAndUpdateHandleMap(const FMassHandle MassAgent, AActor* InActor, const bool bInIsOwnedByMass)
+{
+	SetNoHandleMapUpdate(MassAgent, InActor, bInIsOwnedByMass);
+
+	UWorld* World = InActor->GetWorld();
+	check(World);
+	if (UMassActorSubsystem* MassActorSubsystem = UWorld::GetSubsystem<UMassActorSubsystem>(World))
+	{
+		MassActorSubsystem->SetHandleForActor(InActor, MassAgent);
+	}
+}
+
+void FDataFragment_Actor::ResetAndUpdateHandleMap()
+{
+	if (AActor* ActorPtr = Cast<AActor>(Actor.Get()))
+	{
+		UWorld* World = Actor->GetWorld();
+		check(World);
+		if (UMassActorSubsystem* MassActorSubsystem = UWorld::GetSubsystem<UMassActorSubsystem>(World))
+		{
+			MassActorSubsystem->RemoveHandleForActor(ActorPtr);
+		}
+	}
+
+	ResetNoHandleMapUpdate();
+}
+
+void FDataFragment_Actor::SetNoHandleMapUpdate(const FMassHandle MassAgent, AActor* InActor, const bool bInIsOwnedByMass)
+{
+	check(InActor);
+	check(!Actor.IsValid());
+	check(MassAgent.IsValid());
+	Actor = InActor;
+	bIsOwnedByMass = bInIsOwnedByMass;
+}
+
+void FDataFragment_Actor::ResetNoHandleMapUpdate()
+{
+	Actor.Reset();
+	bIsOwnedByMass = false;
+}
+
+//----------------------------------------------------------------------//
+//  UMassActorSubsystem 
+//----------------------------------------------------------------------//
+void UMassActorSubsystem::Initialize(FSubsystemCollectionBase& Collection)
+{
+	Super::Initialize(Collection);
+	// making sure UMassSimulationSubsystem gets created before the MassActorManager
+	Collection.InitializeDependency<UMassSimulationSubsystem>();
+	
+	EntitySystem = UPipeEntitySubsystem::GetCurrent(GetWorld());
+}
+
+FMassHandle UMassActorSubsystem::GetHandleFromActor(const TObjectKey<const AActor> Actor)
+{
+	UE_MT_SCOPED_READ_ACCESS(ActorHandleMapDetector);
+	FLWEntity* LWEntity = ActorHandleMap.Find(Actor);
+	if (!LWEntity)
+	{
+		return FMassHandle::InvalidHandle;
+	}
+
+	check(TObjectKey<const AActor>(GetActorFromHandle(FMassHandle(*LWEntity))) == Actor);
+	return FMassHandle(*LWEntity);
+}
+
+AActor* UMassActorSubsystem::GetActorFromHandle(FMassHandle Handle) const
+{
+	check(EntitySystem);
+	FDataFragment_Actor* Data = EntitySystem->GetComponentDataPtr<FDataFragment_Actor>(Handle.GetLWEntity());
+	return Data != nullptr ? Data->GetMutable() : nullptr;
+}
+
+void UMassActorSubsystem::SetHandleForActor(const TObjectKey<const AActor> Actor, FMassHandle Handle)
+{
+	UE_MT_SCOPED_WRITE_ACCESS(ActorHandleMapDetector);
+	ActorHandleMap.Add(Actor, Handle.GetLWEntity());
+}
+
+void UMassActorSubsystem::RemoveHandleForActor(const TObjectKey<const AActor> Actor)
+{
+	UE_MT_SCOPED_WRITE_ACCESS(ActorHandleMapDetector);
+	ActorHandleMap.Remove(Actor);
+}
+
+void UMassActorSubsystem::DisconnectActor(const TObjectKey<const AActor> Actor, FMassHandle Handle)
+{
+	if (Handle.IsValid() == false)
+	{
+		return;
+	}
+
+	FLWEntity LWEntity;
+	{
+		UE_MT_SCOPED_WRITE_ACCESS(ActorHandleMapDetector);
+		// We're assuming the Handle does match Actor, so we're RemoveAndCopyValue. If if doesn't we'll add it back.
+		// The expectation is that this won't happen on a regular basis..
+		if (ActorHandleMap.RemoveAndCopyValue(Actor, LWEntity) == false)
+		{
+			// the entity doesn't match the actor
+			return;
+		}
+	}
+
+	if (LWEntity == Handle.GetLWEntity())
+	{
+		check(EntitySystem);
+		if (FDataFragment_Actor* Data = EntitySystem->GetComponentDataPtr<FDataFragment_Actor>(Handle.GetLWEntity()))
+		{
+			Data->ResetAndUpdateHandleMap();
+		}
+	}
+	else
+	{
+		// unexpected mismatch. Add back and notify.
+		UE_VLOG_UELOG(this, LogMass, Warning, TEXT("%s: Trying to disconnect actor %s while the Handle given doesn't match the system\'s records")
+			, ANSI_TO_TCHAR(__FUNCTION__), *AActor::GetDebugName(Actor.ResolveObjectPtr()));
+		SetHandleForActor(Actor, Handle);
+	}
+}
\ No newline at end of file
diff --git a/Engine/Plugins/Runtime/MassGameplay/Source/MassActors/Private/MassAgentSubsystem.cpp b/Engine/Plugins/Runtime/MassGameplay/Source/MassActors/Private/MassAgentSubsystem.cpp
new file mode 100644
--- /dev/null
+++ b/Engine/Plugins/Runtime/MassGameplay/Source/MassActors/Private/MassAgentSubsystem.cpp
@@ -0,0 +1,433 @@
+// Copyright Epic Games, Inc. All Rights Reserved.
+
+#include "MassAgentSubsystem.h"
+#include "LWCCommandBuffer.h"
+#include "Engine/World.h"
+#include "MassEntitySystem.h"
+#include "MassCommonTypes.h"
+#include "MassSimulationSubsystem.h"
+#include "MassSpawnerSubsystem.h"
+#include "MassActorTypes.h"
+#include "MassAgentComponent.h"
+#include "MassEntityTemplateRegistry.h"
+#include "VisualLogger/VisualLogger.h"
+#include "EntityView.h"
+#include "MassReplicationManager.h"
+#include "Engine/NetDriver.h"
+
+
+namespace FMassAgentSubsystemHelper
+{
+
+inline void InitializeAgentComponentFragments(const UMassAgentComponent& AgentComp, FEntityView& EntityView, const EMassTranslationDirection Direction, TConstArrayView<FMassEntityTemplate::FObjectFragmentInitializerFunction> ObjectFragmentInitializers)
+{
+	AActor* Owner = AgentComp.GetOwner();
+	check(Owner);
+	for (const FMassEntityTemplate::FObjectFragmentInitializerFunction& Initializer : ObjectFragmentInitializers)
+	{
+		Initializer(*Owner, EntityView, Direction);
+	}
+}
+
+}
+
+//----------------------------------------------------------------------//
+// UMassAgentSubsystem 
+//----------------------------------------------------------------------//
+void UMassAgentSubsystem::Initialize(FSubsystemCollectionBase& Collection)
+{
+	Super::Initialize(Collection);
+
+	// making sure UMassSimulationSubsystem gets created before the MassAgentManager
+	Collection.InitializeDependency<UMassSimulationSubsystem>();
+	Collection.InitializeDependency<UMassSpawnerSubsystem>();
+#if UE_REPLICATION_COMPILE_CLIENT_CODE
+	Collection.InitializeDependency<UMassReplicationManager>();
+#endif // UE_REPLICATION_COMPILE_CLIENT_CODE
+	
+	UWorld* World = GetWorld();
+	EntitySystem = UPipeEntitySubsystem::GetCurrent(World);
+	SpawnerSystem = UWorld::GetSubsystem<UMassSpawnerSubsystem>(World);
+
+	SimulationSystem = UWorld::GetSubsystem<UMassSimulationSubsystem>(World);
+	SimulationSystem->GetOnProcessingPhaseStarted(EPipeProcessingPhase::PrePhysics).AddUObject(this, &UMassAgentSubsystem::OnProcessingPhaseStarted, EPipeProcessingPhase::PrePhysics);
+
+#if UE_REPLICATION_COMPILE_CLIENT_CODE
+	ReplicationManager = UWorld::GetSubsystem<UMassReplicationManager>(World);
+	check(ReplicationManager);
+
+	ReplicationManager->GetOnMassAgentAdded().AddUObject(this, &UMassAgentSubsystem::OnMassAgentAddedToReplication);
+	ReplicationManager->GetOnRemovingMassAgent().AddUObject(this, &UMassAgentSubsystem::OnMassAgentRemovedFromReplication);
+#endif // UE_REPLICATION_COMPILE_CLIENT_CODE
+}
+
+FMassEntityTemplateID UMassAgentSubsystem::RegisterAgentComponent(UMassAgentComponent& AgentComp)
+{
+	check(EntitySystem);
+	check(SpawnerSystem);
+
+	if (AgentComp.IsPuppet())
+	{
+		MakePuppet(AgentComp);
+		return AgentComp.GetTemplateID();
+	}
+
+	// @todo note that this will happen when a regular actor with AgentComp gets unregister and then registered again.
+	if (AgentComp.GetTemplateID().IsValid())
+	{
+		UE_VLOG_UELOG(this, LogMassActor, Warning, TEXT("UMassAgentSubsystem::RegisterAgentComponent called while the given agent component has already been registered (Owner: %s, entity handle %s, template ID %s)")
+			, *GetNameSafe(AgentComp.GetOwner()), *AgentComp.GetEntityHandle().DebugGetDescription(), *AgentComp.GetTemplateID().ToString());
+		return FMassEntityTemplateID();
+	}
+
+	AActor* AgentActor = AgentComp.GetOwner();
+	check(AgentActor);
+
+	const FMassEntityConfig& EntityConfig = AgentComp.GetEntityConfig();
+	const FMassEntityTemplate* EntityTemplate = EntityConfig.GetOrCreateEntityTemplate(*AgentActor, AgentComp);
+
+	if (!ensureMsgf(EntityTemplate, TEXT("Failed to fetch entity template for %s."), *AgentActor->GetName()))
+	{
+		UE_VLOG(this, LogMassActor, Warning, TEXT("Failed to fetch entity template for %s."), *AgentActor->GetName());
+		return FMassEntityTemplateID();
+	}
+
+#if UE_REPLICATION_COMPILE_CLIENT_CODE
+	if (AgentComp.IsNetSimulating())
+	{
+		AgentComp.PuppetReplicationPending();
+	}
+	else
+#endif // UE_REPLICATION_COMPILE_CLIENT_CODE
+	{
+		FMassAgentInitializationQueue& AgentQueue = PendingAgentEntities.FindOrAdd(EntityTemplate->GetTemplateID());
+		// Agent already in the queue! Earlier conditions should have failed or data is inconsistent.
+		check(AgentQueue.AgentComponents.Find(&AgentComp) == INDEX_NONE);
+		AgentQueue.AgentComponents.Add(&AgentComp);
+
+		UE_VLOG(this, LogMassActor, Verbose, TEXT("%s registered and PENDING entity creation."), *AgentActor->GetName());
+   		AgentComp.EntityCreationPending();
+	}
+
+	return EntityTemplate->GetTemplateID();
+}
+
+void UMassAgentSubsystem::UpdateAgentComponent(const UMassAgentComponent& AgentComp)
+{
+	check(EntitySystem);
+	check(SpawnerSystem)
+
+	if (!ensureMsgf(AgentComp.GetEntityHandle().IsValid(), TEXT("Caling %s is valid only for already registered MassAgentComponents"), ANSI_TO_TCHAR(__FUNCTION__)))
+	{
+		UE_VLOG(this, LogMassActor, Warning, TEXT("%s: called while the given agent component has not been registered yet (Owner: %s)")
+			, ANSI_TO_TCHAR(__FUNCTION__), *GetNameSafe(AgentComp.GetOwner()));
+		return;
+	}
+
+	AActor* AgentActor = AgentComp.GetOwner();
+	check(AgentActor);
+
+	const FMassEntityConfig& EntityConfig = AgentComp.GetEntityConfig();
+	const FMassEntityTemplate* EntityTemplate = EntityConfig.GetOrCreateEntityTemplate(*AgentActor, AgentComp);
+
+	if (!ensureMsgf(EntityTemplate, TEXT("Failed to fetch entity template for %s."), *AgentActor->GetName()))
+	{
+		UE_VLOG(this, LogMassActor, Warning, TEXT("Failed to fetch entity template for %s."), *AgentActor->GetName());
+		return;
+	}
+
+	const FLWEntity Entity = AgentComp.GetEntityHandle().GetLWEntity();
+	const FArchetypeHandle CurrentArchetypeHandle = EntitySystem->GetArchetypeForEntity(Entity);
+	if (CurrentArchetypeHandle == EntityTemplate->GetArchetype())
+	{
+		UE_VLOG(this, LogMassActor, Log, TEXT("%s called for %s but no archetype changes have been found")
+			, ANSI_TO_TCHAR(__FUNCTION__), *AgentActor->GetName());
+	}
+	else
+	{
+		// the tricky case: we want to move the entity over to the new archetype and initialize its fragments, 
+		// but only the ones that the previous archetype didn't have - "delta initialize"
+		ensureMsgf(false, TEXT("Not implemented yet"));
+			
+		// @todo add override flag so that we always initialize a new when moving between archetypes
+	}
+
+	check(AgentComp.GetEntityHandle().IsValid());
+}
+
+void UMassAgentSubsystem::UnregisterAgentComponent(UMassAgentComponent& AgentComp)
+{
+	if (EntitySystem == nullptr || SpawnerSystem == nullptr || SimulationSystem == nullptr)
+	{
+		return;
+	}
+
+	if (AgentComp.IsPuppet())
+	{
+		if (AgentComp.IsPuppetPendingInitialization())
+		{
+			FMassAgentInitializationQueue* PuppetQueue = PendingPuppets.Find(AgentComp.GetTemplateID());
+		
+			if (ensureMsgf(PuppetQueue, TEXT("Trying to remove a puppet's agent component from initialization queue but there's no such queue")))
+			{
+				ensureMsgf(PuppetQueue->AgentComponents.Remove(&AgentComp), TEXT("Trying to remove a puppet's agent component from initialization queue while it's not in the queue"));
+			}
+			AgentComp.PuppetInitializationAborted();
+		}
+		else if (ensureMsgf(AgentComp.GetEntityHandle().IsValid(), TEXT("Trying to unregister a puppet's agent component while it's neither pending initialization nor having a valid entity handle")) )
+		{
+			if (AgentComp.GetPuppetSpecificAddition().IsEmpty() == false)
+			{
+				const FLWEntity Entity = AgentComp.GetEntityHandle().GetLWEntity();
+
+			    // remove fragments that have been added for the puppet agent
+			    if (EntitySystem->IsProcessing())
+			    {
+				    // need to request via command buffer since we can't move entities while processing is happening
+				    EntitySystem->Defer().PushCommand(FCommandRemoveComposition(Entity, AgentComp.GetPuppetSpecificAddition()));
+			    }
+			    else
+			    {
+					EntitySystem->RemoveCompositionFromEntity(Entity, AgentComp.GetPuppetSpecificAddition());
+			    }
+			}
+			AgentComp.PuppetUnregistrationDone();
+		}
+	}
+	else
+	{
+		if (AgentComp.GetEntityHandle().IsValid())
+		{
+			// the entity has already been created. Destroy!
+			const FMassEntityTemplate* EntityTemplate = nullptr;
+			AActor* AgentActor = AgentComp.GetOwner();
+			if (ensure(AgentActor))
+			{
+				const FMassEntityConfig& EntityConfig = AgentComp.GetEntityConfig();
+				EntityTemplate = &EntityConfig.GetEntityTemplateChecked(*AgentActor, AgentComp);
+			}
+
+			FLWEntity Entity = AgentComp.GetEntityHandle().GetLWEntity();
+			// Clearing the entity before it become invalid as the clear contains notifications
+			AgentComp.ClearEntityHandle();
+
+			// Destroy the entity
+			if (ensure(EntityTemplate))
+			{
+				SpawnerSystem->DestroyEntities(EntityTemplate->GetTemplateID(), TArrayView<FLWEntity>(&Entity, 1));
+			}
+		}
+		else if (AgentComp.IsEntityPendingCreation())
+		{
+			// hasn't been registered yet. Just remove it from the queue		
+			FMassAgentInitializationQueue* AgentQueue = PendingAgentEntities.Find(AgentComp.GetTemplateID());
+			if (ensureMsgf(AgentQueue, TEXT("Trying to remove an agent component from initialization queue but there's no such queue")))
+			{
+				ensureMsgf(AgentQueue->AgentComponents.Remove(&AgentComp), TEXT("Trying to remove an agent component from initialization queue it's missing from the queue"));
+			}
+			AgentComp.EntityCreationAborted();
+		}
+	}
+}
+
+void UMassAgentSubsystem::ShutdownAgentComponent(UMassAgentComponent& AgentComp)
+{
+	UnregisterAgentComponent(AgentComp);
+
+#if UE_REPLICATION_COMPILE_CLIENT_CODE
+	ReplicatedAgentComponents.Remove(AgentComp.GetNetID());
+#endif // UE_REPLICATION_COMPILE_CLIENT_CODE
+
+}
+
+void UMassAgentSubsystem::MakePuppet(UMassAgentComponent& AgentComp)
+{
+	if (!ensureMsgf(AgentComp.GetTemplateID().IsValid(), TEXT("%s tried to used %s as puppet but it doesn't have a valid entity template ID"), ANSI_TO_TCHAR(__FUNCTION__), *GetNameSafe(AgentComp.GetOwner())))
+	{
+		return;
+	}
+		
+	if (AgentComp.IsEntityPendingCreation())
+	{
+		FMassAgentInitializationQueue* AgentQueue = PendingAgentEntities.Find(AgentComp.GetTemplateID());
+		if (ensureMsgf(AgentQueue, TEXT("Trying to remove an agent component from initialization queue but there's no such queue")))
+		{
+			ensureMsgf(AgentQueue->AgentComponents.Remove(&AgentComp), TEXT("Trying to remove an agent component from initialization queue it's missing from the queue"));
+		}
+	}
+
+	FMassAgentInitializationQueue& PuppetQueue = PendingPuppets.FindOrAdd(AgentComp.GetTemplateID());
+	// Agent already in the queue! Earlier conditions should have failed or data is inconsistent.
+	check(PuppetQueue.AgentComponents.Find(&AgentComp) == INDEX_NONE);
+	PuppetQueue.AgentComponents.Add(&AgentComp);
+
+	AgentComp.PuppetInitializationPending();
+}
+
+void UMassAgentSubsystem::OnProcessingPhaseStarted(const float DeltaSeconds, const EPipeProcessingPhase Phase)
+{
+	switch (Phase)
+	{
+	case EPipeProcessingPhase::PrePhysics:
+		// initialize what needs initialization
+		if (PendingAgentEntities.Num() > 0 || PendingPuppets.Num() > 0)
+		{
+			HandlePendingInitialization();
+		}
+		break;
+	default:
+		// unhandled phases, by design, not every phase needs to be handled by the Actor subsystem
+		break;
+	}
+}
+
+void UMassAgentSubsystem::HandlePendingInitialization()
+{
+	check(EntitySystem);
+	check(SpawnerSystem);
+	check(SimulationSystem);
+	check(SimulationSystem->GetPhaseManager().IsDuringPipeProcessing() == false);
+
+	for (TTuple<FMassEntityTemplateID, FMassAgentInitializationQueue>& Data : PendingAgentEntities)
+	{
+		const FMassEntityTemplateID EntityTemplateID = Data.Get<0>();
+		const FMassEntityTemplate* EntityTemplate = SpawnerSystem->GetTemplateRegistryInstance().FindTemplateFromTemplateID(EntityTemplateID);
+		check(EntityTemplate);
+		
+		TArray<UMassAgentComponent*>& AgentComponents = Data.Get<1>().AgentComponents;
+		const int32 NewEntityCount = AgentComponents.Num();
+		
+		if (NewEntityCount <= 0)
+		{
+			// this case is perfectly fine, all agents registered and unregistered before we processed the queue
+			continue;
+		}
+		
+		TArray<FLWEntity> Entities;
+		SpawnerSystem->SpawnEntities(*EntityTemplate, NewEntityCount, Entities);
+		check(Entities.Num() == NewEntityCount);
+
+		if (EntityTemplate->GetObjectFragmentInitializers().Num())
+		{
+			const TConstArrayView<FMassEntityTemplate::FObjectFragmentInitializerFunction> ObjectFragmentInitializers = EntityTemplate->GetObjectFragmentInitializers();
+
+			for (int AgentIndex = 0; AgentIndex < Entities.Num(); ++AgentIndex)
+			{
+				FEntityView EntityView(EntityTemplate->GetArchetype(), Entities[AgentIndex]);
+				FMassAgentSubsystemHelper::InitializeAgentComponentFragments(*AgentComponents[AgentIndex], EntityView, EMassTranslationDirection::ActorToMass, ObjectFragmentInitializers);
+			}
+		}
+
+		for (int AgentIndex = 0; AgentIndex < Entities.Num(); ++AgentIndex)
+		{		
+			AgentComponents[AgentIndex]->SetEntityHandle(FMassHandle(Entities[AgentIndex]));
+		}
+	}
+
+	PendingAgentEntities.Reset();
+
+	for (TTuple<FMassEntityTemplateID, FMassAgentInitializationQueue>& Data : PendingPuppets)
+	{
+		const FMassEntityTemplateID EntityTemplateID = Data.Get<0>();
+		const FMassEntityTemplate* EntityTemplate = SpawnerSystem->GetTemplateRegistryInstance().FindTemplateFromTemplateID(EntityTemplateID);
+		if (!ensure(EntityTemplate))
+		{
+			// note that this condition is temporary, we'll be switched to a `check` once we set up characters
+			continue;
+		}
+
+		const FLWCompositionDescriptor TemplateDescriptor = EntityTemplate->GetCompositionDescriptor();
+
+		TArray<UMassAgentComponent*>& AgentComponents = Data.Get<1>().AgentComponents;
+
+		for (UMassAgentComponent* AgentComp : AgentComponents)
+		{
+			const FLWEntity PuppetEntity = AgentComp->GetEntityHandle().GetLWEntity();
+			if (!ensureMsgf(PuppetEntity.IsSet(), TEXT("Trying to initialize puppet's fragments while the pupped doesn't have a corresponding Entity identifier set. This should not happen.")))
+			{
+				continue;
+			}
+
+			FLWCompositionDescriptor& PuppetDescriptor = AgentComp->GetMutablePuppetSpecificAddition();
+			PuppetDescriptor = TemplateDescriptor;
+			EntitySystem->AddCompositionToEntity_GetDelta(PuppetEntity, PuppetDescriptor);
+			
+			if (EntityTemplate->GetObjectFragmentInitializers().Num())
+			{
+				const FArchetypeHandle ArchetypeHandle = EntitySystem->GetArchetypeForEntity(PuppetEntity);
+				FEntityView EntityView(ArchetypeHandle, PuppetEntity);
+				FMassAgentSubsystemHelper::InitializeAgentComponentFragments(*AgentComp, EntityView, EMassTranslationDirection::MassToActor, EntityTemplate->GetObjectFragmentInitializers());
+			}
+
+			AgentComp->PuppetInitializationDone();
+		}
+	}
+
+	PendingPuppets.Reset();
+}
+
+void UMassAgentSubsystem::NotifyMassAgentComponentReplicated(UMassAgentComponent& AgentComp)
+{
+#if UE_REPLICATION_COMPILE_CLIENT_CODE
+	UWorld* World = GetWorld();
+	if (World && ensureMsgf(World->IsNetMode(NM_Client), TEXT("%s: Expecting to only be called in network game on the client"), ANSI_TO_TCHAR(__FUNCTION__)))
+	{
+		check(ReplicationManager);
+		check(AgentComp.GetNetID().IsValid());
+		check(!ReplicatedAgentComponents.Find(AgentComp.GetNetID()));
+		ReplicatedAgentComponents.Add(AgentComp.GetNetID(), &AgentComp);
+		const FLWEntity Entity = ReplicationManager->FindEntity(AgentComp.GetNetID());
+
+		// If not found, the NotifyMassAgentAddedToReplication will link it later once replicated.
+		if (Entity.IsSet())
+		{
+			AgentComp.SetReplicatedPuppetHandle(FMassHandle(Entity));
+			MakePuppet(AgentComp);
+		}
+		else
+		{
+			AgentComp.MakePuppetAReplicatedOrphan();
+		}
+	}
+#endif // UE_REPLICATION_COMPILE_CLIENT_CODE
+}
+
+void UMassAgentSubsystem::NotifyMassAgentComponentEntityAssociated(const UMassAgentComponent& AgentComp) const
+{
+	OnMassAgentComponentEntityAssociated.Broadcast(AgentComp);
+}
+
+void UMassAgentSubsystem::NotifyMassAgentComponentEntityDetaching(const UMassAgentComponent& AgentComp) const
+{
+	OnMassAgentComponentEntityDetaching.Broadcast(AgentComp);
+}
+
+void UMassAgentSubsystem::OnMassAgentAddedToReplication(FMassNetworkID NetID, FLWEntity Entity)
+{
+#if UE_REPLICATION_COMPILE_CLIENT_CODE
+	UWorld* World = GetWorld();
+	if (World && World->IsNetMode(NM_Client))
+	{
+		if (UMassAgentComponent** AgentComp = ReplicatedAgentComponents.Find(NetID))
+		{
+			(*AgentComp)->SetReplicatedPuppetHandle(FMassHandle(Entity));
+			MakePuppet(**AgentComp);
+		}
+	}
+#endif // UE_REPLICATION_COMPILE_CLIENT_CODE
+}
+
+void UMassAgentSubsystem::OnMassAgentRemovedFromReplication(FMassNetworkID NetID, FLWEntity Entity)
+{
+#if UE_REPLICATION_COMPILE_CLIENT_CODE
+	UWorld* World = GetWorld();
+	if (World && World->IsNetMode(NM_Client))
+	{
+		if (UMassAgentComponent** AgentComp = ReplicatedAgentComponents.Find(NetID))
+		{
+			UnregisterAgentComponent(**AgentComp);
+			(*AgentComp)->ClearReplicatedPuppetHandle();
+		}
+	}
+#endif // UE_REPLICATION_COMPILE_CLIENT_CODE
+}
diff --git a/Engine/Plugins/Runtime/MassGameplay/Source/MassActors/Private/MassAgentTraits.cpp b/Engine/Plugins/Runtime/MassGameplay/Source/MassActors/Private/MassAgentTraits.cpp
new file mode 100644
--- /dev/null
+++ b/Engine/Plugins/Runtime/MassGameplay/Source/MassActors/Private/MassAgentTraits.cpp
@@ -0,0 +1,192 @@
+// Copyright Epic Games, Inc. All Rights Reserved.
+
+#include "MassAgentTraits.h"
+#include "MassCommonTypes.h"
+#include "MassMovementFragments.h"
+#include "MassTranslatorRegistry.h"
+#include "MassEntityTemplate.h"
+#include "MassEntityTemplateRegistry.h"
+#include "Translators/MassCapsuleComponentTranslators.h"
+#include "Translators/MassCharacterMovementTranslators.h"
+#include "Translators/MassSceneComponentLocationTranslator.h"
+#include "Components/CapsuleComponent.h"
+#include "GameFramework/Actor.h"
+#include "GameFramework/CharacterMovementComponent.h"
+#include "VisualLogger/VisualLogger.h"
+#include "EntityView.h"
+
+namespace FMassAgentTraitsHelper 
+{
+	template<typename T>
+	T* AsComponent(UObject& Owner)
+	{
+		T* Component = nullptr;
+		if (AActor* AsActor = Cast<AActor>(&Owner))
+		{
+			Component = AsActor->FindComponentByClass<T>();
+		}
+		else
+		{
+			Component = Cast<T>(&Owner);
+		}
+
+		UE_CVLOG_UELOG(Component == nullptr, &Owner, LogMass, Error, TEXT("Trying to extract %s from %s failed")
+			, *T::StaticClass()->GetName(), *Owner.GetName());
+
+		return Component;
+	}
+}
+
+//----------------------------------------------------------------------//
+//  UMassAgentCapsuleCollisionSyncTrait
+//----------------------------------------------------------------------//
+void UMassAgentCapsuleCollisionSyncTrait::BuildTemplate(FMassEntityTemplateBuildContext& BuildContext, UWorld& World) const
+{
+	BuildContext.AddFragment<FDataFragment_CapsuleComponentWrapper>();
+	BuildContext.AddFragment<FDataFragment_AgentRadius>();
+	if (bSyncTransform)
+	{
+		BuildContext.AddFragment<FDataFragment_Transform>();
+	}
+	
+	BuildContext.GetMutableObjectFragmentInitializers().Add([=](UObject& Owner, FEntityView& EntityView, const EMassTranslationDirection CurrentDirection)
+		{
+			if (UCapsuleComponent* CapsuleComponent = FMassAgentTraitsHelper::AsComponent<UCapsuleComponent>(Owner))
+			{
+				FDataFragment_CapsuleComponentWrapper& CapsuleFragment = EntityView.GetComponentData<FDataFragment_CapsuleComponentWrapper>();
+				CapsuleFragment.Component = CapsuleComponent;
+
+				FDataFragment_AgentRadius& RadiusFragment = EntityView.GetComponentData<FDataFragment_AgentRadius>();
+				RadiusFragment.Radius = CapsuleComponent->GetScaledCapsuleRadius();
+
+				if (bSyncTransform)
+				{
+					FDataFragment_Transform& TransformFragment = EntityView.GetComponentData<FDataFragment_Transform>();
+					TransformFragment.GetMutableTransform() = CapsuleComponent->GetComponentTransform();
+				}
+			}
+		});
+
+	if (bSyncTransform)
+	{
+		if (EnumHasAnyFlags(SyncDirection, EMassTranslationDirection::ActorToMass))
+		{
+			BuildContext.AddTranslator<UMassCapsuleTransformToMassTranslator>();
+		}
+
+		if (EnumHasAnyFlags(SyncDirection, EMassTranslationDirection::MassToActor))
+		{
+			BuildContext.AddTranslator<UMassTransformToActorCapsuleTranslator>();
+		}
+	}
+}
+
+//----------------------------------------------------------------------//
+//  UMassAgentMovementSyncTrait
+//----------------------------------------------------------------------//
+void UMassAgentMovementSyncTrait::BuildTemplate(FMassEntityTemplateBuildContext& BuildContext, UWorld& World) const
+{
+	BuildContext.AddFragment<FDataFragment_CharacterMovementComponentWrapper>();
+	BuildContext.AddFragment<FMassVelocityFragment>();
+	
+	BuildContext.GetMutableObjectFragmentInitializers().Add([=](UObject& Owner, FEntityView& EntityView, const EMassTranslationDirection CurrentDirection)
+		{
+			if (UCharacterMovementComponent* MovementComp = FMassAgentTraitsHelper::AsComponent<UCharacterMovementComponent>(Owner))
+			{
+				FDataFragment_CharacterMovementComponentWrapper& ComponentFragment = EntityView.GetComponentData<FDataFragment_CharacterMovementComponentWrapper>();
+				ComponentFragment.Component = MovementComp;
+
+				FMassVelocityFragment& VelocityFragment = EntityView.GetComponentData<FMassVelocityFragment>();
+
+				// the entity is the authority
+				if (CurrentDirection ==  EMassTranslationDirection::MassToActor)
+				{
+					MovementComp->bRunPhysicsWithNoController = true;
+					MovementComp->SetMovementMode(EMovementMode::MOVE_Walking);
+					MovementComp->Velocity = VelocityFragment.Value;
+				}
+				// actor is the authority
+				else
+				{
+					VelocityFragment.Value = MovementComp->GetLastUpdateVelocity();
+				}
+			}
+		});
+
+	if (EnumHasAnyFlags(SyncDirection, EMassTranslationDirection::ActorToMass))
+	{
+		BuildContext.AddTranslator<UMassCharacterMovementToMassTranslator>();
+	}
+
+	if (EnumHasAnyFlags(SyncDirection, EMassTranslationDirection::MassToActor))
+	{
+		BuildContext.AddTranslator<UMassCharacterMovementToActorTranslator>();
+	}
+}
+
+//----------------------------------------------------------------------//
+//  UMassAgentOrientationSyncTrait
+//----------------------------------------------------------------------//
+void UMassAgentOrientationSyncTrait::BuildTemplate(FMassEntityTemplateBuildContext& BuildContext, UWorld& World) const
+{
+	BuildContext.AddFragment<FDataFragment_CharacterMovementComponentWrapper>();
+	// @todo: Figure out how we can share init with UMassAgentMovementSyncTrait, or make this this trait to depend on UMassAgentMovementSyncTrait.
+
+	if (EnumHasAnyFlags(SyncDirection, EMassTranslationDirection::ActorToMass))
+	{
+		BuildContext.AddTranslator<UMassCharacterOrientationToMassTranslator>();
+	}
+
+	if (EnumHasAnyFlags(SyncDirection, EMassTranslationDirection::MassToActor))
+	{
+		BuildContext.AddTranslator<UMassCharacterOrientationToActorTranslator>();
+	}
+}
+
+//----------------------------------------------------------------------//
+//  UMassAgentFeetLocationSyncTrait
+//----------------------------------------------------------------------//
+void UMassAgentFeetLocationSyncTrait::BuildTemplate(FMassEntityTemplateBuildContext& BuildContext, UWorld& World) const
+{
+	BuildContext.AddFragment<FMassSceneComponentWrapperFragment>();
+	BuildContext.AddFragment<FDataFragment_Transform>();
+
+	BuildContext.GetMutableObjectFragmentInitializers().Add([=](UObject& Owner, FEntityView& EntityView, const EMassTranslationDirection CurrentDirection)
+		{
+			AActor* AsActor = Cast<AActor>(&Owner);
+			if (AsActor && AsActor->GetRootComponent())
+			{
+				USceneComponent* Component = AsActor->GetRootComponent();
+				FMassSceneComponentWrapperFragment& ComponentFragment = EntityView.GetComponentData<FMassSceneComponentWrapperFragment>();
+				ComponentFragment.Component = Component;
+
+				FDataFragment_Transform& TransformFragment = EntityView.GetComponentData<FDataFragment_Transform>();
+
+				REDIRECT_OBJECT_TO_VLOG(Component, &Owner);
+				UE_VLOG_LOCATION(&Owner, LogMass, Log, Component->GetComponentLocation(), 30, FColor::Yellow, TEXT("Initial component location"));
+				UE_VLOG_LOCATION(&Owner, LogMass, Log, TransformFragment.GetTransform().GetLocation(), 30, FColor::Red, TEXT("Initial entity location"));
+
+				// the entity is the authority
+				if (CurrentDirection == EMassTranslationDirection::MassToActor)
+				{
+					Component->SetWorldLocation(TransformFragment.GetTransform().GetLocation() + FVector(0.f, 0.f, Component->Bounds.BoxExtent.Z));
+				}
+				// actor is the authority
+				else
+				{
+					TransformFragment.GetMutableTransform().SetLocation(Component->GetComponentTransform().GetLocation() - FVector(0.f, 0.f, Component->Bounds.BoxExtent.Z));
+				}
+			}
+		});
+
+	if (EnumHasAnyFlags(SyncDirection, EMassTranslationDirection::ActorToMass))
+	{
+		BuildContext.AddTranslator<UMassSceneComponentLocationToMassTranslator>();
+	}
+
+	if (EnumHasAnyFlags(SyncDirection, EMassTranslationDirection::MassToActor))
+	{
+		BuildContext.AddTranslator<UMassSceneComponentLocationToActorTranslator>();
+	}
+}
+
diff --git a/Engine/Plugins/Runtime/MassGameplay/Source/MassActors/Private/Translators/MassCapsuleComponentTranslators.cpp b/Engine/Plugins/Runtime/MassGameplay/Source/MassActors/Private/Translators/MassCapsuleComponentTranslators.cpp
new file mode 100644
--- /dev/null
+++ b/Engine/Plugins/Runtime/MassGameplay/Source/MassActors/Private/Translators/MassCapsuleComponentTranslators.cpp
@@ -0,0 +1,75 @@
+// Copyright Epic Games, Inc. All Rights Reserved.
+
+#include "Translators/MassCapsuleComponentTranslators.h"
+#include "MassCommonTypes.h"
+#include "MassEntitySystem.h"
+#include "Components/CapsuleComponent.h"
+#include "MassMovementTypes.h"
+
+
+//----------------------------------------------------------------------//
+// UMassCapsuleTransformToMassTranslator
+//----------------------------------------------------------------------//
+UMassCapsuleTransformToMassTranslator::UMassCapsuleTransformToMassTranslator()
+{
+	ExecutionFlags = (int32)EProcessorExecutionFlags::All;
+	RequiredTags.Add<FMassCapsuleTransformCopyToMassTag>();
+	ExecutionOrder.ExecuteInGroup = UE::Mass::ProcessorGroupNames::SyncWorldToMass;
+}
+
+void UMassCapsuleTransformToMassTranslator::ConfigureQueries()
+{
+	AddRequiredTagsToQuery(EntityQuery);
+	EntityQuery.AddRequirement<FDataFragment_CapsuleComponentWrapper>(ELWComponentAccess::ReadOnly);
+	EntityQuery.AddRequirement<FDataFragment_Transform>(ELWComponentAccess::ReadWrite);
+}
+
+void UMassCapsuleTransformToMassTranslator::Execute(UEntitySubsystem& EntitySubsystem, FLWComponentSystemExecutionContext& Context)
+{
+	EntityQuery.ForEachEntityChunk(EntitySubsystem, Context, [this](FLWComponentSystemExecutionContext& Context)
+		{
+			const TConstArrayView<FDataFragment_CapsuleComponentWrapper> CapsuleComponentList = Context.GetComponentView<FDataFragment_CapsuleComponentWrapper>();
+			const TArrayView<FDataFragment_Transform> LocationList = Context.GetMutableComponentView<FDataFragment_Transform>();
+			for (int i = 0; i < CapsuleComponentList.Num(); ++i)
+			{
+				if (const UCapsuleComponent* CapsuleComp = CapsuleComponentList[i].Component.Get())
+				{
+					LocationList[i].GetMutableTransform() = CapsuleComp->GetComponentTransform();
+				}
+			}
+		});
+}
+
+//----------------------------------------------------------------------//
+// UMassTransformToActorCapsuleTranslator
+//----------------------------------------------------------------------//
+UMassTransformToActorCapsuleTranslator::UMassTransformToActorCapsuleTranslator()
+{
+	ExecutionFlags = (int32)EProcessorExecutionFlags::All;
+	RequiredTags.Add<FMassCapsuleTransformCopyToActorTag>();
+	ExecutionOrder.ExecuteInGroup = UE::Mass::ProcessorGroupNames::UpdateWorldFromMass;
+	ExecutionOrder.ExecuteAfter.Add(UE::Mass::ProcessorGroupNames::Avoidance);
+}
+
+void UMassTransformToActorCapsuleTranslator::ConfigureQueries()
+{
+	AddRequiredTagsToQuery(EntityQuery);
+	EntityQuery.AddRequirement<FDataFragment_CapsuleComponentWrapper>(ELWComponentAccess::ReadWrite);
+	EntityQuery.AddRequirement<FDataFragment_Transform>(ELWComponentAccess::ReadOnly);
+}
+
+void UMassTransformToActorCapsuleTranslator::Execute(UEntitySubsystem& EntitySubsystem, FLWComponentSystemExecutionContext& Context)
+{
+	EntityQuery.ForEachEntityChunk(EntitySubsystem, Context, [this](FLWComponentSystemExecutionContext& Context)
+		{
+			const TArrayView<FDataFragment_CapsuleComponentWrapper> CapsuleComponentList = Context.GetMutableComponentView<FDataFragment_CapsuleComponentWrapper>();
+			const TConstArrayView<FDataFragment_Transform> LocationList = Context.GetComponentView<FDataFragment_Transform>();
+			for (int i = 0; i < CapsuleComponentList.Num(); ++i)
+			{
+				if (UCapsuleComponent* CapsuleComp = CapsuleComponentList[i].Component.Get())
+				{
+					CapsuleComp->SetWorldTransform(LocationList[i].GetTransform(), /*bSweep=*/false);
+				}
+			}
+		});
+}
diff --git a/Engine/Plugins/Runtime/MassGameplay/Source/MassActors/Private/Translators/MassCharacterMovementTranslators.cpp b/Engine/Plugins/Runtime/MassGameplay/Source/MassActors/Private/Translators/MassCharacterMovementTranslators.cpp
new file mode 100644
--- /dev/null
+++ b/Engine/Plugins/Runtime/MassGameplay/Source/MassActors/Private/Translators/MassCharacterMovementTranslators.cpp
@@ -0,0 +1,269 @@
+// Copyright Epic Games, Inc. All Rights Reserved.
+
+#include "Translators/MassCharacterMovementTranslators.h"
+#include "Logging/LogMacros.h"
+#include "GameFramework/CharacterMovementComponent.h"
+#include "VisualLogger/VisualLogger.h"
+#include "MassEntitySystem.h"
+#include "MassCommonTypes.h"
+#include "MassMovementTypes.h"
+#include "MassMovementFragments.h"
+#include "MassSpawnerTypes.h"
+
+//----------------------------------------------------------------------//
+//  UMassCharacterMovementToMassTranslator
+//----------------------------------------------------------------------//
+UMassCharacterMovementToMassTranslator::UMassCharacterMovementToMassTranslator()
+{
+	ExecutionFlags = (int32)EProcessorExecutionFlags::All;
+	ExecutionOrder.ExecuteInGroup = UE::Mass::ProcessorGroupNames::SyncWorldToMass;
+	RequiredTags.Add<FMassCharacterMovementCopyToMassTag>();
+}
+
+void UMassCharacterMovementToMassTranslator::ConfigureQueries()
+{
+	AddRequiredTagsToQuery(EntityQuery);
+	EntityQuery.AddRequirement<FDataFragment_CharacterMovementComponentWrapper>(ELWComponentAccess::ReadOnly);
+	EntityQuery.AddRequirement<FDataFragment_Transform>(ELWComponentAccess::ReadWrite);
+	EntityQuery.AddRequirement<FMassVelocityFragment>(ELWComponentAccess::ReadWrite);
+}
+
+void UMassCharacterMovementToMassTranslator::Execute(UEntitySubsystem& EntitySubsystem, FLWComponentSystemExecutionContext& Context)
+{
+	EntityQuery.ForEachEntityChunk(EntitySubsystem, Context, [this](FLWComponentSystemExecutionContext& Context)
+	{
+		const TConstArrayView<FDataFragment_CharacterMovementComponentWrapper> ComponentList = Context.GetComponentView<FDataFragment_CharacterMovementComponentWrapper>();
+		const TArrayView<FDataFragment_Transform> LocationList = Context.GetMutableComponentView<FDataFragment_Transform>();
+		const TArrayView<FMassVelocityFragment> VelocityList = Context.GetMutableComponentView<FMassVelocityFragment>();
+
+		const int32 NumEntities = Context.GetEntitiesNum();
+		
+		for (int32 i = 0; i < NumEntities; ++i)
+		{
+			if (const UCharacterMovementComponent* AsMovementComponent = ComponentList[i].Component.Get())
+			{
+				LocationList[i].GetMutableTransform().SetLocation(AsMovementComponent->GetActorNavLocation());
+				
+				VelocityList[i].Value = AsMovementComponent->Velocity;
+			}
+		}
+	});
+}
+
+//----------------------------------------------------------------------//
+//  UMassCharacterMovementToActorTranslator
+//----------------------------------------------------------------------//
+UMassCharacterMovementToActorTranslator::UMassCharacterMovementToActorTranslator()
+{
+	ExecutionFlags = (int32)EProcessorExecutionFlags::All;
+	ExecutionOrder.ExecuteInGroup = UE::Mass::ProcessorGroupNames::UpdateWorldFromMass;
+	ExecutionOrder.ExecuteAfter.Add(UE::Mass::ProcessorGroupNames::Avoidance);
+	RequiredTags.Add<FMassCharacterMovementCopyToActorTag>();
+}
+
+void UMassCharacterMovementToActorTranslator::ConfigureQueries()
+{
+	AddRequiredTagsToQuery(EntityQuery);
+	EntityQuery.AddRequirement<FDataFragment_CharacterMovementComponentWrapper>(ELWComponentAccess::ReadWrite);
+	EntityQuery.AddRequirement<FMassVelocityFragment>(ELWComponentAccess::ReadOnly);
+}
+
+void UMassCharacterMovementToActorTranslator::Execute(UEntitySubsystem& EntitySubsystem, FLWComponentSystemExecutionContext& Context)
+{
+	EntityQuery.ForEachEntityChunk(EntitySubsystem, Context, [this](FLWComponentSystemExecutionContext& Context)
+	{
+		const TArrayView<FDataFragment_CharacterMovementComponentWrapper> ComponentList = Context.GetMutableComponentView<FDataFragment_CharacterMovementComponentWrapper>();
+		const TConstArrayView<FMassVelocityFragment> VelocityList = Context.GetComponentView<FMassVelocityFragment>();
+		
+		const int32 NumEntities = Context.GetEntitiesNum();
+
+		for (int32 i = 0; i < NumEntities; ++i)
+		{
+			if (UCharacterMovementComponent* AsMovementComponent = ComponentList[i].Component.Get())
+			{
+				AsMovementComponent->RequestDirectMove(VelocityList[i].Value, /*bForceMaxSpeed=*/false);
+			}
+		}
+	});
+}
+
+//----------------------------------------------------------------------//
+//  UMassCharacterOrientationToMassTranslator
+//----------------------------------------------------------------------//
+UMassCharacterOrientationToMassTranslator::UMassCharacterOrientationToMassTranslator()
+{
+	ExecutionFlags = (int32)EProcessorExecutionFlags::All;
+	ExecutionOrder.ExecuteInGroup = UE::Mass::ProcessorGroupNames::SyncWorldToMass;
+	RequiredTags.Add<FMassCharacterOrientationCopyToMassTag>();
+}
+
+void UMassCharacterOrientationToMassTranslator::ConfigureQueries()
+{
+	AddRequiredTagsToQuery(EntityQuery);
+	EntityQuery.AddRequirement<FDataFragment_CharacterMovementComponentWrapper>(ELWComponentAccess::ReadOnly);
+	EntityQuery.AddRequirement<FDataFragment_Transform>(ELWComponentAccess::ReadWrite);
+}
+
+void UMassCharacterOrientationToMassTranslator::Execute(UEntitySubsystem& EntitySubsystem, FLWComponentSystemExecutionContext& Context)
+{
+	EntityQuery.ForEachEntityChunk(EntitySubsystem, Context, [this](FLWComponentSystemExecutionContext& Context)
+	{
+		const TConstArrayView<FDataFragment_CharacterMovementComponentWrapper> ComponentList = Context.GetComponentView<FDataFragment_CharacterMovementComponentWrapper>();
+		const TArrayView<FDataFragment_Transform> LocationList = Context.GetMutableComponentView<FDataFragment_Transform>();
+
+		const int32 NumEntities = Context.GetEntitiesNum();
+		
+		for (int32 i = 0; i < NumEntities; ++i)
+		{
+			if (const UCharacterMovementComponent* AsMovementComponent = ComponentList[i].Component.Get())
+			{
+				if (AsMovementComponent->UpdatedComponent != nullptr)
+				{
+					LocationList[i].GetMutableTransform().SetRotation(AsMovementComponent->UpdatedComponent->GetComponentTransform().GetRotation());
+				}
+			}
+		}
+	});
+}
+
+//----------------------------------------------------------------------//
+//  UMassCharacterOrientationToActorTranslator
+//----------------------------------------------------------------------//
+UMassCharacterOrientationToActorTranslator::UMassCharacterOrientationToActorTranslator()
+{
+	ExecutionFlags = (int32)EProcessorExecutionFlags::All;
+	ExecutionOrder.ExecuteInGroup = UE::Mass::ProcessorGroupNames::UpdateWorldFromMass;
+	ExecutionOrder.ExecuteAfter.Add(UE::Mass::ProcessorGroupNames::Avoidance);
+	RequiredTags.Add<FMassCharacterOrientationCopyToActorTag>();
+}
+
+void UMassCharacterOrientationToActorTranslator::ConfigureQueries()
+{
+	AddRequiredTagsToQuery(EntityQuery);
+	EntityQuery.AddRequirement<FDataFragment_CharacterMovementComponentWrapper>(ELWComponentAccess::ReadWrite);
+	EntityQuery.AddRequirement<FDataFragment_Transform>(ELWComponentAccess::ReadOnly);
+}
+
+void UMassCharacterOrientationToActorTranslator::Execute(UEntitySubsystem& EntitySubsystem, FLWComponentSystemExecutionContext& Context)
+{
+	EntityQuery.ForEachEntityChunk(EntitySubsystem, Context, [this](FLWComponentSystemExecutionContext& Context)
+	{
+		const TArrayView<FDataFragment_CharacterMovementComponentWrapper> ComponentList = Context.GetMutableComponentView<FDataFragment_CharacterMovementComponentWrapper>();
+		const TConstArrayView<FDataFragment_Transform> TransformList = Context.GetComponentView<FDataFragment_Transform>();
+		
+		const int32 NumEntities = Context.GetEntitiesNum();
+
+		for (int32 i = 0; i < NumEntities; ++i)
+		{
+			if (UCharacterMovementComponent* AsMovementComponent = ComponentList[i].Component.Get())
+			{
+				if (AsMovementComponent->UpdatedComponent != nullptr)
+				{
+					const FDataFragment_Transform& Transform = TransformList[i];
+					AsMovementComponent->bOrientRotationToMovement = false;
+					AsMovementComponent->UpdatedComponent->SetWorldRotation(Transform.GetTransform().GetRotation());
+				}
+			}
+		}
+	});
+}
+
+//----------------------------------------------------------------------//
+// UMassFragmentInitializer_NavLocation 
+//----------------------------------------------------------------------//
+UMassFragmentInitializer_NavLocation::UMassFragmentInitializer_NavLocation()
+{
+	FragmentType = FDataFragment_NavLocation::StaticStruct();
+}
+
+void UMassFragmentInitializer_NavLocation::ConfigureQueries() 
+{
+	EntityQuery.AddRequirement<FDataFragment_NavLocation>(ELWComponentAccess::ReadWrite);
+}
+
+void UMassFragmentInitializer_NavLocation::Execute(UEntitySubsystem& EntitySubsystem, FLWComponentSystemExecutionContext& Context)
+{
+	EntityQuery.ForEachEntityChunk(EntitySubsystem, Context, [this](FLWComponentSystemExecutionContext& Context)
+		{
+			const TArrayView<FDataFragment_NavLocation> NavLocationList = Context.GetMutableComponentView<FDataFragment_NavLocation>();
+			const int32 NumEntities = Context.GetEntitiesNum();
+			for (int32 i = 0; i < NumEntities; ++i)
+			{
+				NavLocationList[i].NodeRef = INVALID_NAVNODEREF;
+			}
+		});
+}
+
+//----------------------------------------------------------------------//
+// UMassFragmentInitializer_Transform 
+//----------------------------------------------------------------------//
+UMassFragmentInitializer_Transform::UMassFragmentInitializer_Transform()
+{
+	ExecutionFlags = (int32)EProcessorExecutionFlags::All;
+	FragmentType = FDataFragment_Transform::StaticStruct();
+}
+
+void UMassFragmentInitializer_Transform::ConfigureQueries()
+{
+	EntityQuery.AddRequirement<FDataFragment_Transform>(ELWComponentAccess::ReadWrite);
+}
+
+void UMassFragmentInitializer_Transform::Execute(UEntitySubsystem& EntitySubsystem, FLWComponentSystemExecutionContext& Context)
+{
+	const UWorld* World = EntitySubsystem.GetWorld();
+
+	check(World);
+
+	const ENetMode NetMode = World->GetNetMode();
+
+	if (NetMode != NM_Client)
+	{
+		const FInstancedStruct& AuxInput = Context.GetMutableAuxData();
+		if (!(AuxInput.IsValid() && AuxInput.GetScriptStruct()->IsChildOf(FMassSpawnAuxData::StaticStruct())))
+		{
+			UE_VLOG_UELOG(this, LogMass, Log, TEXT("Execution context has invalid AuxData or it's not FMassSpawnAuxData. Entity transforms won't be initialized."));
+			return;
+		}
+
+		TArray<FTransform>& Transforms = AuxInput.GetMutable<FMassSpawnAuxData>().Transforms;
+
+		const int32 NumSpawnTransforms = Transforms.Num();
+		if (NumSpawnTransforms == 0)
+		{
+			UE_VLOG_UELOG(this, LogMass, Error, TEXT("No spawn transforms provided. Entity transforms won't be initialized."));
+			return;
+		}
+
+		int32 NumRequiredSpawnTransforms = 0;
+		EntityQuery.ForEachEntityChunk(EntitySubsystem, Context, [&NumRequiredSpawnTransforms](const FLWComponentSystemExecutionContext& Context)
+			{
+				NumRequiredSpawnTransforms += Context.GetEntitiesNum();
+			});
+
+		const int32 NumToAdd = NumRequiredSpawnTransforms - NumSpawnTransforms;
+		if (NumToAdd > 0)
+		{
+			UE_VLOG_UELOG(this, LogMass, Warning,
+				TEXT("Not enough spawn locations provided (%d) for all entities (%d). Existing locations will be reused randomly to fill the %d missing positions."),
+				NumSpawnTransforms, NumRequiredSpawnTransforms, NumToAdd);
+
+			Transforms.AddUninitialized(NumToAdd);
+			for (int i = 0; i < NumToAdd; ++i)
+			{
+				Transforms[NumSpawnTransforms + i] = Transforms[FMath::RandRange(0, NumSpawnTransforms - 1)];
+			}
+		}
+
+		EntityQuery.ForEachEntityChunk(EntitySubsystem, Context, [&Transforms, this](FLWComponentSystemExecutionContext& Context)
+			{
+				const TArrayView<FDataFragment_Transform> LocationList = Context.GetMutableComponentView<FDataFragment_Transform>();
+				const int32 NumEntities = Context.GetEntitiesNum();
+				for (int32 i = 0; i < NumEntities; ++i)
+				{
+					const int32 AuxIndex = FMath::RandRange(0, Transforms.Num() - 1);
+					LocationList[i].GetMutableTransform() = Transforms[AuxIndex];
+					Transforms.RemoveAtSwap(AuxIndex, 1, /*bAllowShrinking=*/false);
+				}
+			});
+	}
+}
diff --git a/Engine/Plugins/Runtime/MassGameplay/Source/MassActors/Private/Translators/MassSceneComponentLocationTranslator.cpp b/Engine/Plugins/Runtime/MassGameplay/Source/MassActors/Private/Translators/MassSceneComponentLocationTranslator.cpp
new file mode 100644
--- /dev/null
+++ b/Engine/Plugins/Runtime/MassGameplay/Source/MassActors/Private/Translators/MassSceneComponentLocationTranslator.cpp
@@ -0,0 +1,77 @@
+// Copyright Epic Games, Inc. All Rights Reserved.
+
+#include "Translators/MassSceneComponentLocationTranslator.h"
+#include "MassMovementTypes.h"
+#include "Components/SceneComponent.h"
+#include "MassEntitySystem.h"
+
+//----------------------------------------------------------------------//
+//  UMassSceneComponentLocationToMassTranslator
+//----------------------------------------------------------------------//
+UMassSceneComponentLocationToMassTranslator::UMassSceneComponentLocationToMassTranslator()
+{
+	ExecutionFlags = (int32)EProcessorExecutionFlags::All;
+	ExecutionOrder.ExecuteInGroup = UE::Mass::ProcessorGroupNames::SyncWorldToMass;
+	RequiredTags.Add<FMassSceneComponentLocationCopyToMassTag>();
+}
+
+void UMassSceneComponentLocationToMassTranslator::ConfigureQueries()
+{
+	AddRequiredTagsToQuery(EntityQuery);
+	EntityQuery.AddRequirement<FMassSceneComponentWrapperFragment>(ELWComponentAccess::ReadOnly);
+	EntityQuery.AddRequirement<FDataFragment_Transform>(ELWComponentAccess::ReadWrite);
+}
+
+void UMassSceneComponentLocationToMassTranslator::Execute(UEntitySubsystem& EntitySubsystem, FLWComponentSystemExecutionContext& Context)
+{
+	EntityQuery.ForEachEntityChunk(EntitySubsystem, Context, [this](FLWComponentSystemExecutionContext& Context)
+	{
+		const TConstArrayView<FMassSceneComponentWrapperFragment> ComponentList = Context.GetComponentView<FMassSceneComponentWrapperFragment>();
+		const TArrayView<FDataFragment_Transform> LocationList = Context.GetMutableComponentView<FDataFragment_Transform>();
+
+		const int32 NumEntities = Context.GetEntitiesNum();
+		for (int32 i = 0; i < NumEntities; ++i)
+		{
+			if (const USceneComponent* AsComponent = ComponentList[i].Component.Get())
+			{
+				LocationList[i].GetMutableTransform().SetLocation(AsComponent->GetComponentTransform().GetLocation() - FVector(0.f, 0.f, AsComponent->Bounds.BoxExtent.Z));
+			}
+		}
+	});
+}
+
+//----------------------------------------------------------------------//
+//  UMassSceneComponentLocationToActorTranslator
+//----------------------------------------------------------------------//
+UMassSceneComponentLocationToActorTranslator::UMassSceneComponentLocationToActorTranslator()
+{
+	ExecutionFlags = (int32)EProcessorExecutionFlags::All;
+	ExecutionOrder.ExecuteInGroup = UE::Mass::ProcessorGroupNames::UpdateWorldFromMass;
+	ExecutionOrder.ExecuteAfter.Add(UE::Mass::ProcessorGroupNames::Avoidance);
+	RequiredTags.Add<FMassSceneComponentLocationCopyToActorTag>();
+}
+
+void UMassSceneComponentLocationToActorTranslator::ConfigureQueries()
+{
+	AddRequiredTagsToQuery(EntityQuery);
+	EntityQuery.AddRequirement<FMassSceneComponentWrapperFragment>(ELWComponentAccess::ReadOnly);
+	EntityQuery.AddRequirement<FDataFragment_Transform>(ELWComponentAccess::ReadWrite);
+}
+
+void UMassSceneComponentLocationToActorTranslator::Execute(UEntitySubsystem& EntitySubsystem, FLWComponentSystemExecutionContext& Context)
+{
+	EntityQuery.ForEachEntityChunk(EntitySubsystem, Context, [this](FLWComponentSystemExecutionContext& Context)
+		{
+			const TConstArrayView<FMassSceneComponentWrapperFragment> ComponentList = Context.GetComponentView<FMassSceneComponentWrapperFragment>();
+			const TArrayView<FDataFragment_Transform> LocationList = Context.GetMutableComponentView<FDataFragment_Transform>();
+
+			const int32 NumEntities = Context.GetEntitiesNum();
+			for (int32 i = 0; i < NumEntities; ++i)
+			{
+				if (USceneComponent* AsComponent = ComponentList[i].Component.Get())
+				{
+					AsComponent->SetWorldLocation(LocationList[i].GetTransform().GetLocation() + FVector(0.f, 0.f, AsComponent->Bounds.BoxExtent.Z));
+				}
+			}
+		});
+}
\ No newline at end of file
diff --git a/Engine/Plugins/Runtime/MassGameplay/Source/MassActors/Private/Translators/MassTranslators_BehaviorTree.cpp b/Engine/Plugins/Runtime/MassGameplay/Source/MassActors/Private/Translators/MassTranslators_BehaviorTree.cpp
new file mode 100644
--- /dev/null
+++ b/Engine/Plugins/Runtime/MassGameplay/Source/MassActors/Private/Translators/MassTranslators_BehaviorTree.cpp
@@ -0,0 +1,15 @@
+// Copyright Epic Games, Inc. All Rights Reserved.
+
+#include "Translators/MassTranslators_BehaviorTree.h"
+#include "BehaviorTree/BehaviorTreeComponent.h"
+#include "MassEntitySystem.h"
+
+UMassTranslator_BehaviorTree::UMassTranslator_BehaviorTree()
+{
+	bAutoRegisterWithProcessingPhases = false;
+}
+
+void UMassTranslator_BehaviorTree::ConfigureQueries()
+{
+	EntityQuery.AddRequirement<FDataFragment_BehaviorTreeComponentWrapper>(ELWComponentAccess::ReadWrite);
+}
diff --git a/Engine/Plugins/Runtime/MassGameplay/Source/MassActors/Public/MassActorPoolableInterface.h b/Engine/Plugins/Runtime/MassGameplay/Source/MassActors/Public/MassActorPoolableInterface.h
new file mode 100644
--- /dev/null
+++ b/Engine/Plugins/Runtime/MassGameplay/Source/MassActors/Public/MassActorPoolableInterface.h
@@ -0,0 +1,29 @@
+// Copyright Epic Games, Inc. All Rights Reserved.
+
+#pragma once
+
+#include "UObject/Interface.h"
+
+#include "MassActorPoolableInterface.generated.h"
+
+UINTERFACE(Blueprintable)
+class MASSACTORS_API UMassActorPoolableInterface : public UInterface
+{
+	GENERATED_UINTERFACE_BODY()
+};
+
+class MASSACTORS_API IMassActorPoolableInterface : public IInterface
+{
+	GENERATED_IINTERFACE_BODY()
+
+public:
+
+	UFUNCTION(BlueprintNativeEvent, BlueprintCallable, Category="Mass|Actor Pooling")
+	bool CanBePooled();
+
+	UFUNCTION(BlueprintNativeEvent, BlueprintCallable, Category="Mass|Actor Pooling")
+	void PrepareForPooling();
+
+	UFUNCTION(BlueprintNativeEvent, BlueprintCallable, Category="Mass|Actor Pooling")
+	void PrepareForGame();
+};
\ No newline at end of file
diff --git a/Engine/Plugins/Runtime/MassGameplay/Source/MassActors/Public/MassActorSpawnerSubsystem.h b/Engine/Plugins/Runtime/MassGameplay/Source/MassActors/Public/MassActorSpawnerSubsystem.h
new file mode 100644
--- /dev/null
+++ b/Engine/Plugins/Runtime/MassGameplay/Source/MassActors/Public/MassActorSpawnerSubsystem.h
@@ -0,0 +1,238 @@
+// Copyright Epic Games, Inc. All Rights Reserved.
+
+#pragma once
+
+#include "InstancedStruct.h"
+#include "IndexedHandle.h"
+#include "Delegates/Delegate.h"
+#include "MassCommonTypes.h"
+#include "Templates/SubclassOf.h"
+#include "GameFramework/Actor.h"
+#include "Subsystems/WorldSubsystem.h"
+
+#include "MassActorSpawnerSubsystem.generated.h"
+
+// Handle for an actor spawning request
+USTRUCT()
+struct MASSACTORS_API FMassHandle_ActorSpawnRequest : public FIndexedHandleBase
+{
+	GENERATED_BODY()
+
+	FMassHandle_ActorSpawnRequest() = default;
+
+	/** @note passing INDEX_NONE as index will make this handle Invalid */
+	FMassHandle_ActorSpawnRequest(const int32 InIndex, const uint32 InSerialNumber) : FIndexedHandleBase(InIndex, InSerialNumber)
+	{
+	}
+};
+
+// Managing class of spawning requests handles
+typedef FIndexedHandleManager<FMassHandle_ActorSpawnRequest, true/*bOptimizeHandleReuse*/> FMassHandleManager_ActorSpawnRequest;
+
+
+DECLARE_DELEGATE_TwoParams(FMassActorPreSpawnDelegate, const FMassHandle_ActorSpawnRequest& SpawnRequestHandle, const FStructView& SpawnRequest);
+enum EMassActorSpawnRequestAction
+{
+	Keep, // Will leave spawning request in the queue and it will be users job to call RemoveActorSpawnRequest
+	Remove, // Will remove the spawning request from the queue once the callback ends
+};
+DECLARE_DELEGATE_RetVal_TwoParams(EMassActorSpawnRequestAction, FMassActorPostSpawnDelegate, const FMassHandle_ActorSpawnRequest& SpawnRequestHandle, const FStructView& SpawnRequest);
+
+UENUM()
+enum class ESpawnRequestStatus : uint8
+{
+	None, // Not in the queue to be spawn
+	Pending, // Still in the queue to be spawn
+	Processing, // in the process of spawning the actor
+	Succeeded, // Successfully spawned the actor
+	Failed, // Error while spawning the actor
+	RetryPending, // Failed spawn request that are being retried (lower priority)
+};
+
+/**
+ * Base class for all spawn request
+ */
+USTRUCT()
+struct MASSACTORS_API FMassActorSpawnRequest
+{
+	GENERATED_BODY()
+public:
+	/** The actual mass agent handle corresponding to the actor to spawn */
+	FMassHandle MassAgent;
+
+	/** The template BP actor to spawn */
+	UPROPERTY(Transient)
+	TSubclassOf<AActor> Template;
+
+	/** The location of where to spawn that actor */
+	FTransform	Transform;
+
+	/** Priority of this spawn request in comparison with the others, the lower the value is, the higher the priority is */
+	float Priority = MAX_FLT;
+
+	/** Delegate that will be called just before the spawning an actor, giving the chance to the processor to prepare it */
+	FMassActorPreSpawnDelegate ActorPreSpawnDelegate;
+
+	/** Delegate that will be called once the spawning is done */
+	FMassActorPostSpawnDelegate ActorPostSpawnDelegate;
+
+	/** The current status of the spawn request */
+	ESpawnRequestStatus SpawnStatus = ESpawnRequestStatus::None;
+
+	/** The pointer to the actor once it is spawned */
+	UPROPERTY(Transient)
+	AActor* SpawnedActor = nullptr;
+
+	/** Internal request serial number (used to cycle through next spawning request) */
+	uint32 SerialNumber = 0;
+
+	/** Requested world time seconds */
+	float RequestedTime = 0.0f;
+
+	void Reset()
+	{
+		MassAgent = FMassHandle();
+		Template = nullptr;
+		Priority = MAX_FLT;
+		ActorPreSpawnDelegate.Unbind();
+		ActorPostSpawnDelegate.Unbind();
+		SpawnStatus = ESpawnRequestStatus::None;
+		SpawnedActor = nullptr;
+		SerialNumber = 0;
+		RequestedTime = 0.0f;
+	}
+};
+
+/**
+ * A subsystem managing spawning of actors for all mass subsystems
+ */
+UCLASS()
+class MASSACTORS_API UMassActorSpawnerSubsystem : public UWorldSubsystem
+{
+	GENERATED_BODY()
+public:
+
+	// USubsystem BEGIN
+	virtual void Initialize(FSubsystemCollectionBase& Collection) override;
+	virtual void Deinitialize() override;
+	static void AddReferencedObjects(UObject* InThis, FReferenceCollector& Collector);
+	// USubsystem END
+
+	/** Request an actor to spawn
+	 * Note: If you do not provide a spawn delegate, the requester is responsible to remove the request by hand.
+	 * It will be auto removed after the execution of the spawn delegate.
+	 * @param InSpawnRequest the spawn request parameters, You can provide any type of UStruct as long at it derives from FMassActorSpawnRequest. This let you add more spawning information.
+	 */
+	template< typename T, typename = typename TEnableIf<TIsDerivedFrom<typename TRemoveReference<T>::Type, FMassActorSpawnRequest>::IsDerived, void>::Type >
+	FMassHandle_ActorSpawnRequest RequestActorSpawn(const T& InSpawnRequest)
+	{
+		return RequestActorSpawnInternal(FConstStructView::Make(InSpawnRequest));
+	}
+
+	/** Retries a failed spawn request
+	 * @param SpawnRequestHandle the spawn request handle to retry
+	 */
+	void RetryActorSpawnRequest(FMassHandle_ActorSpawnRequest SpawnRequestHandle);
+
+	/**
+	 * Removes a spawn request
+	 * The only time a spawn request cannot be removed is when its state is processing
+	 * Also spawn requests are auto removed if you provided a spawn delegate after it was being executed.
+	 * @param SpawnRequestHandle [IN/OUT] the spawn request handle to remove
+	 * @return true if successfully removed the request
+	 */
+	 bool RemoveActorSpawnRequest(FMassHandle_ActorSpawnRequest& SpawnRequestHandle);
+
+	/**
+	 * Returns the stored spawn request from the handle, useful to update the transform
+	 * @param SpawnRequestHandle the spawn request handle to get the request from
+	 * @return The spawn request cast in the provided template argument
+	 */
+	template<typename T>
+	const T& GetSpawnRequest(FMassHandle_ActorSpawnRequest SpawnRequestHandle) const
+	{
+		check(SpawnRequestHandleManager.IsValidHandle(SpawnRequestHandle));
+		check(SpawnRequests.IsValidIndex(SpawnRequestHandle.GetIndex()));
+		return SpawnRequests[SpawnRequestHandle.GetIndex()].Get<T>();
+	}
+
+	/**
+	 * Returns the stored spawn request from the handle, useful to update the transform
+	 * @param SpawnRequestHandle the spawn request handle to get the request from
+	 * @return The spawn request cast in the provided template argument
+	 */
+	template<typename T> 
+	T& GetMutableSpawnRequest(FMassHandle_ActorSpawnRequest SpawnRequestHandle)
+	{
+		check(SpawnRequestHandleManager.IsValidHandle(SpawnRequestHandle));
+		check(SpawnRequests.IsValidIndex(SpawnRequestHandle.GetIndex()));
+		return SpawnRequests[SpawnRequestHandle.GetIndex()].GetMutable<T>();
+	}
+
+	/**
+	 * Destroy an actor 
+	 * @param Actor to destroy
+	 * @param bImmediate to do the destruction immediately, otherwise will be queued up for later
+	 */
+	void DestroyActor(AActor* Actor, bool bImmediate = false);
+	
+	void EnableActorPooling();
+	void DisableActorPooling();
+
+	void ReleaseAllResources();
+
+protected:
+
+	/** Called at the start of the PrePhysics mass processing phase and calls ProcessPendingSpawningRequest */ 
+	void OnPrePhysicsPhaseStarted(const float DeltaSeconds);
+	
+	/** Called at the end of the PrePhysics mass processing phase and calls ProcessPendingDestruction */ 
+	void OnPrePhysicsPhaseFinished(const float DeltaSeconds); 
+	
+	/** Retrieve what would be the next best spawning request to spawn, can be overridden to have different logic
+	 *  Default implementation is the first valid request in the list, no interesting logic yet
+	 *  @return the next best handle to spawn. */
+	virtual FMassHandle_ActorSpawnRequest GetNextRequestToSpawn() const;
+
+	virtual AActor* SpawnOrRetrieveFromPool(const FStructView& SpawnRequest);
+
+	/** Actual code that will spawn the actor, overridable by subclass if need to be.
+	 *  @return spawned actor if succeeded. */
+	virtual AActor* SpawnActor(const FStructView& SpawnRequest) const;
+
+	/** Go through the spawning request and spawn them until we reach the budget 
+	 * @param MaxTimeSlicePerTick is the budget in seconds allowed to do spawning */
+	void ProcessPendingSpawningRequest(const double MaxTimeSlicePerTick);
+
+	/** Go through the queued actors to destroy and destroy them until we reach the budget
+	 * @param MaxTimeSlicePerTick is the budget in seconds allowed to do destruction */
+	void ProcessPendingDestruction(const double MaxTimeSlicePerTick);
+
+	/** Try releasing this actor to pool if possible 
+	 * @param Actor to release to the bool
+	 * @return true if the actor was actually released to the pool */
+	bool ReleaseActorToPool(AActor* Actor);
+
+	/** Internal generic request actor spawn to make sure the request derives from FMassActorSpawnRequest 
+	 *  @param SpawnRequest the spawn request parameters, We are allowing any type of spawn request, let's store it internally as a FInstancedStruct. This parameter is the FStructView over provide user struct */
+	FMassHandle_ActorSpawnRequest RequestActorSpawnInternal(const FConstStructView SpawnRequest);
+
+protected:
+
+	UPROPERTY()
+	TArray<FInstancedStruct> SpawnRequests;
+
+	UPROPERTY()
+	TArray<AActor*> ActorsToDestroy;
+
+	UPROPERTY()
+	TArray<AActor*> DeactivatedActorsToDestroy;
+
+	bool bActorPoolingEnabled = true;
+
+	TMap<TSubclassOf<AActor>, TArray<AActor*>> PooledActors;
+
+	FMassHandleManager_ActorSpawnRequest SpawnRequestHandleManager;
+
+	std::atomic<uint32> RequestSerialNumberCounter;
+};
diff --git a/Engine/Plugins/Runtime/MassGameplay/Source/MassActors/Public/MassActorSubsystem.h b/Engine/Plugins/Runtime/MassGameplay/Source/MassActors/Public/MassActorSubsystem.h
new file mode 100644
--- /dev/null
+++ b/Engine/Plugins/Runtime/MassGameplay/Source/MassActors/Public/MassActorSubsystem.h
@@ -0,0 +1,112 @@
+// Copyright Epic Games, Inc. All Rights Reserved.
+
+#pragma once
+
+#include "LWComponentTypes.h"
+#include "MassCommonFragments.h"
+#include "MassEntityTemplate.h"
+#include "Misc/MTAccessDetector.h"
+#include "UObject/ObjectKey.h"
+#include "Subsystems/WorldSubsystem.h"
+#include "MassActorSubsystem.generated.h"
+
+struct FMassHandle;
+class AActor;
+class UPipeEntitySubsystem;
+
+/**
+ * Light weight component data to save the actor pointer of a light weight entity if it exist
+ */
+USTRUCT()
+struct MASSACTORS_API FDataFragment_Actor : public FDataFragment_ObjectWrapper
+{
+	GENERATED_BODY()
+
+	/**
+	 * Set the actor associated to a mass agent, will also keep the map back in MassActorSubsystem up to date.
+	 * @param MassAgent to associated with the actor
+	 * @param InActor the actor associated with the mass agent
+	 * @param bInIsOwnedByMass tell whether the actors was spawned by mass(MassVisualization) or externally(ReplicatedActors)
+	 */
+	void SetAndUpdateHandleMap(const FMassHandle MassAgent, AActor* InActor, const bool bInIsOwnedByMass);
+
+	/** Resets the actor pointed by this fragment, will also keep the map back in MassActorSubsystem up to date */
+	void ResetAndUpdateHandleMap();
+
+	/**
+	 * Set the actor associated to a mass agent, will NOT keep map back in MassActorSubsystem up to date.
+	 * The user needs to do the map update by hand.
+	 * Useful in multithreaded environment, to queue the update of the map inside a deferred command
+	 * @param MassAgent to associated with the actor
+	 * @param InActor the actor associated with the mass agent
+	 * @param bInIsOwnedByMass tell whether the actors was spawned by mass(MassVisualization) or externally(ReplicatedActors)
+	 */
+	void SetNoHandleMapUpdate(const FMassHandle MassAgent, AActor* InActor, const bool bInIsOwnedByMass);
+
+	/** Resets the actor pointed by this fragment, will NOT keep map back in MassActorSubsystem up to date.
+	 * The user needs to do the map update by hand.
+	 * Useful in multithreaded environment, to queue the update of the map inside a deferred command
+	 */
+	void ResetNoHandleMapUpdate();
+
+
+	/** @return none const pointer to the actor	*/
+	FORCEINLINE AActor* GetMutable() { return Actor.Get(); }
+
+	/** @return none const pointer to the actor	only if owned by mass */
+	FORCEINLINE AActor* GetOwnedByMassMutable() { return bIsOwnedByMass ? Actor.Get() : nullptr; }
+
+	/** @return none const pointer to the actor	only if owned by mass */
+	FORCEINLINE const AActor* Get() const { return Actor.Get(); }
+
+	/** @return if the actor is owned by mass */
+	FORCEINLINE bool IsOwnedByMass() const { return bIsOwnedByMass; }
+
+	/** @return if the actor is a valid pointer */
+	FORCEINLINE bool IsValid() const { return Actor.IsValid(); }
+
+private:
+	TWeakObjectPtr<AActor> Actor;
+
+	/** Ownership of the actor */
+	bool bIsOwnedByMass = false;
+};
+
+/**
+ * A subsystem managing communication between Actors and Mass
+ */
+UCLASS()
+class MASSACTORS_API UMassActorSubsystem : public UWorldSubsystem
+{
+	GENERATED_BODY()
+
+public:
+	/** Get mass handle from an actor */
+	FMassHandle GetHandleFromActor(const TObjectKey<const AActor> Actor);
+
+	/** Set the mass handle associated to an actor */
+	void SetHandleForActor(const TObjectKey<const AActor> Actor, const FMassHandle Handle);
+
+	/** Remove entry associated to an actor */
+	void RemoveHandleForActor(const TObjectKey<const AActor> Actor);
+
+	/** Get an actor pointer from a mass handle */
+	AActor* GetActorFromHandle(const FMassHandle Handle) const;
+
+	/** 
+	 *  Removes the connection between Actor and the given entity. Does all the required book keeping 
+	 *  (as opposed to straight up RemoveHandleForActor call). If the Handle doesn't match Actor no action is taken.
+	 */
+	void DisconnectActor(const TObjectKey<const AActor> Actor, const FMassHandle Handle);
+
+protected:
+	// USubsystem BEGIN
+	virtual void Initialize(FSubsystemCollectionBase& Collection) override;
+	// USubsystem END
+	
+	TMap<TObjectKey<const AActor>, FLWEntity> ActorHandleMap;
+	UE_MT_DECLARE_RW_ACCESS_DETECTOR(ActorHandleMapDetector);
+	
+	UPROPERTY()
+	UPipeEntitySubsystem* EntitySystem;
+};
diff --git a/Engine/Plugins/Runtime/MassGameplay/Source/MassActors/Public/MassAgentComponent.h b/Engine/Plugins/Runtime/MassGameplay/Source/MassActors/Public/MassAgentComponent.h
new file mode 100644
--- /dev/null
+++ b/Engine/Plugins/Runtime/MassGameplay/Source/MassActors/Public/MassAgentComponent.h
@@ -0,0 +1,163 @@
+// Copyright Epic Games, Inc. All Rights Reserved.
+
+#pragma once
+
+#include "Components/ActorComponent.h"
+#include "MassCommonTypes.h"
+#include "MassEntityTemplate.h"
+#include "MassEntityConfigAsset.h"
+#include "MassAgentComponent.generated.h"
+
+
+UENUM()
+enum class EAgentComponentState : uint8
+{
+	None, // Default starting state
+
+	// States of actors needing mass entity creation
+	EntityPendingCreation, // Actor waiting for entity creation
+	EntityCreated, // Actor with a created entity
+
+	// States are for Actor driven by Mass (puppet) needing fragments initialization
+	PuppetPendingInitialization, // Puppet actor queued for fragments initialization
+	PuppetInitialized, // Puppet actor with all initialized fragments
+	PuppetPaused, // Puppet actor with all fragments removed 
+	PuppetPendingReplication, // Replicated puppet actor waiting for NetID
+	PuppetReplicatedOrphan, // Replicated puppet actor not associated to a MassAgent
+};
+
+
+/** 
+ *  There are two primary use cases for this component:
+ *  1. If placed on an AActor blueprint can it lets the used specify additional fragments that will be created for 
+ *     entities spawned based on this given blueprint. 
+ *  2. If present present on an actor in the world it makes it communicate with the MassSimulation which will create an 
+ *     entity representing given actor. Use case 1) will also applicable in this case. The component is unregistered by 
+ *     default and requires manual enabling via a 'Enable' call
+ * 
+ *  @todo use case 2) is currently sitting in a shelved CL of mine. Will be worked on next.
+ */
+UCLASS(Blueprintable, ClassGroup = Mass, meta = (BlueprintSpawnableComponent), hidecategories = (Sockets, Collision))
+class MASSACTORS_API UMassAgentComponent : public UActorComponent
+{
+	GENERATED_BODY()
+
+public:
+
+	UMassAgentComponent();
+
+protected:
+	virtual void OnRegister() override;
+	virtual void OnUnregister() override;
+
+#if WITH_EDITOR
+	virtual void PostInitProperties() override;
+	virtual void PostLoad() override;
+#endif // WITH_EDITOR
+
+public:
+	// Support poolable mass agents going into stasis
+	void RegisterWithAgentManager();
+	void UnregisterWithAgentManager();
+	bool IsReadyForPooling() const;
+
+	/** @todo to enforce encapsulation we could make this protected and have a UMassAgentSubsystem friend. I'm not sure I like it though. */
+
+	/** Methods handling the state for actors that needs mass entity creation*/
+	FMassHandle GetEntityHandle() const { return AgentHandle; }
+	void SetEntityHandle(FMassHandle NewHandle);
+	void ClearEntityHandle();
+	void EntityCreationPending();
+	void EntityCreationAborted();
+	bool IsEntityPendingCreation() const { return (State == EAgentComponentState::EntityPendingCreation); }
+
+	/** Methods handling the state of puppet actors that needs fragments initialization */
+	void SetPuppetHandle(FMassHandle NewHandle);
+	void PuppetInitializationPending();
+	void PuppetInitializationDone();
+	void PuppetInitializationAborted();
+	void PuppetUnregistrationDone();
+	bool IsPuppetPendingInitialization() const { return (State == EAgentComponentState::PuppetPendingInitialization); }
+	bool IsPuppetReplicationPending() const { return (State == EAgentComponentState::PuppetPendingReplication); }
+	bool IsPuppet() const { return State == EAgentComponentState::PuppetInitialized || State == EAgentComponentState::PuppetPendingInitialization || State == EAgentComponentState::PuppetPaused; }
+	/**
+	 * Re-adds/Removes all puppet fragments added on the mass agent
+	 * This is only supported in Puppet flow
+	 * @param bPause true to pause or false to unpause
+	 */
+	void PausePuppet(const bool bPause);
+	/* @return boolean whether this component was paused via PausePuppet method */
+	bool IsPuppetPaused() const { return State == EAgentComponentState::PuppetPaused; }
+
+	/** Methods handling the state of a server replicated puppet */
+	void PuppetReplicationPending();
+	void SetReplicatedPuppetHandle(FMassHandle NewHandle);
+	void ClearReplicatedPuppetHandle();
+	void MakePuppetAReplicatedOrphan();
+
+	FMassEntityTemplateID GetTemplateID() const { return TemplateID; }
+
+	const FMassEntityConfig& GetEntityConfig() const { return EntityConfig; }
+	void SetEntityConfig(const FMassEntityConfig& InEntityConfig);
+
+	const FLWCompositionDescriptor& GetPuppetSpecificAddition() const { return PuppetSpecificAddition; }
+	FLWCompositionDescriptor& GetMutablePuppetSpecificAddition() { return PuppetSpecificAddition; }
+
+	/** Registers the component with the owner effectively turning it on. Calling it multiple times won't break anything  */
+	UFUNCTION(BlueprintCallable, Category = Mass)
+	void Enable();
+
+	/** Registers the component with the owner effectively turning it off */
+	UFUNCTION(BlueprintCallable, Category = Mass)
+	void Disable();
+
+	UFUNCTION(BlueprintCallable, Category = Mass)
+	void KillEntity(const bool bDestroyActor);
+
+	//~ Begin UObject Interface
+	virtual void GetLifetimeReplicatedProps(TArray<FLifetimeProperty>& OutLifetimeProps) const override;
+	//~ End UObject Interface
+
+	FMassNetworkID GetNetID() const
+	{
+		return NetID;
+	}
+
+	UFUNCTION()
+	virtual void OnRep_NetID();
+
+protected:
+
+	void SwitchToState(EAgentComponentState NewState);
+	void SetEntityHandleInternal(FMassHandle NewHandle);
+	void ClearEntityHandleInternal();
+	void DebugCheckStateConsistency();
+
+	/**
+	 *  Contains all the fragments added to the entity during puppet's initialization. Required for clean up when
+	 *  despawning puppet while the entity remains alive.
+	 */
+	FLWCompositionDescriptor PuppetSpecificAddition;
+
+	UPROPERTY(EditAnywhere, Category = "Mass")
+	FMassEntityConfig EntityConfig;
+
+#if WITH_EDITORONLY_DATA
+	/** for debugging purposes only. The Entity used at runtime is stored in AgentHandle */
+	UPROPERTY(VisibleAnywhere, Category = "Mass")
+	FLWEntity DebugEntity;
+#endif // WITH_EDITORONLY_DATA
+
+	FMassHandle AgentHandle;
+	FMassEntityTemplateID TemplateID;
+
+	EAgentComponentState State;
+
+	UPROPERTY(replicatedUsing = OnRep_NetID)
+	FMassNetworkID NetID;
+	
+#if WITH_EDITORONLY_DATA
+	UPROPERTY(EditAnywhere, Category = "Mass")
+	uint32 bAutoRegisterInEditorMode : 1;
+#endif // WITH_EDITORONLY_DATA
+};
