AggregateTicking and Pipe merged into Engine/Plugins/Runtime/MassEntity

diff --git a/Engine/Plugins/Runtime/MassGameplay/Source/MassActors/Public/MassAgentSubsystem.h b/Engine/Plugins/Runtime/MassGameplay/Source/MassActors/Public/MassAgentSubsystem.h
new file mode 100644
--- /dev/null
+++ b/Engine/Plugins/Runtime/MassGameplay/Source/MassActors/Public/MassAgentSubsystem.h
@@ -0,0 +1,148 @@
+// Copyright Epic Games, Inc. All Rights Reserved.
+
+#pragma once
+
+#include "LWComponentTypes.h"
+#include "MassEntityTypes.h"
+#include "MassEntityTemplate.h"
+#include "Subsystems/WorldSubsystem.h"
+#include "MassAgentSubsystem.generated.h"
+
+struct FMassHandle;
+class AActor;
+class UPipeEntitySubsystem;
+class UMassSpawnerSubsystem;
+class UMassAgentComponent;
+class UMassSimulationSubsystem;
+class UMassAgentSubsystem;
+class UMassAgentComponent;
+class UMassReplicationManager;
+
+namespace UE::MassActor
+{
+	DECLARE_MULTICAST_DELEGATE_OneParam(FMassAgentComponentDelegate, const UMassAgentComponent& /*AgentComponent*/);
+
+} // UE::MassActor
+
+USTRUCT()
+struct MASSACTORS_API FMassAgentInitializationQueue
+{
+	GENERATED_BODY()
+
+	UPROPERTY()
+	TArray<UMassAgentComponent*> AgentComponents;
+};
+
+/**
+ * A subsystem managing communication between Actors and Mass
+ */
+UCLASS()
+class MASSACTORS_API UMassAgentSubsystem : public UWorldSubsystem
+{
+	GENERATED_BODY()
+public:
+
+	// USubsystem BEGIN
+	virtual void Initialize(FSubsystemCollectionBase& Collection) override;
+	// USubsystem END
+
+	/** Registers given AgentComp with the Mass Simulation, including creation of a LWEntity using 
+	 *  UMassAgentComponent.LWComponentList to determine the Archetype to use. */
+	FMassEntityTemplateID RegisterAgentComponent(UMassAgentComponent& AgentComp);
+
+	/** Called for AgentComp that has already been registered with the Agent Manager to notify it that AgentComp's 
+	 *  Mass-relevant properties had changed, most notably its fragment composition. The Agent Manager will update 
+	 *  the information on Mass side potentially reallocating the associated entity to a different archetype */
+	void UpdateAgentComponent(const UMassAgentComponent& AgentComp);
+
+	/**
+	 * Removes given AgentComp instance from the system. If there's an entity created with the AgentComp
+	 *  instance then it will be destroyed.
+	 * @param AgentComp the component to unregister from the system
+	 */
+	 void UnregisterAgentComponent(UMassAgentComponent& AgentComp);
+
+	 /**
+	  * Same as UnregisterAgentComponent, but on top of that it tells the system it will never register again
+	  * @param AgentComp the component to shutdown from the system
+	  */
+	 void ShutdownAgentComponent(UMassAgentComponent& AgentComp);
+
+	/** lets the system know given agent is a puppet (an unreal-side representation of a mass entity) */
+	void MakePuppet(UMassAgentComponent& AgentComp);
+
+	/**
+	 * Notifies that this MassAgentComponent is now replicated with a valid NetID
+	 * @param AgentComp that is now replicated
+	 */
+	void NotifyMassAgentComponentReplicated(UMassAgentComponent& AgentComp);
+
+	/**
+	 * Notifies that this MassAgentComponent is now associated to a mass entity
+	 * @param AgentComp that is now associated to a mass entity
+	 */
+	void NotifyMassAgentComponentEntityAssociated(const UMassAgentComponent& AgentComp) const;
+
+	/**
+	 * Notifies that this MassAgentComponent is now detaching from its mass entity
+	 * @param AgentComp that is detaching from its mass entity
+	 */
+	void NotifyMassAgentComponentEntityDetaching(const UMassAgentComponent& AgentComp) const;
+
+	/**
+	 * @return The delegate of when MassAgentComponent gets associated to a mass entity
+	 */
+	UE::MassActor::FMassAgentComponentDelegate& GetOnMassAgentComponentEntityAssociated()
+	{
+		return OnMassAgentComponentEntityAssociated;
+	}
+
+	/**
+	 * @return The delegate of when MassAgentComponent is detaching from its mass entity
+	 */
+	UE::MassActor::FMassAgentComponentDelegate& GetOnMassAgentComponentEntityDetaching()
+	{
+		return OnMassAgentComponentEntityDetaching;
+	}
+
+protected:
+	/** 
+	 * Processes PendingAgentEntities to initialize fragments of recently created agent entities and PendingPuppets 
+	 * to create and initialize puppet-specific fragments 
+	 */
+	void HandlePendingInitialization();
+
+	/** Bound to UMassSimulationSubsystem.OnProcessingPhaseStartedDelegate and called before every processing phase start */
+	void OnProcessingPhaseStarted(const float DeltaSeconds, const EPipeProcessingPhase Phase);
+
+	/** Callback registered to the replication manager when a mass agent is added to the replication (client only) */
+	void OnMassAgentAddedToReplication(FMassNetworkID NetID, FLWEntity Entity);
+
+	/** Callback registered to the replication manager when a mass agent is removed from the replication (client only) */
+	void OnMassAgentRemovedFromReplication(FMassNetworkID NetID, FLWEntity Entity);
+
+protected:
+	UPROPERTY()
+	UPipeEntitySubsystem* EntitySystem;
+
+	UPROPERTY()
+	UMassSpawnerSubsystem* SpawnerSystem;
+
+	UPROPERTY()
+	UMassSimulationSubsystem* SimulationSystem;
+
+	UPROPERTY()
+	TMap<FMassEntityTemplateID, FMassAgentInitializationQueue> PendingAgentEntities;
+
+	UPROPERTY()
+	TMap<FMassEntityTemplateID, FMassAgentInitializationQueue> PendingPuppets;
+
+	UPROPERTY()
+	UMassReplicationManager* ReplicationManager;
+
+	UPROPERTY()
+	TMap<FMassNetworkID, UMassAgentComponent*> ReplicatedAgentComponents;
+
+	UE::MassActor::FMassAgentComponentDelegate OnMassAgentComponentEntityAssociated;
+	UE::MassActor::FMassAgentComponentDelegate OnMassAgentComponentEntityDetaching;
+};
diff --git a/Engine/Plugins/Runtime/MassGameplay/Source/MassActors/Public/MassAgentTraits.h b/Engine/Plugins/Runtime/MassGameplay/Source/MassActors/Public/MassAgentTraits.h
new file mode 100644
--- /dev/null
+++ b/Engine/Plugins/Runtime/MassGameplay/Source/MassActors/Public/MassAgentTraits.h
@@ -0,0 +1,66 @@
+// Copyright Epic Games, Inc. All Rights Reserved.
+
+#pragma once
+
+#include "CoreMinimal.h"
+#include "MassEntityTraitBase.h"
+#include "MassTranslator.h"
+#include "MassAgentTraits.generated.h"
+
+class USceneComponent;
+
+UCLASS(Abstract)
+class MASSACTORS_API UMassAgentSyncTrait : public UMassEntityTraitBase
+{
+	GENERATED_BODY()
+
+public:
+	EMassTranslationDirection GetSyncDirection() const { return SyncDirection; }
+	void SetSyncDirection(const EMassTranslationDirection NewDirection) { SyncDirection = NewDirection; }
+
+protected:
+	UPROPERTY(EditAnywhere, Category = Mass)
+	EMassTranslationDirection SyncDirection = EMassTranslationDirection::BothWays;
+};
+
+/** The trait initializes the entity with actor capsule component's radius. In addition, if bSyncTransform is true 
+ *  the trait keeps actor capsule component's and entity's transforms in sync. */
+UCLASS(BlueprintType, EditInlineNew, CollapseCategories, meta = (DisplayName = "Agent Capsule Collision Sync"))
+class MASSACTORS_API UMassAgentCapsuleCollisionSyncTrait : public UMassAgentSyncTrait
+{
+	GENERATED_BODY()
+
+protected:
+	virtual void BuildTemplate(FMassEntityTemplateBuildContext& BuildContext, UWorld& World) const override;
+
+	UPROPERTY(EditAnywhere, Category = Mass)
+	bool bSyncTransform = true;
+};
+
+UCLASS(BlueprintType, EditInlineNew, CollapseCategories, meta = (DisplayName = "Agent Movement Sync"))
+class MASSACTORS_API UMassAgentMovementSyncTrait : public UMassAgentSyncTrait
+{
+	GENERATED_BODY()
+
+protected:
+	virtual void BuildTemplate(FMassEntityTemplateBuildContext& BuildContext, UWorld& World) const override;
+};
+
+UCLASS(BlueprintType, EditInlineNew, CollapseCategories, meta = (DisplayName = "Agent Orientation Sync"))
+class MASSACTORS_API UMassAgentOrientationSyncTrait : public UMassAgentSyncTrait
+{
+	GENERATED_BODY()
+
+	protected:
+	virtual void BuildTemplate(FMassEntityTemplateBuildContext& BuildContext, UWorld& World) const override;
+};
+
+UCLASS(BlueprintType, EditInlineNew, CollapseCategories, meta = (DisplayName = "Agent Feet Location Sync"))
+class MASSACTORS_API UMassAgentFeetLocationSyncTrait : public UMassAgentSyncTrait
+{
+	GENERATED_BODY()
+
+protected:
+	virtual void BuildTemplate(FMassEntityTemplateBuildContext& BuildContext, UWorld& World) const override;
+};
+
diff --git a/Engine/Plugins/Runtime/MassGameplay/Source/MassActors/Public/Translators/MassCapsuleComponentTranslators.h b/Engine/Plugins/Runtime/MassGameplay/Source/MassActors/Public/Translators/MassCapsuleComponentTranslators.h
new file mode 100644
--- /dev/null
+++ b/Engine/Plugins/Runtime/MassGameplay/Source/MassActors/Public/Translators/MassCapsuleComponentTranslators.h
@@ -0,0 +1,66 @@
+// Copyright Epic Games, Inc. All Rights Reserved.
+
+#pragma once
+
+#include "MassCommonFragments.h"
+#include "MassTranslator.h"
+#include "MassCapsuleComponentTranslators.generated.h"
+
+
+class UCapsuleComponent;
+struct FDataFragment_AgentRadius;
+
+USTRUCT()
+struct FDataFragment_CapsuleComponentWrapper : public FDataFragment_ObjectWrapper
+{
+	GENERATED_BODY()
+	TWeakObjectPtr<UCapsuleComponent> Component;
+};
+
+/**
+ * @todo TBD
+ * I'm a bit on a fence regarding having separate tags per copy direction. My concern is that we can end up with a very 
+ * dispersed entity population spread across multiple archetypes storing a small number of entities each. An alternative
+ * would be to have a property on the Wrapper fragment, but that doesn't sit well with me either, since that data would be 
+ * essentially static, meaning it will (in most cases) never change for a given entity, and we could waste a lot of time 
+ * iterating over fragments just to check that specific value.
+ */
+USTRUCT()
+struct FMassCapsuleTransformCopyToMassTag : public FComponentTag
+{
+	GENERATED_BODY()
+};
+
+UCLASS()
+class MASSACTORS_API UMassCapsuleTransformToMassTranslator : public UMassTranslator
+{
+	GENERATED_BODY()
+public:
+	UMassCapsuleTransformToMassTranslator();
+
+protected:
+	virtual void ConfigureQueries() override;
+	virtual void Execute(UEntitySubsystem& EntitySubsystem, FLWComponentSystemExecutionContext& Context) override;	
+
+	FLWComponentQuery EntityQuery;
+};
+
+USTRUCT()
+struct FMassCapsuleTransformCopyToActorTag : public FComponentTag
+{
+	GENERATED_BODY()
+};
+
+UCLASS()
+class MASSACTORS_API UMassTransformToActorCapsuleTranslator : public UMassTranslator
+{
+	GENERATED_BODY()
+public:
+	UMassTransformToActorCapsuleTranslator();
+
+protected:
+	virtual void ConfigureQueries() override;
+	virtual void Execute(UEntitySubsystem& EntitySubsystem, FLWComponentSystemExecutionContext& Context) override;
+
+	FLWComponentQuery EntityQuery;
+};
diff --git a/Engine/Plugins/Runtime/MassGameplay/Source/MassActors/Public/Translators/MassCharacterMovementTranslators.h b/Engine/Plugins/Runtime/MassGameplay/Source/MassActors/Public/Translators/MassCharacterMovementTranslators.h
new file mode 100644
--- /dev/null
+++ b/Engine/Plugins/Runtime/MassGameplay/Source/MassActors/Public/Translators/MassCharacterMovementTranslators.h
@@ -0,0 +1,136 @@
+// Copyright Epic Games, Inc. All Rights Reserved.
+
+#pragma once
+
+#include "MassCommonFragments.h"
+#include "MassTranslator.h"
+#include "MassCharacterMovementTranslators.generated.h"
+
+class UCharacterMovementComponent;
+struct FDataFragment_CharacterMovementComponentWrapper;
+struct FNavLocationComponent;
+struct FDataFragment_Transform;
+struct FMassVelocityFragment;
+
+USTRUCT()
+struct FDataFragment_CharacterMovementComponentWrapper : public FDataFragment_ObjectWrapper
+{
+	GENERATED_BODY()
+	TWeakObjectPtr<UCharacterMovementComponent> Component;
+};
+
+USTRUCT()
+struct FMassCharacterMovementCopyToMassTag : public FComponentTag
+{
+	GENERATED_BODY()
+};
+
+UCLASS()
+class MASSACTORS_API UMassCharacterMovementToMassTranslator : public UMassTranslator
+{
+	GENERATED_BODY()
+
+public:
+	UMassCharacterMovementToMassTranslator();
+
+protected:
+	virtual void ConfigureQueries() override;
+	virtual void Execute(UEntitySubsystem& EntitySubsystem, FLWComponentSystemExecutionContext& Context) override;
+
+	FLWComponentQuery EntityQuery;
+};
+
+USTRUCT()
+struct FMassCharacterMovementCopyToActorTag : public FComponentTag
+{
+	GENERATED_BODY()
+};
+
+
+UCLASS()
+class MASSACTORS_API UMassCharacterMovementToActorTranslator : public UMassTranslator
+{
+	GENERATED_BODY()
+
+public:
+	UMassCharacterMovementToActorTranslator();
+
+protected:
+	virtual void ConfigureQueries() override;
+	virtual void Execute(UEntitySubsystem& EntitySubsystem, FLWComponentSystemExecutionContext& Context) override;
+
+	FLWComponentQuery EntityQuery;
+};
+
+
+USTRUCT()
+struct FMassCharacterOrientationCopyToMassTag : public FComponentTag
+{
+	GENERATED_BODY()
+};
+
+UCLASS()
+class MASSACTORS_API UMassCharacterOrientationToMassTranslator : public UMassTranslator
+{
+	GENERATED_BODY()
+
+public:
+	UMassCharacterOrientationToMassTranslator();
+
+protected:
+	virtual void ConfigureQueries() override;
+	virtual void Execute(UEntitySubsystem& EntitySubsystem, FLWComponentSystemExecutionContext& Context) override;
+
+	FLWComponentQuery EntityQuery;
+};
+
+USTRUCT()
+struct FMassCharacterOrientationCopyToActorTag : public FComponentTag
+{
+	GENERATED_BODY()
+};
+
+
+UCLASS()
+class MASSACTORS_API UMassCharacterOrientationToActorTranslator : public UMassTranslator
+{
+	GENERATED_BODY()
+
+public:
+	UMassCharacterOrientationToActorTranslator();
+
+protected:
+	virtual void ConfigureQueries() override;
+	virtual void Execute(UEntitySubsystem& EntitySubsystem, FLWComponentSystemExecutionContext& Context) override;
+
+	FLWComponentQuery EntityQuery;
+};
+
+UCLASS()
+class MASSACTORS_API UMassFragmentInitializer_NavLocation : public UMassFragmentInitializer
+{
+	GENERATED_BODY()
+public:
+	UMassFragmentInitializer_NavLocation();
+
+protected:
+	virtual void ConfigureQueries() override;
+	virtual void Execute(UEntitySubsystem& EntitySubsystem, FLWComponentSystemExecutionContext& Context) override;
+
+	FLWComponentQuery EntityQuery;
+};
+
+UCLASS()
+class MASSACTORS_API UMassFragmentInitializer_Transform : public UMassFragmentInitializer
+{
+	GENERATED_BODY()
+
+public:
+	UMassFragmentInitializer_Transform();
+
+protected:
+	virtual void ConfigureQueries() override;
+	virtual void Execute(UEntitySubsystem& EntitySubsystem, FLWComponentSystemExecutionContext& Context) override;
+
+	FLWComponentQuery EntityQuery;
+};
\ No newline at end of file
diff --git a/Engine/Plugins/Runtime/MassGameplay/Source/MassActors/Public/Translators/MassSceneComponentLocationTranslator.h b/Engine/Plugins/Runtime/MassGameplay/Source/MassActors/Public/Translators/MassSceneComponentLocationTranslator.h
new file mode 100644
--- /dev/null
+++ b/Engine/Plugins/Runtime/MassGameplay/Source/MassActors/Public/Translators/MassSceneComponentLocationTranslator.h
@@ -0,0 +1,58 @@
+// Copyright Epic Games, Inc. All Rights Reserved.
+
+#pragma once
+
+#include "MassCommonFragments.h"
+#include "MassTranslator.h"
+#include "MassSceneComponentLocationTranslator.generated.h"
+
+
+USTRUCT()
+struct FMassSceneComponentWrapperFragment : public FDataFragment_ObjectWrapper
+{
+	GENERATED_BODY()
+	TWeakObjectPtr<USceneComponent> Component;
+};
+
+USTRUCT()
+struct FMassSceneComponentLocationCopyToMassTag : public FComponentTag
+{
+	GENERATED_BODY()
+};
+
+UCLASS()
+class MASSACTORS_API UMassSceneComponentLocationToMassTranslator : public UMassTranslator
+{
+	GENERATED_BODY()
+
+public:
+	UMassSceneComponentLocationToMassTranslator();
+
+protected:
+	virtual void ConfigureQueries() override;
+	virtual void Execute(UEntitySubsystem& EntitySubsystem, FLWComponentSystemExecutionContext& Context) override;
+
+	FLWComponentQuery EntityQuery;
+};
+
+
+USTRUCT()
+struct FMassSceneComponentLocationCopyToActorTag : public FComponentTag
+{
+	GENERATED_BODY()
+};
+
+UCLASS()
+class MASSACTORS_API UMassSceneComponentLocationToActorTranslator : public UMassTranslator
+{
+	GENERATED_BODY()
+
+public:
+	UMassSceneComponentLocationToActorTranslator();
+
+protected:
+	virtual void ConfigureQueries() override;
+	virtual void Execute(UEntitySubsystem& EntitySubsystem, FLWComponentSystemExecutionContext& Context) override;
+
+	FLWComponentQuery EntityQuery;
+};
\ No newline at end of file
diff --git a/Engine/Plugins/Runtime/MassGameplay/Source/MassActors/Public/Translators/MassTranslators_BehaviorTree.h b/Engine/Plugins/Runtime/MassGameplay/Source/MassActors/Public/Translators/MassTranslators_BehaviorTree.h
new file mode 100644
--- /dev/null
+++ b/Engine/Plugins/Runtime/MassGameplay/Source/MassActors/Public/Translators/MassTranslators_BehaviorTree.h
@@ -0,0 +1,32 @@
+// Copyright Epic Games, Inc. All Rights Reserved.
+
+#pragma once
+
+#include "MassTranslator.h"
+#include "MassCommonFragments.h"
+#include "MassTranslators_BehaviorTree.generated.h"
+
+//////////////////////////////////////////////////////////////////////////
+class UBehaviorTreeComponent;
+
+USTRUCT()
+struct FDataFragment_BehaviorTreeComponentWrapper : public FDataFragment_ObjectWrapper
+{
+	GENERATED_BODY()
+	TWeakObjectPtr<UBehaviorTreeComponent> Component;
+};
+
+UCLASS()
+class MASSACTORS_API UMassTranslator_BehaviorTree : public UMassTranslator
+{
+	GENERATED_BODY()
+
+public:
+	UMassTranslator_BehaviorTree();
+
+protected:
+	virtual void ConfigureQueries() override;
+	virtual void Execute(UEntitySubsystem& EntitySubsystem, FLWComponentSystemExecutionContext& Context) override {}
+
+	FLWComponentQuery EntityQuery;
+};
diff --git a/Engine/Plugins/Runtime/MassGameplay/Source/MassCommon/Private/MassCommonTypes.cpp b/Engine/Plugins/Runtime/MassGameplay/Source/MassCommon/Private/MassCommonTypes.cpp
new file mode 100644
--- /dev/null
+++ b/Engine/Plugins/Runtime/MassGameplay/Source/MassCommon/Private/MassCommonTypes.cpp
@@ -0,0 +1,151 @@
+// Copyright Epic Games, Inc. All Rights Reserved.
+
+#include "MassCommonTypes.h"
+#include "EntitySubsystem.h"
+#include "Math/ColorList.h"
+
+#if WITH_MASS_DEBUG
+namespace UE { namespace MassDebug {
+
+// First Id of a range of lightweight entity for which we want to activate debug information
+int32 DebugEntityBegin = INDEX_NONE;
+
+// Last Id of a range of lightweight entity for which we want to activate debug information
+int32 DebugEntityEnd = INDEX_NONE;
+
+static FAutoConsoleCommand SetDebugEntityRange(
+	TEXT("ai.debug.mass.SetDebugEntityRange"),
+	TEXT("Range of lightweight entity IDs that we want to debug.")
+	TEXT("Usage: \"ai.debug.mass.SetDebugEntityRange <FirstEntity> <LastEntity>\""),
+	FConsoleCommandWithArgsDelegate::CreateLambda([](const TArray<FString>& Args)
+		{
+			if (Args.Num() != 2)
+			{
+				UE_LOG(LogConsoleResponse, Display, TEXT("Error: Expecting 2 parameters"));
+				return;
+			}
+
+			int32 FirstID = INDEX_NONE;
+			int32 LastID = INDEX_NONE;
+			if (!LexTryParseString<int32>(FirstID, *Args[0]))
+			{
+				UE_LOG(LogConsoleResponse, Display, TEXT("Error: first parameter must be an integer"));
+				return;
+			}
+			
+			if (!LexTryParseString<int32>(LastID, *Args[1]))
+			{
+				UE_LOG(LogConsoleResponse, Display, TEXT("Error: second parameter must be an integer"));
+				return;
+			}
+
+			DebugEntityBegin = FirstID;
+			DebugEntityEnd = LastID;
+		}));
+
+static FAutoConsoleCommand SetDebugEntity(
+	TEXT("ai.debug.mass.DebugEntity"),
+	TEXT("ID of a lightweight entity that we want to debug.")
+	TEXT("Usage: \"ai.debug.mass.DebugEntity <Entity>\""),
+	FConsoleCommandWithArgsDelegate::CreateLambda([](const TArray<FString>& Args)
+		{
+			if (Args.Num() != 1)
+			{
+				UE_LOG(LogConsoleResponse, Display, TEXT("Error: Expecting 1 parameter"));
+				return;
+			}
+
+			int32 ID = INDEX_NONE;
+			if (!LexTryParseString<int32>(ID, *Args[0]))
+			{
+				UE_LOG(LogConsoleResponse, Display, TEXT("Error: parameter must be an integer"));
+				return;
+			}
+
+			DebugEntityBegin = ID;
+			DebugEntityEnd = ID;
+		}));
+
+static FAutoConsoleCommand ResetDebugEntity(
+	TEXT("ai.debug.mass.ResetDebugEntity"),
+	TEXT("Disables lightweight entities debugging.")
+	TEXT("Usage: \"ai.debug.mass.ResetDebugEntity\""),
+	FConsoleCommandWithArgsDelegate::CreateLambda([](const TArray<FString>& Args)
+		{
+			DebugEntityBegin = INDEX_NONE;
+			DebugEntityEnd = INDEX_NONE;
+		}));
+
+bool HasDebugEntities()
+{
+	return DebugEntityBegin != INDEX_NONE && DebugEntityEnd != INDEX_NONE;
+}
+	
+bool IsDebuggingEntity(FLWEntity Entity, FColor* OutEntityColor)
+{
+	const int32 EntityIdx = Entity.Index;
+	const bool bIsDebuggingEntity = (DebugEntityBegin != INDEX_NONE && DebugEntityEnd != INDEX_NONE && DebugEntityBegin <= EntityIdx && EntityIdx <= DebugEntityEnd);
+	
+	if (bIsDebuggingEntity && OutEntityColor != nullptr)
+	{
+		*OutEntityColor = GetEntityDebugColor(Entity);
+	}
+
+	return bIsDebuggingEntity;
+}
+
+FColor GetEntityDebugColor(FLWEntity Entity)
+{
+	const int32 EntityIdx = Entity.Index;
+	return EntityIdx != INDEX_NONE ? GColorList.GetFColorByIndex(EntityIdx % GColorList.GetColorsNum()) : FColor::Black;
+}
+
+}} // namespace UE::MassDebug
+#endif // WITH_MASS_DEBUG
+
+const FMassHandle FMassHandle::InvalidHandle(UEntitySubsystem::InvalidEntity);
+
+//----------------------------------------------------------------------//
+// FMassUniqueFragmentCollection 
+//----------------------------------------------------------------------//
+void FMassUniqueFragmentCollection::Add(const FInstancedStruct& InFragment)
+{
+	if (InFragment.GetScriptStruct() && FragmentInstances.FindByPredicate(FSameTypeScriptStructPredicate(InFragment)) == nullptr)
+	{
+		FragmentInstances.Add(InFragment);
+	}
+}
+
+void FMassUniqueFragmentCollection::Add(const UScriptStruct* InFragmentType)
+{
+	if (InFragmentType && FragmentInstances.FindByPredicate(FSameTypeScriptStructPredicate(InFragmentType)) == nullptr)
+	{
+		FragmentInstances.Add(FInstancedStruct(InFragmentType));
+		FragmentBitSet.Add(*InFragmentType);
+	}
+}
+
+void FMassUniqueFragmentCollection::Append(const FMassUniqueFragmentCollection& Other)
+{
+	FragmentInstances.Reserve(FragmentInstances.Num() + Other.FragmentInstances.Num());
+	for (const FInstancedStruct& NewFragment : Other.FragmentInstances)
+	{
+		Add(NewFragment);
+	}
+	// not shrinking on purpose, we can take the hit
+}
+
+void FMassUniqueFragmentCollection::DebugOutputDescription(FOutputDevice& Ar) const
+{
+#if WITH_MASS_DEBUG
+	const bool bAutoLineEnd = Ar.GetAutoEmitLineTerminator();
+	Ar.SetAutoEmitLineTerminator(false);
+
+	for (const FInstancedStruct& Fragment : FragmentInstances)
+	{
+		const UScriptStruct* StructType = Fragment.GetScriptStruct();
+		Ar.Logf(TEXT("%s, "), *GetNameSafe(StructType));
+	}
+	Ar.SetAutoEmitLineTerminator(bAutoLineEnd);
+#endif // WITH_MASS_DEBUG
+}
diff --git a/Engine/Plugins/Runtime/MassGameplay/Source/MassCommon/Private/MassCommonUtils.cpp b/Engine/Plugins/Runtime/MassGameplay/Source/MassCommon/Private/MassCommonUtils.cpp
new file mode 100644
--- /dev/null
+++ b/Engine/Plugins/Runtime/MassGameplay/Source/MassCommon/Private/MassCommonUtils.cpp
@@ -0,0 +1,23 @@
+// Copyright Epic Games, Inc. All Rights Reserved.
+
+#include "MassCommonUtils.h"
+
+namespace UE::Mass::Utils
+{
+	TArray<FLWEntity> EntityQueueToArray(TQueue<FLWEntity, EQueueMode::Mpsc>& EntitiesQueue, const int32 EntitiesCount)
+	{
+		check(EntitiesCount > 0);
+		TArray<FLWEntity> EntitiesArray;
+		EntitiesArray.AddUninitialized(EntitiesCount);
+
+		FLWEntity TempEntity;
+		uint32 CurrentIndex = 0;
+		while (EntitiesQueue.Dequeue(TempEntity))
+		{
+			EntitiesArray[CurrentIndex++] = TempEntity;
+		}
+		ensure(CurrentIndex == EntitiesCount);
+
+		return MoveTemp(EntitiesArray);
+	}
+} // namespace UE::Mass::Utils
\ No newline at end of file
diff --git a/Engine/Plugins/Runtime/MassGameplay/Source/MassCommon/Private/MassGenericProcessors.cpp b/Engine/Plugins/Runtime/MassGameplay/Source/MassCommon/Private/MassGenericProcessors.cpp
new file mode 100644
--- /dev/null
+++ b/Engine/Plugins/Runtime/MassGameplay/Source/MassCommon/Private/MassGenericProcessors.cpp
@@ -0,0 +1,51 @@
+// Copyright Epic Games, Inc. All Rights Reserved.
+
+#include "MassGenericProcessors.h"
+#include "MassCommonFragments.h"
+
+//----------------------------------------------------------------------//
+// URandomizeVectorProcessor
+//----------------------------------------------------------------------//
+URandomizeVectorProcessor::URandomizeVectorProcessor()
+{
+	ComponentType = FVectorComponent::StaticStruct();
+	bAutoRegisterWithProcessingPhases = false;
+}
+
+void URandomizeVectorProcessor::Execute(UEntitySubsystem& EntitySubsystem, FLWComponentSystemExecutionContext& Context)
+{
+	QUICK_SCOPE_CYCLE_COUNTER(RandomizeVectorProcessor_Run);
+	
+	EntityQuery.ForEachEntityChunk(EntitySubsystem, Context, [this](FLWComponentSystemExecutionContext& Context)
+		{
+			for (FVector& Vector : ComponentList)
+			{
+				const FRotator Rotation(0.0f, FMath::FRand() * 360.0f, 0.0f);
+				Vector = Rotation.Vector() * (MinMagnitude + FMath::FRand() * (MaxMagnitude - MinMagnitude));
+			}
+		});
+}
+
+#if WITH_EDITOR
+void URandomizeVectorProcessor::PostEditChangeProperty(FPropertyChangedEvent& PropertyChangedEvent)
+{
+	static const FName ComponentTypeName = GET_MEMBER_NAME_CHECKED(URandomizeVectorProcessor, ComponentType);
+	static const int32 SupportedSize = FVectorComponent::StaticStruct()->GetStructureSize();
+
+	Super::PostEditChangeProperty(PropertyChangedEvent);
+
+	if (PropertyChangedEvent.Property)
+	{
+		const FName PropName = PropertyChangedEvent.Property->GetFName();
+		if (PropName == ComponentTypeName)
+		{
+			if (ComponentType == nullptr
+				|| ComponentType->IsChildOf(FLWComponentData::StaticStruct()) == false
+				|| ComponentType->GetStructureSize() != SupportedSize)
+			{
+				ComponentType = FVectorComponent::StaticStruct();
+			}
+		}
+	}
+}
+#endif // WITH_EDITOR
diff --git a/Engine/Plugins/Runtime/MassGameplay/Source/MassCommon/Public/MassCommonFragments.h b/Engine/Plugins/Runtime/MassGameplay/Source/MassCommon/Public/MassCommonFragments.h
new file mode 100644
--- /dev/null
+++ b/Engine/Plugins/Runtime/MassGameplay/Source/MassCommon/Public/MassCommonFragments.h
@@ -0,0 +1,57 @@
+// Copyright Epic Games, Inc. All Rights Reserved.
+
+#pragma once
+
+#include "LWComponentTypes.h"
+#include "AI/Navigation/NavigationTypes.h"
+#include "MassCommonFragments.generated.h"
+
+
+USTRUCT()
+struct MASSCOMMON_API FDataFragment_Transform : public FLWComponentData
+{
+	GENERATED_BODY()
+
+	const FTransform& GetTransform() const { return Transform; }
+	void SetTransform(const FTransform& InTransform) { Transform = InTransform; }
+	FTransform& GetMutableTransform() { return Transform; }
+
+protected:
+	UPROPERTY(Transient)
+	FTransform Transform;
+};
+
+USTRUCT()
+struct MASSCOMMON_API FDataFragment_NavLocation : public FLWComponentData
+{
+	GENERATED_BODY()
+	NavNodeRef NodeRef = INVALID_NAVNODEREF;
+};
+
+/** this component is a kind of an experiment. I've made FMassVelocityFragment and FMassAvoidanceFragment extend it and 
+ *  use URandomizeVectorProcessor to initialize it to random variables at spawn time.
+ */
+USTRUCT()
+struct MASSCOMMON_API FVectorComponent : public FLWComponentData
+{
+	GENERATED_BODY()
+	UPROPERTY(EditAnywhere, Category = "")
+	FVector Value = FVector::ZeroVector;
+};
+
+USTRUCT()
+struct MASSCOMMON_API FDataFragment_AgentRadius : public FLWComponentData
+{
+	GENERATED_BODY()
+	UPROPERTY(EditAnywhere, Category = "")
+	float Radius = 40.f;
+};
+
+/** This is a common type for all the wrappers pointing at UObjects used to copy data from them or set data based on
+ *	Mass simulation..
+ */
+USTRUCT()
+struct FDataFragment_ObjectWrapper : public FLWComponentData
+{
+	GENERATED_BODY()
+};
diff --git a/Engine/Plugins/Runtime/MassGameplay/Source/MassCommon/Public/MassCommonTypes.h b/Engine/Plugins/Runtime/MassGameplay/Source/MassCommon/Public/MassCommonTypes.h
new file mode 100644
--- /dev/null
+++ b/Engine/Plugins/Runtime/MassGameplay/Source/MassCommon/Public/MassCommonTypes.h
@@ -0,0 +1,424 @@
+// Copyright Epic Games, Inc. All Rights Reserved.
+
+#pragma once
+
+#include "LWComponentTypes.h"
+#include "InstancedStruct.h"
+#include "SequentialID.h"
+#include "MassCommonTypes.generated.h"
+
+#define WITH_MASS_DEBUG (!(UE_BUILD_SHIPPING || UE_BUILD_SHIPPING_WITH_EDITOR || UE_BUILD_TEST) && WITH_AGGREGATETICKING_DEBUG && 1)
+
+class UStaticMesh;
+class UMaterialInterface;
+
+MASSCOMMON_API DECLARE_LOG_CATEGORY_EXTERN(LogMass, Warning, All);
+
+#if WITH_MASS_DEBUG
+namespace UE::MassDebug
+{
+	MASSCOMMON_API extern bool HasDebugEntities();
+	MASSCOMMON_API extern bool IsDebuggingEntity(FLWEntity Entity, FColor* OutEntityColor = nullptr);
+	MASSCOMMON_API extern FColor GetEntityDebugColor(FLWEntity Entity);
+} // namespace UE::MassDebug
+#endif // WITH_MASS_DEBUG
+
+namespace UE::Mass::ProcessorGroupNames
+{
+	const FName UpdateWorldFromMass = FName(TEXT("UpdateWorldFromMass"));
+	const FName SyncWorldToMass = FName(TEXT("SyncWorldToMass"));
+	const FName Behavior = FName(TEXT("Behavior"));
+	const FName Tasks = FName(TEXT("Tasks"));
+}
+
+/**
+ * Generic handle on mass agent which is associated to the specific Mass EntitySubsystem 
+ */
+struct MASSCOMMON_API FMassHandle
+{
+public:
+	FMassHandle() = default;
+
+	explicit FMassHandle(FLWEntity InLWEntity)
+		: LWEntity{ InLWEntity }
+	{}
+
+	FLWEntity GetLWEntity() const
+	{
+		return LWEntity;
+	}
+
+	bool IsValid() const
+	{
+		return LWEntity.IsSet();
+	}
+
+	FString DebugGetDescription() const 
+	{
+#if WITH_MASS_DEBUG
+		return LWEntity.DebugGetDescription();
+#else
+		return FString();
+#endif // WITH_MASS_DEBUG
+	}
+
+	const static FMassHandle InvalidHandle;
+
+	bool operator==(const FMassHandle& Other) const
+	{
+		return LWEntity == Other.LWEntity;
+	}
+
+	bool operator!=(const FMassHandle& Other) const
+	{
+		return !operator==(Other);
+	}
+
+	friend uint32 GetTypeHash(const FMassHandle& Handle)
+	{
+		return GetTypeHash(Handle.GetLWEntity());
+	}
+
+private:
+	FLWEntity LWEntity;
+};
+
+USTRUCT()
+struct FMassNetworkID : public FSequentialIDBase
+{
+	GENERATED_BODY()
+
+	FMassNetworkID() = default;
+	FMassNetworkID(const FMassNetworkID& Other) = default;
+	explicit FMassNetworkID(uint32 InID) : FSequentialIDBase(InID) {}
+};
+
+/**
+ *  No Remove* functions on purpose. Collections are to be built once and stored, never manipulated
+ */
+USTRUCT()
+struct MASSCOMMON_API FMassUniqueFragmentCollection
+{
+	GENERATED_BODY()
+
+	struct FSameTypeScriptStructPredicate
+	{
+		const UScriptStruct* TypePtr;
+		explicit FSameTypeScriptStructPredicate(const FInstancedStruct& InRef) : TypePtr(InRef.GetScriptStruct()) {}
+		explicit FSameTypeScriptStructPredicate(const UScriptStruct* InTypePtr) : TypePtr(InTypePtr) {}
+		bool operator()(const FInstancedStruct& Other) const { return Other.GetScriptStruct() == TypePtr; }
+	};
+
+	FMassUniqueFragmentCollection() = default;
+
+	/** Adds a unique fragment instance. If a fragment of this type has already been added then its values won't be 
+	 *  overridden and this call will have no effect */
+	void Add(const FInstancedStruct& InFragment);
+
+	/** Adds a unique fragment type. If a fragment of this type has already been added then no action is performed */
+	void Add(const UScriptStruct* InFragmentType);
+
+	template<typename T>
+	FORCEINLINE void Add()
+	{
+		Add(T::StaticStruct());
+	}
+
+	/** Fetches a fragment instance of a given type. If it hasn't been added before it will be created first. */
+	template<typename T>
+	T& Add_GetRef()
+	{
+		static_assert(TIsDerivedFrom<T, FLWComponentData>::IsDerived, "Given struct is not of a fragment type. If you use this FMassUniqueFragmentCollection with the EntitySubsystem you'll hit some checks");
+		if (const FInstancedStruct* ElementPtr = FragmentInstances.FindByPredicate(FSameTypeScriptStructPredicate(T::StaticStruct())))
+		{
+			return ElementPtr->GetMutable<T>();
+		}
+
+		FragmentBitSet.Add<T>();
+		return FragmentInstances.Add_GetRef(FInstancedStruct(T::StaticStruct())).GetMutable<T>();
+	}
+
+	void Append(const FMassUniqueFragmentCollection& Other);
+
+	template<typename T>
+	T& GetRefChecked()
+	{
+		static_assert(TIsDerivedFrom<T, FLWComponentData>::IsDerived, "Given struct is not of a fragment type. Make sure your struct extends FLWComponentData or its derivatives");
+		FInstancedStruct* ElementPtr = FragmentInstances.FindByPredicate(FSameTypeScriptStructPredicate(T::StaticStruct()));
+		check(ElementPtr);
+		return ElementPtr->GetMutable<T>();
+	}
+
+	int32 GetFragmentsNum() const { return FragmentInstances.Num(); }
+	TArrayView<const FInstancedStruct> GetFragments() const { return MakeArrayView(FragmentInstances); }
+	const FLWComponentBitSet& GetFragmentBitSet() const { return FragmentBitSet; }
+
+	bool IsEmpty() const { return FragmentInstances.IsEmpty(); }
+
+	void DebugOutputDescription(FOutputDevice& Ar) const;
+
+private:
+	UPROPERTY()
+	TArray<FInstancedStruct> FragmentInstances;
+
+	FLWComponentBitSet FragmentBitSet;
+};
+
+/** Float encoded in int16, 1cm accuracy. */
+USTRUCT()
+struct MASSCOMMON_API FMassInt16Real
+{
+	GENERATED_BODY()
+
+	FMassInt16Real() = default;
+	
+	explicit FMassInt16Real(const float InValue)
+	{
+		Set(InValue);
+	}
+	
+	void Set(const float InValue)
+	{
+		Value = (int16)FMath::Clamp(FMath::RoundToInt(InValue), -(int32)MAX_int16, (int32)MAX_int16);
+	}
+
+	float Get() const
+	{
+		return (float)Value;
+	}
+
+	bool operator<(const FMassInt16Real RHS) const { return Value < RHS.Value; }
+	bool operator<=(const FMassInt16Real RHS) const { return Value <= RHS.Value; }
+	bool operator>(const FMassInt16Real RHS) const { return Value > RHS.Value; }
+	bool operator>=(const FMassInt16Real RHS) const { return Value >= RHS.Value; }
+	bool operator==(const FMassInt16Real RHS) const { return Value == RHS.Value; }
+	bool operator!=(const FMassInt16Real RHS) const { return Value != RHS.Value; }
+	
+protected:
+	UPROPERTY(Transient)
+	int16 Value = 0;
+};
+
+/** Float encoded in int16, 10cm accuracy. */
+USTRUCT()
+struct MASSCOMMON_API FMassInt16Real10
+{
+	GENERATED_BODY()
+
+	static constexpr float Scale = 0.1f;
+	
+	FMassInt16Real10() = default;
+	
+	explicit FMassInt16Real10(const float InValue)
+	{
+		Set(InValue);
+	}
+	
+	void Set(const float InValue)
+	{
+		Value = (int16)FMath::Clamp(FMath::RoundToInt(InValue * Scale), -(int32)MAX_int16, (int32)MAX_int16);
+	}
+
+	float Get() const
+	{
+		return (float)Value / Scale;
+	}
+
+	bool operator<(const FMassInt16Real10 RHS) const { return Value < RHS.Value; }
+	bool operator<=(const FMassInt16Real10 RHS) const { return Value <= RHS.Value; }
+	bool operator>(const FMassInt16Real10 RHS) const { return Value > RHS.Value; }
+	bool operator>=(const FMassInt16Real10 RHS) const { return Value >= RHS.Value; }
+	bool operator==(const FMassInt16Real10 RHS) const { return Value == RHS.Value; }
+	bool operator!=(const FMassInt16Real10 RHS) const { return Value != RHS.Value; }
+
+protected:
+	UPROPERTY(Transient)
+	int16 Value = 0;
+};
+
+/** Vector which components are in range [-1..1], encoded in signed bytes. */
+USTRUCT()
+struct MASSCOMMON_API FMassSnorm8Vector
+{
+	GENERATED_BODY()
+
+	static constexpr float Scale = (float)MAX_int8;   
+
+	FMassSnorm8Vector() = default;
+	
+	explicit FMassSnorm8Vector(const FVector& InVector)
+	{
+		Set(InVector);
+	}
+	
+	void Set(const FVector& InVector)
+	{
+		X = (int8)FMath::Clamp(FMath::RoundToInt(InVector.X * Scale), -(int32)MAX_int8, (int32)MAX_int8);
+		Y = (int8)FMath::Clamp(FMath::RoundToInt(InVector.Y * Scale), -(int32)MAX_int8, (int32)MAX_int8);
+		Z = (int8)FMath::Clamp(FMath::RoundToInt(InVector.Z * Scale), -(int32)MAX_int8, (int32)MAX_int8);
+	}
+
+	FVector Get() const
+	{
+		return FVector(X / Scale, Y / Scale, Z / Scale);
+	}
+
+protected:
+	UPROPERTY(Transient)
+	int8 X = 0;
+
+	UPROPERTY(Transient)
+	int8 Y = 0;
+
+	UPROPERTY(Transient)
+	int8 Z = 0;
+};
+
+/** Vector2D which components are in range [-1..1], encoded in signed bytes. */
+USTRUCT()
+struct MASSCOMMON_API FMassSnorm8Vector2D
+{
+	GENERATED_BODY()
+
+	static constexpr float Scale = (float)MAX_int8;
+
+	FMassSnorm8Vector2D() = default;
+
+	explicit FMassSnorm8Vector2D(const FVector2D& InVector)
+	{
+		Set(InVector);
+	}
+	
+	explicit FMassSnorm8Vector2D(const FVector& InVector)
+	{
+		Set(FVector2D(InVector));
+	}
+
+	void Set(const FVector2D& InVector)
+	{
+		X = (int8)FMath::Clamp(FMath::RoundToInt(InVector.X * Scale), -(int32)MAX_int8, (int32)MAX_int8);
+		Y = (int8)FMath::Clamp(FMath::RoundToInt(InVector.Y * Scale), -(int32)MAX_int8, (int32)MAX_int8);
+	}
+
+	FVector2D Get() const
+	{
+		return FVector2D(X / Scale, Y / Scale);
+	}
+
+	FVector GetVector(const float InZ = 0.0f) const
+	{
+		return FVector(X / Scale, Y / Scale, InZ);
+	}
+	
+protected:
+	UPROPERTY(Transient)
+	int8 X = 0;
+
+	UPROPERTY(Transient)
+	int8 Y = 0;
+};
+
+/** Real in range [0..1], encoded in signed bytes. */
+USTRUCT()
+struct MASSCOMMON_API FMassUnorm8Real
+{
+	GENERATED_BODY()
+
+	static constexpr float Scale = MAX_uint8;
+
+	FMassUnorm8Real() = default;
+
+	explicit FMassUnorm8Real(const float InValue)
+	{
+		Set(InValue);
+	}
+
+	void Set(const float InValue)
+	{
+		Value = (int8)FMath::Clamp(FMath::RoundToInt(InValue * Scale), 0, (int32)MAX_uint8);
+	}
+
+	float Get() const
+	{
+		return (Value / Scale);
+	}
+
+protected:
+	UPROPERTY(Transient)
+	uint8 Value = 0;
+};
+
+/** Vector encoded in int16, 1cm accuracy. */
+USTRUCT()
+struct MASSCOMMON_API FMassInt16Vector
+{
+	GENERATED_BODY()
+
+	FMassInt16Vector() = default;
+	FMassInt16Vector(const FVector& InVector)
+	{
+		Set(InVector);
+	}
+	
+	void Set(const FVector& InVector)
+	{
+		X = (int16)FMath::Clamp(FMath::RoundToInt(InVector.X), -(int32)MAX_int16, (int32)MAX_int16);
+		Y = (int16)FMath::Clamp(FMath::RoundToInt(InVector.Y), -(int32)MAX_int16, (int32)MAX_int16);
+		Z = (int16)FMath::Clamp(FMath::RoundToInt(InVector.Z), -(int32)MAX_int16, (int32)MAX_int16);
+	}
+
+	FVector Get() const
+	{
+		return FVector(X, Y, Z);
+	}
+
+protected:
+	UPROPERTY(Transient)
+	int16 X = 0;
+
+	UPROPERTY(Transient)
+	int16 Y = 0;
+
+	UPROPERTY(Transient)
+	int16 Z = 0;
+};
+
+/** Vector2D encoded in int16, 1cm accuracy. */
+USTRUCT()
+struct MASSCOMMON_API FMassInt16Vector2D
+{
+	GENERATED_BODY()
+
+	FMassInt16Vector2D() = default;
+	FMassInt16Vector2D(const FVector2D& InVector)
+	{
+		Set(InVector);
+	}
+	FMassInt16Vector2D(const FVector& InVector)
+	{
+		Set(FVector2D(InVector));
+	}
+	
+	void Set(const FVector2D& InVector)
+	{
+		X = (int16)FMath::Clamp(FMath::RoundToInt(InVector.X), -(int32)MAX_int16, (int32)MAX_int16);
+		Y = (int16)FMath::Clamp(FMath::RoundToInt(InVector.Y), -(int32)MAX_int16, (int32)MAX_int16);
+	}
+
+	FVector2D Get() const
+	{
+		return FVector2D(X, Y);
+	}
+
+	FVector GetVector(const float InZ = 0.0f) const
+	{
+		return FVector(X, Y, InZ);
+	}
+
+protected:
+	UPROPERTY(Transient)
+	int16 X = 0;
+
+	UPROPERTY(Transient)
+	int16 Y = 0;
+};
diff --git a/Engine/Plugins/Runtime/MassGameplay/Source/MassCommon/Public/MassGenericProcessors.h b/Engine/Plugins/Runtime/MassGameplay/Source/MassCommon/Public/MassGenericProcessors.h
new file mode 100644
--- /dev/null
+++ b/Engine/Plugins/Runtime/MassGameplay/Source/MassCommon/Public/MassGenericProcessors.h
@@ -0,0 +1,52 @@
+// Copyright Epic Games, Inc. All Rights Reserved.
+
+#pragma once
+
+#include "MassProcessor.h"
+#include "MassCommonTypes.h"
+#include "MassGenericProcessors.generated.h"
+
+
+/** 
+ *  An experimental processor that can be configured to work on any LWComponent type as long as it's size matches
+ *  FVectorComponent (effectively the size of FVector). Meant mostly for prototyping.
+ *  @todo currently the processor is randomizing it as if it was a fixed-max-length 2d velocity vector. If we keep this
+ *  idea around we can add more properties to the processor that would configure the behavior.
+ */
+UCLASS()
+class MASSCOMMON_API URandomizeVectorProcessor : public UPipeProcessor
+{
+	GENERATED_BODY()
+
+protected:
+	TArrayView<FVector> ComponentList;
+	// note that only types having same size as FVector and extending FLWComponentData will be accepted
+	UPROPERTY(EditAnywhere, Category=Processor, NoClear)
+	UScriptStruct* ComponentType;
+
+public:
+	URandomizeVectorProcessor();
+
+	virtual void ConfigureQueries() override
+	{
+		if (ComponentType)
+		{
+			EntityQuery.AddRequirement(ComponentType, ELWComponentAccess::ReadWrite);
+		}
+	}
+
+	virtual void Execute(UEntitySubsystem& EntitySubsystem, FLWComponentSystemExecutionContext& Context) override;
+
+	UPROPERTY(EditAnywhere, Category = "Mass", meta = (ClampMin = 1, UIMin = 1))
+	float MinMagnitude = 1.f;
+
+	UPROPERTY(EditAnywhere, Category = "Mass", meta = (ClampMin = 1, UIMin = 1))
+	float MaxMagnitude = 1.f;
+
+#if WITH_EDITOR
+	virtual void PostEditChangeProperty(FPropertyChangedEvent& PropertyChangedEvent) override;
+#endif // WITH_EDITOR
+
+protected:
+	FLWComponentQuery EntityQuery;
+};
diff --git a/Engine/Plugins/Runtime/MassGameplay/Source/MassGameplayDebug/MassGameplayDebug.Build.cs b/Engine/Plugins/Runtime/MassGameplay/Source/MassGameplayDebug/MassGameplayDebug.Build.cs
new file mode 100644
--- /dev/null
+++ b/Engine/Plugins/Runtime/MassGameplay/Source/MassGameplayDebug/MassGameplayDebug.Build.cs
@@ -0,0 +1,62 @@
+// Copyright Epic Games, Inc. All Rights Reserved.
+
+namespace UnrealBuildTool.Rules
+{
+	public class MassGameplayDebug : ModuleRules
+	{
+		public MassGameplayDebug(ReadOnlyTargetRules Target) : base(Target)
+		{
+			PCHUsage = ModuleRules.PCHUsageMode.UseExplicitOrSharedPCHs;
+
+			PublicIncludePaths.AddRange(
+				new string[] {
+					"Runtime/AIModule/Public",
+					ModuleDirectory + "/Public",
+				}
+			);
+
+
+			PrivateIncludePaths.AddRange(
+				new string[] {
+				}
+			);
+
+			PublicDependencyModuleNames.AddRange(
+				new string[] {
+					"Core",
+					"CoreUObject",
+					"Engine",
+					"InputCore",
+					"MassEntity",
+					"NavigationSystem",
+					"StateTreeModule",
+					"MassActors",
+					"MassCommon",
+					"MassMovement",
+					"MassSmartObjects",
+					"MassSpawner",
+					"MassSimulation",
+					"MassRepresentation",
+					"MassLOD",
+					"StructUtils",
+				}
+			);
+
+			if (Target.bBuildEditor == true)
+			{
+				PrivateDependencyModuleNames.Add("UnrealEd");
+				PublicDependencyModuleNames.Add("MassEntityEditor");
+			}
+
+			if (Target.bBuildDeveloperTools || (Target.Configuration != UnrealTargetConfiguration.Shipping && Target.Configuration != UnrealTargetConfiguration.Test))
+			{
+				PrivateDependencyModuleNames.Add("GameplayDebugger");
+				PublicDefinitions.Add("WITH_GAMEPLAY_DEBUGGER=1");
+			}
+			else
+			{
+				PublicDefinitions.Add("WITH_GAMEPLAY_DEBUGGER=0");
+			}
+		}
+	}
+}
diff --git a/Engine/Plugins/Runtime/MassGameplay/Source/MassGameplayDebug/Private/AssignDebugVisProcessor.cpp b/Engine/Plugins/Runtime/MassGameplay/Source/MassGameplayDebug/Private/AssignDebugVisProcessor.cpp
new file mode 100644
--- /dev/null
+++ b/Engine/Plugins/Runtime/MassGameplay/Source/MassGameplayDebug/Private/AssignDebugVisProcessor.cpp
@@ -0,0 +1,63 @@
+// Copyright Epic Games, Inc. All Rights Reserved.
+
+#include "AssignDebugVisProcessor.h"
+#include "MassGameplayDebugTypes.h"
+#include "MassDebugVisualizationComponent.h"
+#include "MassDebuggerSubsystem.h"
+#include "Engine/World.h"
+
+//----------------------------------------------------------------------//
+// UAssignDebugVisProcessor
+//----------------------------------------------------------------------//
+UAssignDebugVisProcessor::UAssignDebugVisProcessor()
+{
+	bAutoRegisterWithProcessingPhases = false;
+	FragmentType = FSimDebugVisComponent::StaticStruct();
+}
+
+void UAssignDebugVisProcessor::ConfigureQueries()
+{
+	EntityQuery.AddRequirement<FSimDebugVisComponent>(ELWComponentAccess::ReadWrite);
+}
+
+void UAssignDebugVisProcessor::Execute(UEntitySubsystem& EntitySubsystem, FLWComponentSystemExecutionContext& Context)
+{
+	UMassDebugVisualizationComponent* Visualizer = WeakVisualizer.Get();
+	if (!ensure(Visualizer))
+	{
+		return;
+	}
+
+#if WITH_EDITORONLY_DATA
+	QUICK_SCOPE_CYCLE_COUNTER(AssignDebugVisProcessor_Execute);
+	
+	// note that this function will create the "visual components" only it they're missing or out of sync. 
+	Visualizer->ConditionallyConstructVisualComponent();
+
+	EntityQuery.ForEachEntityChunk(EntitySubsystem, Context, [this, Visualizer](FLWComponentSystemExecutionContext& Context)
+	{
+		const TArrayView<FSimDebugVisComponent> DebugVisList = Context.GetMutableComponentView<FSimDebugVisComponent>();
+		for (FSimDebugVisComponent& VisualComp : DebugVisList)
+		{
+			// VisualComp.VisualType needs to be assigned by now. Should be performed as part of spawning, copied from the AgentTemplate
+			if (ensure(VisualComp.VisualType != INDEX_NONE))
+			{
+				VisualComp.InstanceIndex = Visualizer->AddDebugVisInstance(VisualComp.VisualType);
+			}
+		}
+	});
+	Visualizer->DirtyVisuals();
+#endif // WITH_EDITORONLY_DATA
+}
+
+void UAssignDebugVisProcessor::Initialize(UObject& InOwner)
+{
+	Super::Initialize(InOwner);
+	
+	UMassDebuggerSubsystem* Debugger = UWorld::GetSubsystem<UMassDebuggerSubsystem>(InOwner.GetWorld());
+	if (ensure(Debugger))
+	{
+		WeakVisualizer = Debugger->GetVisualizationComponent();	
+		ensure(WeakVisualizer.Get());
+	}
+}
diff --git a/Engine/Plugins/Runtime/MassGameplay/Source/MassGameplayDebug/Private/DebugEntityTemplateBuilder.cpp b/Engine/Plugins/Runtime/MassGameplay/Source/MassGameplayDebug/Private/DebugEntityTemplateBuilder.cpp
new file mode 100644
--- /dev/null
+++ b/Engine/Plugins/Runtime/MassGameplay/Source/MassGameplayDebug/Private/DebugEntityTemplateBuilder.cpp
@@ -0,0 +1,69 @@
+// Copyright Epic Games, Inc. All Rights Reserved.
+
+#include "DebugEntityTemplateBuilder.h"
+#include "MassDebuggerSubsystem.h"
+#include "MassCommonFragments.h"
+#include "MassDebugVisualizationComponent.h"
+#include "MassEntityTemplateRegistry.h"
+#include "MassMovementFragments.h"
+#include "MassTranslatorRegistry.h"
+#include "Engine/World.h"
+
+namespace UE::Mass::DebugEntityTemplateBuilder
+{
+	void BuildTemplate(const UWorld* World, const FInstancedStruct& InStructInstance, FMassEntityTemplateBuildContext& BuildContext)
+	{
+		const FMassSpawnProps& Data = InStructInstance.Get<FMassSpawnProps>();
+
+		for (const FInstancedStruct& Fragment : Data.AdditionalDataFragments)
+		{
+			BuildContext.AddFragmentWithDefaultInitializer(Fragment);
+		}
+
+		// the following needs to be always there for mesh vis to work. Adding following fragments after already
+		// adding Data.AdditionalDataFragments to let user configure the fragments first. Calling OutFragments.Add()
+		// won't override any fragments that are already there
+		BuildContext.AddFragmentWithDefaultInitializer<FDataFragment_Transform>();
+		BuildContext.AddFragmentWithDefaultInitializer<FDataFragment_NavLocation>();
+		BuildContext.AddFragmentWithDefaultInitializer<FMassVelocityFragment>();
+
+#if WITH_EDITORONLY_DATA
+		const UStaticMesh* const DebugMesh = Data.DebugShape.Mesh;
+#else
+	const UStaticMesh* const DebugMesh = nullptr;
+#endif
+
+		if (DebugMesh)
+		{
+#if WITH_EDITORONLY_DATA
+			FSimDebugVisComponent& DebugVisFragment = BuildContext.AddFragmentWithDefaultInitializer_GetRef<FSimDebugVisComponent>();
+			UMassDebuggerSubsystem* Debugger = UWorld::GetSubsystem<UMassDebuggerSubsystem>(World);
+			if (ensure(Debugger))
+			{
+				UMassDebugVisualizationComponent* DebugVisComponent = Debugger->GetVisualizationComponent();
+				if (ensure(DebugVisComponent))
+				{
+					DebugVisFragment.VisualType = DebugVisComponent->AddDebugVisType(Data.DebugShape);
+				}
+			}
+
+			BuildContext.AddDefaultInitializer<FSimDebugVisComponent>();
+#endif // WITH_EDITORONLY_DATA
+		}
+			// add fragments needed whenever we have debugging capabilities
+			// @todo could be useful to distinguish the manually added fragments from the ones we've auto-added. The template could store this info
+#if !(UE_BUILD_SHIPPING || UE_BUILD_TEST)
+		else // Data.DebugShape.Mesh == null
+		{
+			BuildContext.AddFragmentWithDefaultInitializer<FDataFragment_DebugVis>();
+			BuildContext.AddFragmentWithDefaultInitializer<FDataFragment_AgentRadius>();
+		}
+#endif // if !(UE_BUILD_SHIPPING || UE_BUILD_TEST)
+	}
+
+	void Register()
+	{
+		UMassEntityTemplateRegistry::FindOrAdd(*FMassSpawnProps::StaticStruct()).BindStatic(&BuildTemplate);
+	}
+} // namespace UE::Mass::DebugEntityTemplateBuilder
+
diff --git a/Engine/Plugins/Runtime/MassGameplay/Source/MassGameplayDebug/Private/DebugEntityTemplateBuilder.h b/Engine/Plugins/Runtime/MassGameplay/Source/MassGameplayDebug/Private/DebugEntityTemplateBuilder.h
new file mode 100644
--- /dev/null
+++ b/Engine/Plugins/Runtime/MassGameplay/Source/MassGameplayDebug/Private/DebugEntityTemplateBuilder.h
@@ -0,0 +1,11 @@
+// Copyright Epic Games, Inc. All Rights Reserved.
+
+#pragma once
+
+namespace UE { namespace Mass { namespace DebugEntityTemplateBuilder
+{
+
+void Register();
+
+}}} // namespace UE::Mass::DebugEntityTemplateBuilder
+
diff --git a/Engine/Plugins/Runtime/MassGameplay/Source/MassGameplayDebug/Private/DebugVisLocationProcessor.cpp b/Engine/Plugins/Runtime/MassGameplay/Source/MassGameplayDebug/Private/DebugVisLocationProcessor.cpp
new file mode 100644
--- /dev/null
+++ b/Engine/Plugins/Runtime/MassGameplay/Source/MassGameplayDebug/Private/DebugVisLocationProcessor.cpp
@@ -0,0 +1,118 @@
+// Copyright Epic Games, Inc. All Rights Reserved.
+
+#include "DebugVisLocationProcessor.h"
+#include "MassDebuggerSubsystem.h"
+#include "MassDebugVisualizationComponent.h"
+#include "MassCommonFragments.h"
+#include "Components/HierarchicalInstancedStaticMeshComponent.h"
+#include "EntitySubsystem.h"
+#include "MassMovementFragments.h"
+
+//----------------------------------------------------------------------//
+// UDebugVisLocationProcessor
+//----------------------------------------------------------------------//
+UDebugVisLocationProcessor::UDebugVisLocationProcessor()
+{
+	ExecutionOrder.ExecuteAfter.Add(UE::Mass::ProcessorGroupNames::SyncWorldToMass);
+}
+
+void UDebugVisLocationProcessor::ConfigureQueries()
+{
+	EntityQuery.AddRequirement<FDataFragment_Transform>(ELWComponentAccess::ReadOnly);
+	EntityQuery.AddRequirement<FSimDebugVisComponent>(ELWComponentAccess::ReadOnly);
+	EntityQuery.AddRequirement<FMassVelocityFragment>(ELWComponentAccess::ReadOnly);
+}
+
+void UDebugVisLocationProcessor::Execute(UEntitySubsystem& EntitySubsystem, FLWComponentSystemExecutionContext& Context)
+{
+#if WITH_EDITORONLY_DATA
+	UMassDebugVisualizationComponent* Visualizer = WeakVisualizer.Get();
+	if (!ensure(Visualizer))
+	{
+		return;
+	}
+
+	QUICK_SCOPE_CYCLE_COUNTER(DebugVisLocationProcessor_Run);
+
+	TArrayView<UHierarchicalInstancedStaticMeshComponent*> VisualDataISMCs = Visualizer->GetVisualDataISMCs();
+	if (VisualDataISMCs.Num() > 0)
+	{
+		EntityQuery.ForEachEntityChunk(EntitySubsystem, Context, [this, &VisualDataISMCs](const FLWComponentSystemExecutionContext& Context)
+		{
+			const int32 NumEntities = Context.GetEntitiesNum();
+			const TConstArrayView<FDataFragment_Transform> LocationList = Context.GetComponentView<FDataFragment_Transform>();
+			const TConstArrayView<FSimDebugVisComponent> DebugVisList = Context.GetComponentView<FSimDebugVisComponent>();
+
+			for (int32 i = 0; i < NumEntities; ++i)
+			{
+				const FSimDebugVisComponent& VisualComp = DebugVisList[i];
+
+				// @todo: remove this code once the asset is exported with correct alignment SM_Mannequin.uasset
+				FTransform SMTransform = LocationList[i].GetTransform();
+				FQuat FromEngineToSM(FVector::UpVector, -HALF_PI);
+				SMTransform.SetRotation(FromEngineToSM * SMTransform.GetRotation());
+
+				VisualDataISMCs[VisualComp.VisualType]->UpdateInstanceTransform(VisualComp.InstanceIndex, SMTransform, true);
+			}
+		});
+
+		Visualizer->DirtyVisuals();
+	}
+	else 
+	{
+		UE_LOG(LogMassDebug, Log, TEXT("UDebugVisLocationProcessor: Trying to update InstanceStaticMeshes while none created. Check your debug visualization setup"));
+	}
+#endif // WITH_EDITORONLY_DATA
+}
+
+void UDebugVisLocationProcessor::Initialize(UObject& InOwner)
+{
+	Super::Initialize(InOwner);
+	UMassDebuggerSubsystem* Debugger = UWorld::GetSubsystem<UMassDebuggerSubsystem>(InOwner.GetWorld());
+	if (ensure(Debugger))
+	{
+		WeakVisualizer = Debugger->GetVisualizationComponent();
+		ensure(WeakVisualizer.Get());
+	}
+}
+
+//----------------------------------------------------------------------//
+//  UMassProcessor_UpdateDebugVis
+//----------------------------------------------------------------------//
+UMassProcessor_UpdateDebugVis::UMassProcessor_UpdateDebugVis()
+{
+	ExecutionOrder.ExecuteAfter.Add(UE::Mass::ProcessorGroupNames::UpdateWorldFromMass);
+}
+
+void UMassProcessor_UpdateDebugVis::ConfigureQueries() 
+{
+	// @todo only FDataFragment_DebugVis should be mandatory, rest optional 
+	EntityQuery.AddRequirement<FDataFragment_Transform>(ELWComponentAccess::ReadOnly);
+	EntityQuery.AddRequirement<FDataFragment_DebugVis>(ELWComponentAccess::ReadWrite);
+	EntityQuery.AddRequirement<FDataFragment_AgentRadius>(ELWComponentAccess::ReadWrite);
+}
+
+void UMassProcessor_UpdateDebugVis::Execute(UEntitySubsystem& EntitySubsystem, FLWComponentSystemExecutionContext& Context)
+{
+	UMassDebuggerSubsystem* Debugger = UWorld::GetSubsystem<UMassDebuggerSubsystem>(GetWorld());
+	if (Debugger == nullptr)
+	{
+		return;
+	}
+
+	QUICK_SCOPE_CYCLE_COUNTER(UMassProcessor_UpdateDebugVis_Run);
+
+	EntityQuery.ForEachEntityChunk(EntitySubsystem, Context, [this, Debugger](FLWComponentSystemExecutionContext& Context)
+		{
+			const int32 NumEntities = Context.GetEntitiesNum();
+			const TConstArrayView<FDataFragment_Transform> LocationList = Context.GetComponentView<FDataFragment_Transform>();
+			const TArrayView<FDataFragment_DebugVis> DebugVisList = Context.GetMutableComponentView<FDataFragment_DebugVis>();
+			const TArrayView<FDataFragment_AgentRadius> RadiiList = Context.GetMutableComponentView<FDataFragment_AgentRadius>();
+
+			for (int32 i = 0; i < NumEntities; ++i)
+			{
+				Debugger->AddShape(DebugVisList[i].Shape, LocationList[i].GetTransform().GetLocation(), RadiiList[i].Radius);
+				Debugger->AddEntityLocation(Context.GetEntity(i), LocationList[i].GetTransform().GetLocation());
+			}
+		});
+}
diff --git a/Engine/Plugins/Runtime/MassGameplay/Source/MassGameplayDebug/Private/MassDebugVisualizationComponent.cpp b/Engine/Plugins/Runtime/MassGameplay/Source/MassGameplayDebug/Private/MassDebugVisualizationComponent.cpp
new file mode 100644
--- /dev/null
+++ b/Engine/Plugins/Runtime/MassGameplay/Source/MassGameplayDebug/Private/MassDebugVisualizationComponent.cpp
@@ -0,0 +1,95 @@
+// Copyright Epic Games, Inc. All Rights Reserved.
+
+#include "MassDebugVisualizationComponent.h"
+#include "Components/HierarchicalInstancedStaticMeshComponent.h"
+
+
+#if WITH_EDITORONLY_DATA
+#include "MassDebugVisualizer.h"
+
+void UMassDebugVisualizationComponent::PostInitProperties()
+{
+	Super::PostInitProperties();
+	if (HasAnyFlags(RF_ClassDefaultObject) == false && GetOuter())
+	{
+		ensureMsgf(GetOuter()->GetClass()->IsChildOf(AMassDebugVisualizer::StaticClass()), TEXT("UMassDebugVisualizationComponent should only be added to AMassDebugVisualizer-like instances"));
+	}
+}
+
+void UMassDebugVisualizationComponent::DirtyVisuals()
+{
+	for (UHierarchicalInstancedStaticMeshComponent* ISM : VisualDataISMCs)
+	{
+		if (ensure(ISM))
+		{
+			ISM->MarkRenderStateDirty();
+		}
+	}
+}
+
+int32 UMassDebugVisualizationComponent::AddDebugVisInstance(const uint16 VisualType)
+{
+	return VisualDataISMCs[VisualType]->AddInstance(FTransform::Identity);
+}
+
+void UMassDebugVisualizationComponent::ConditionallyConstructVisualComponent()
+{
+	if (VisualDataISMCs.Num() == 0 || VisualDataISMCs.Num() != VisualDataTable.Num())
+	{
+		ConstructVisualComponent();
+	}
+}
+
+void UMassDebugVisualizationComponent::ConstructVisualComponent()
+{
+	for (UHierarchicalInstancedStaticMeshComponent* ISM : VisualDataISMCs)
+	{
+		if (ensure(ISM))
+		{
+			ISM->ClearInstances();
+			ISM->UnregisterComponent();
+		}
+	}
+	VisualDataISMCs.Reset();
+
+	AActor* ActorOwner = GetOwner();
+	check(ActorOwner);
+	
+	for (const FAgentDebugVisualization& VisualData : VisualDataTable)
+	{
+		UHierarchicalInstancedStaticMeshComponent* HISMC = NewObject<UHierarchicalInstancedStaticMeshComponent>(ActorOwner);
+		HISMC->SetStaticMesh(VisualData.Mesh);
+		if (VisualData.MaterialOverride != nullptr)
+		{
+			HISMC->SetMaterial(0, VisualData.MaterialOverride);
+		}
+
+		HISMC->SetCullDistances(VisualData.VisualNearCullDistance, VisualData.VisualFarCullDistance);
+		HISMC->SetupAttachment(ActorOwner->GetRootComponent());
+		HISMC->SetCanEverAffectNavigation(false);
+		HISMC->bDisableCollision = true;
+		HISMC->SetCastShadow(false);
+		HISMC->RegisterComponent();
+
+		VisualDataISMCs.Add(HISMC);
+	}
+}
+
+int32 UMassDebugVisualizationComponent::AddDebugVisType(const FAgentDebugVisualization& Data)
+{
+	return VisualDataTable.Add(Data);
+}
+
+void UMassDebugVisualizationComponent::Clear()
+{
+	for (UHierarchicalInstancedStaticMeshComponent* ISM : VisualDataISMCs)
+	{
+		if (ensure(ISM))
+		{
+			ISM->ClearInstances();
+			ISM->UnregisterComponent();
+		}
+	}
+	VisualDataISMCs.Reset();
+}
+#endif // WITH_EDITORONLY_DATA
diff --git a/Engine/Plugins/Runtime/MassGameplay/Source/MassGameplayDebug/Private/MassDebugVisualizationTrait.cpp b/Engine/Plugins/Runtime/MassGameplay/Source/MassGameplayDebug/Private/MassDebugVisualizationTrait.cpp
new file mode 100644
--- /dev/null
+++ b/Engine/Plugins/Runtime/MassGameplay/Source/MassGameplayDebug/Private/MassDebugVisualizationTrait.cpp
@@ -0,0 +1,50 @@
+// Copyright Epic Games, Inc. All Rights Reserved.
+
+#include "MassDebugVisualizationTrait.h"
+#include "MassDebuggerSubsystem.h"
+#include "MassDebugVisualizationComponent.h"
+#include "MassEntityTemplateRegistry.h"
+#include "MassCommonFragments.h"
+#include "Engine/World.h"
+
+void UMassDebugVisualizationTrait::BuildTemplate(FMassEntityTemplateBuildContext& BuildContext, UWorld& World) const
+{
+#if WITH_EDITORONLY_DATA
+	const UStaticMesh* const DebugMesh = DebugShape.Mesh;
+#else
+	const UStaticMesh* const DebugMesh = nullptr;
+#endif
+
+	if (DebugMesh)
+	{
+#if WITH_EDITORONLY_DATA
+		FSimDebugVisComponent& DebugVisFragment = BuildContext.AddFragmentWithDefaultInitializer_GetRef<FSimDebugVisComponent>();
+		UMassDebuggerSubsystem* Debugger = World.GetSubsystem<UMassDebuggerSubsystem>();
+		if (ensure(Debugger))
+		{
+			UMassDebugVisualizationComponent* DebugVisComponent = Debugger->GetVisualizationComponent();
+			if (ensure(DebugVisComponent))
+			{
+				DebugVisFragment.VisualType = DebugVisComponent->AddDebugVisType(DebugShape);
+			}
+			// @todo this path requires a fragment destructor that will remove the mesh from the debugger.
+		}
+
+		BuildContext.AddDefaultInitializer<FSimDebugVisComponent>();
+#endif // WITH_EDITORONLY_DATA
+	}
+	// add fragments needed whenever we have debugging capabilities
+#if !(UE_BUILD_SHIPPING || UE_BUILD_TEST)
+	else // Data.DebugShape.Mesh == null
+	{
+#if WITH_EDITORONLY_DATA
+		BuildContext.AddFragmentWithDefaultInitializer_GetRef<FDataFragment_DebugVis>().Shape = DebugShape.WireShape;
+#else
+		// DebugShape unavailable, will used default instead
+		BuildContext.AddFragmentWithDefaultInitializer<FDataFragment_DebugVis>();
+#endif // WITH_EDITORONLY_DATA
+		BuildContext.AddFragmentWithDefaultInitializer<FDataFragment_AgentRadius>();
+	}
+#endif // if !(UE_BUILD_SHIPPING || UE_BUILD_TEST)
+
+}
\ No newline at end of file
diff --git a/Engine/Plugins/Runtime/MassGameplay/Source/MassGameplayDebug/Private/MassDebuggerSubsystem.cpp b/Engine/Plugins/Runtime/MassGameplay/Source/MassGameplayDebug/Private/MassDebuggerSubsystem.cpp
new file mode 100644
--- /dev/null
+++ b/Engine/Plugins/Runtime/MassGameplay/Source/MassGameplayDebug/Private/MassDebuggerSubsystem.cpp
@@ -0,0 +1,105 @@
+// Copyright Epic Games, Inc. All Rights Reserved.
+
+#include "MassDebuggerSubsystem.h"
+#include "EngineUtils.h"
+#include "MassCommonTypes.h"
+#include "MassSimulationSubsystem.h"
+#include "MassDebugVisualizationComponent.h"
+#include "MassEntitySystem.h"
+#include "MassDebugVisualizer.h"
+
+
+void UMassDebuggerSubsystem::Initialize(FSubsystemCollectionBase& Collection)
+{
+	Collection.InitializeDependency(UMassSimulationSubsystem::StaticClass());
+	Super::Initialize(Collection);
+
+	UMassSimulationSubsystem* SimSystem = UWorld::GetSubsystem<UMassSimulationSubsystem>(GetWorld());
+	check(SimSystem);
+	SimSystem->GetOnProcessingPhaseStarted(EPipeProcessingPhase::PrePhysics).AddUObject(this, &UMassDebuggerSubsystem::OnProcessingPhaseStarted);
+}
+
+void UMassDebuggerSubsystem::PreTickProcessors()
+{
+	// get ready to receive new debug info
+	for (TArray<FShapeDesc>& Array : Shapes)
+	{
+		Array.Reset();
+	}
+	
+	Entities.Reset();
+	Locations.Reset();
+	SelectedEntityDetails.Empty();
+}
+
+void UMassDebuggerSubsystem::AddEntityLocation(FLWEntity Entity, const FVector& Location)
+{
+	Entities.Add(Entity);
+	Locations.Add(Location);
+}
+
+void UMassDebuggerSubsystem::OnProcessingPhaseStarted(const float DeltaSeconds)
+{
+	PreTickProcessors();
+}
+
+void UMassDebuggerSubsystem::SetSelectedEntity(const FLWEntity& InSelectedEntity)
+{
+	SelectedEntity = InSelectedEntity;
+	SelectedEntityDetails.Empty();
+}
+
+void UMassDebuggerSubsystem::AppendSelectedEntityInfo(const FString& Info)
+{
+	SelectedEntityDetails += Info;
+}
+
+UMassDebugVisualizationComponent* UMassDebuggerSubsystem::GetVisualizationComponent()
+{ 
+#if WITH_EDITORONLY_DATA
+	if (VisualizationComponent == nullptr)
+	{
+		if (!ensureMsgf(DebugVisualizer == nullptr, TEXT("If we do have a DebugVisualizer but don't have VisualizationComponent then somethin's wrong")))
+		{
+			VisualizationComponent = &DebugVisualizer->GetDebugVisComponent();
+		}
+		else
+		{
+			if (UWorld* World = GetWorld())
+			{
+				const AMassDebugVisualizer& VisualizerActor = GetOrSpawnDebugVisualizer(*World);
+				VisualizationComponent = &VisualizerActor.GetDebugVisComponent();
+			}
+		}
+	}
+	ensureMsgf(VisualizationComponent,  TEXT("In editor builds we always expect to have a visualizer component available"));
+#endif // WITH_EDITORONLY_DATA
+	return VisualizationComponent;
+}
+
+#if WITH_EDITORONLY_DATA
+AMassDebugVisualizer& UMassDebuggerSubsystem::GetOrSpawnDebugVisualizer(UWorld& InWorld)
+{
+	if (DebugVisualizer)
+	{
+		return *DebugVisualizer;
+	}
+
+	// see if there is one already and we've missed it somehow
+	for (const TActorIterator<AMassDebugVisualizer> It(&InWorld); It;)
+	{
+		DebugVisualizer = *It;
+		return *DebugVisualizer;
+	}
+
+	FActorSpawnParameters SpawnInfo;
+	SpawnInfo.SpawnCollisionHandlingOverride = ESpawnActorCollisionHandlingMethod::AlwaysSpawn;
+	// The helper actor is created on demand and only once per world so we can allow it to spawn during construction script.
+	SpawnInfo.bAllowDuringConstructionScript = true;
+	DebugVisualizer = InWorld.SpawnActor<AMassDebugVisualizer>(SpawnInfo);
+	check(DebugVisualizer);
+	VisualizationComponent = &DebugVisualizer->GetDebugVisComponent();
+
+	return *DebugVisualizer;
+}
+#endif
\ No newline at end of file
diff --git a/Engine/Plugins/Runtime/MassGameplay/Source/MassGameplayDebug/Public/AssignDebugVisProcessor.h b/Engine/Plugins/Runtime/MassGameplay/Source/MassGameplayDebug/Public/AssignDebugVisProcessor.h
new file mode 100644
--- /dev/null
+++ b/Engine/Plugins/Runtime/MassGameplay/Source/MassGameplayDebug/Public/AssignDebugVisProcessor.h
@@ -0,0 +1,26 @@
+// Copyright Epic Games, Inc. All Rights Reserved.
+
+#pragma once
+
+#include "MassTranslator.h"
+#include "AssignDebugVisProcessor.generated.h"
+
+
+class UMassDebugVisualizationComponent;
+struct FSimDebugVisComponent;
+
+UCLASS()
+class MASSGAMEPLAYDEBUG_API UAssignDebugVisProcessor : public UMassFragmentInitializer
+{
+	GENERATED_BODY()
+public:
+	UAssignDebugVisProcessor();
+	virtual void ConfigureQueries() override;
+	virtual void Execute(UEntitySubsystem& EntitySubsystem, FLWComponentSystemExecutionContext& Context) override;
+	virtual void Initialize(UObject& InOwner) override;
+
+protected:
+	TWeakObjectPtr<UMassDebugVisualizationComponent> WeakVisualizer;
+
+	FLWComponentQuery EntityQuery;
+};
diff --git a/Engine/Plugins/Runtime/MassGameplay/Source/MassGameplayDebug/Public/DebugVisLocationProcessor.h b/Engine/Plugins/Runtime/MassGameplay/Source/MassGameplayDebug/Public/DebugVisLocationProcessor.h
new file mode 100644
--- /dev/null
+++ b/Engine/Plugins/Runtime/MassGameplay/Source/MassGameplayDebug/Public/DebugVisLocationProcessor.h
@@ -0,0 +1,45 @@
+// Copyright Epic Games, Inc. All Rights Reserved.
+
+#pragma once
+
+#include "MassProcessor.h"
+#include "DebugVisLocationProcessor.generated.h"
+
+class UMassDebugVisualizationComponent;
+struct FSimDebugVisComponent;
+
+UCLASS()
+class MASSGAMEPLAYDEBUG_API UDebugVisLocationProcessor : public UPipeProcessor
+{
+	GENERATED_BODY()
+
+public:
+	UDebugVisLocationProcessor();
+	
+protected:
+	virtual void ConfigureQueries() override;
+	virtual void Execute(UEntitySubsystem& EntitySubsystem, FLWComponentSystemExecutionContext& Context) override;
+	virtual void Initialize(UObject& InOwner) override;
+
+	TWeakObjectPtr<UMassDebugVisualizationComponent> WeakVisualizer;
+	FLWComponentQuery EntityQuery;
+};
+
+//----------------------------------------------------------------------//
+// new one 
+//----------------------------------------------------------------------//
+//class UMassDebugger;
+
+UCLASS()
+class MASSGAMEPLAYDEBUG_API UMassProcessor_UpdateDebugVis : public UPipeProcessor
+{
+	GENERATED_BODY()
+public:
+	UMassProcessor_UpdateDebugVis();
+	virtual void ConfigureQueries() override;
+	virtual void Execute(UEntitySubsystem& EntitySubsystem, FLWComponentSystemExecutionContext& Context) override;
+
+protected:
+	FLWComponentQuery EntityQuery;
+};
+
diff --git a/Engine/Plugins/Runtime/MassGameplay/Source/MassGameplayDebug/Public/MassDebugVisualizationComponent.h b/Engine/Plugins/Runtime/MassGameplay/Source/MassGameplayDebug/Public/MassDebugVisualizationComponent.h
new file mode 100644
--- /dev/null
+++ b/Engine/Plugins/Runtime/MassGameplay/Source/MassGameplayDebug/Public/MassDebugVisualizationComponent.h
@@ -0,0 +1,40 @@
+// Copyright Epic Games, Inc. All Rights Reserved.
+
+#pragma once
+
+#include "MassGameplayDebugTypes.h"
+#include "Components/ActorComponent.h"
+#include "MassDebugVisualizationComponent.generated.h"
+
+
+class UHierarchicalInstancedStaticMeshComponent;
+
+/** meant to be created procedurally and owned by an AMassDebugVisualizer instance. Will ensure if placed on a different type of actor */
+UCLASS()
+class MASSGAMEPLAYDEBUG_API UMassDebugVisualizationComponent : public UActorComponent
+{
+	GENERATED_BODY()
+public:
+#if WITH_EDITORONLY_DATA
+	/**  will create Owner's "visual components" only it they're missing or out of sync with VisualDataTable */
+	void ConditionallyConstructVisualComponent();
+	void DirtyVisuals();
+	int32 AddDebugVisInstance(const uint16 VisualType);
+	/** returns index to the newly created VisualDataTable entry */
+	int32 AddDebugVisType(const FAgentDebugVisualization& Data);
+	TArrayView<UHierarchicalInstancedStaticMeshComponent*> GetVisualDataISMCs() { return MakeArrayView(VisualDataISMCs); }
+
+	void Clear();
+protected:
+	virtual void PostInitProperties() override;
+	void ConstructVisualComponent();
+
+protected:
+
+	UPROPERTY(Transient)
+	TArray<FAgentDebugVisualization> VisualDataTable;
+
+	UPROPERTY(Transient)
+	TArray<UHierarchicalInstancedStaticMeshComponent*> VisualDataISMCs;
+#endif // WITH_EDITORONLY_DATA
+};
diff --git a/Engine/Plugins/Runtime/MassGameplay/Source/MassGameplayDebug/Public/MassDebugVisualizationTrait.h b/Engine/Plugins/Runtime/MassGameplay/Source/MassGameplayDebug/Public/MassDebugVisualizationTrait.h
new file mode 100644
--- /dev/null
+++ b/Engine/Plugins/Runtime/MassGameplay/Source/MassGameplayDebug/Public/MassDebugVisualizationTrait.h
@@ -0,0 +1,23 @@
+// Copyright Epic Games, Inc. All Rights Reserved.
+
+#pragma once 
+
+#include "MassEntityTraitBase.h"
+#include "MassGameplayDebugTypes.h"
+#include "MassDebugVisualizationTrait.generated.h"
+
+
+UCLASS(BlueprintType, EditInlineNew, CollapseCategories, meta = (DisplayName = "Debug Visualization"))
+class MASSGAMEPLAYDEBUG_API UMassDebugVisualizationTrait : public UMassEntityTraitBase
+{
+	GENERATED_BODY()
+
+protected:
+
+	virtual void BuildTemplate(FMassEntityTemplateBuildContext& BuildContext, UWorld& World) const override;
+
+#if WITH_EDITORONLY_DATA
+	UPROPERTY(EditAnywhere, Category = "Mass|Spawn")
+	FAgentDebugVisualization DebugShape;
+#endif // WITH_EDITORONLY_DATA
+};
diff --git a/Engine/Plugins/Runtime/MassGameplay/Source/MassGameplayDebug/Public/MassDebugVisualizer.h b/Engine/Plugins/Runtime/MassGameplay/Source/MassGameplayDebug/Public/MassDebugVisualizer.h
new file mode 100644
--- /dev/null
+++ b/Engine/Plugins/Runtime/MassGameplay/Source/MassGameplayDebug/Public/MassDebugVisualizer.h
@@ -0,0 +1,24 @@
+// Copyright Epic Games, Inc. All Rights Reserved.
+
+#pragma once
+
+#include "GameFramework/Actor.h"
+#include "MassDebugVisualizer.generated.h"
+
+
+UCLASS(NotPlaceable, Transient)
+class MASSGAMEPLAYDEBUG_API AMassDebugVisualizer : public AActor
+{
+	GENERATED_BODY()
+public:
+	AMassDebugVisualizer();
+
+#if WITH_EDITORONLY_DATA
+	/** If this function is callable we guarantee the debug vis component to exist*/
+	class UMassDebugVisualizationComponent& GetDebugVisComponent() const { return *DebugVisComponent; }
+
+protected:
+	UPROPERTY()
+	class UMassDebugVisualizationComponent* DebugVisComponent;
+#endif // WITH_EDITORONLY_DATA
+};
diff --git a/Engine/Plugins/Runtime/MassGameplay/Source/MassGameplayDebug/Public/MassDebuggerSubsystem.h b/Engine/Plugins/Runtime/MassGameplay/Source/MassGameplayDebug/Public/MassDebuggerSubsystem.h
new file mode 100644
--- /dev/null
+++ b/Engine/Plugins/Runtime/MassGameplay/Source/MassGameplayDebug/Public/MassDebuggerSubsystem.h
@@ -0,0 +1,73 @@
+// Copyright Epic Games, Inc. All Rights Reserved.
+
+#pragma once
+
+#include "MassGameplayDebugTypes.h"
+#include "Subsystems/WorldSubsystem.h"
+#include "MassDebuggerSubsystem.generated.h"
+
+
+class UMassDebugVisualizationComponent;
+class AMassDebugVisualizer;
+
+UCLASS()
+class MASSGAMEPLAYDEBUG_API UMassDebuggerSubsystem : public UWorldSubsystem
+{
+	GENERATED_BODY()
+public:
+	struct FShapeDesc
+	{
+		FVector Location = {}; // no init on purpose, value will come from constructor
+		float Size = {};
+		FShapeDesc(const FVector InLocation, const float InSize) : Location(InLocation), Size(InSize) {}
+	};
+
+	// Methods to optimize the collection of data to only when category is enabled
+	bool IsCollectingData() const { return bCollectingData; }
+	void SetCollectingData() { bCollectingData = true; }
+	void DataCollected() { bCollectingData = false; }
+
+	void AddShape(EMassEntityDebugShape Shape, FVector Location, float Size) { Shapes[uint8(Shape)].Add(FShapeDesc(Location, Size)); }
+	const TArray<FShapeDesc>* GetShapes() const { return Shapes; }
+
+	void AddEntityLocation(FLWEntity Entity, const FVector& Location);
+	TConstArrayView<FLWEntity> GetEntities() const { return Entities; }
+	TConstArrayView<FVector> GetLocations() const { return Locations; }
+
+	FLWEntity GetSelectedEntity() const { return SelectedEntity; }
+	void SetSelectedEntity(const FLWEntity& InSelectedEntity);
+
+	void AppendSelectedEntityInfo(const FString& Info);
+	const FString& GetSelectedEntityInfo() const { return SelectedEntityDetails; }
+	
+	/** Fetches the UMassDebugVisualizationComponent owned by lazily created DebugVisualizer */
+	UMassDebugVisualizationComponent* GetVisualizationComponent();
+
+#if WITH_EDITORONLY_DATA
+	AMassDebugVisualizer& GetOrSpawnDebugVisualizer(UWorld& InWorld);
+#endif
+
+protected:
+	// USubsystem BEGIN
+	virtual void Initialize(FSubsystemCollectionBase& Collection) override;
+	// USubsystem END
+	
+	void OnProcessingPhaseStarted(const float DeltaSeconds);
+	void PreTickProcessors();
+
+protected:
+	bool bCollectingData = false;
+
+	TArray<FShapeDesc> Shapes[uint8(EMassEntityDebugShape::MAX)];
+	TArray<FLWEntity> Entities;
+	TArray<FVector> Locations;
+	FLWEntity SelectedEntity;
+	FString SelectedEntityDetails;
+
+	UPROPERTY(Transient)
+	UMassDebugVisualizationComponent* VisualizationComponent;
+
+	UPROPERTY(Transient)
+	AMassDebugVisualizer* DebugVisualizer;
+};
+
diff --git a/Engine/Plugins/Runtime/MassGameplay/Source/MassGameplayDebug/Public/MassGameplayDebugTypes.h b/Engine/Plugins/Runtime/MassGameplay/Source/MassGameplayDebug/Public/MassGameplayDebugTypes.h
new file mode 100644
--- /dev/null
+++ b/Engine/Plugins/Runtime/MassGameplay/Source/MassGameplayDebug/Public/MassGameplayDebugTypes.h
@@ -0,0 +1,110 @@
+// Copyright Epic Games, Inc. All Rights Reserved.
+
+#pragma once
+
+#include "CoreMinimal.h"
+#include "Stats/Stats.h"
+#include "UObject/ObjectMacros.h"
+#include "Logging/LogMacros.h"
+#include "Templates/SubclassOf.h"
+#include "Engine/DataTable.h"
+#include "LWComponentTypes.h"
+#include "MassSpawnerTypes.h"
+#include "MassGameplayDebugTypes.generated.h"
+
+
+DECLARE_LOG_CATEGORY_EXTERN(LogMassDebug, Warning, All);
+
+#if WITH_EDITORONLY_DATA
+class UBillboardComponent;
+#endif // WITH_EDITORONLY_DATA
+class UStaticMesh;
+class UMaterialInterface;
+
+
+USTRUCT()
+struct FSimDebugDataRow : public FTableRowBase
+{
+	GENERATED_BODY()
+
+	UPROPERTY(EditAnywhere, Category = Debug)
+	UStaticMesh* Mesh = nullptr;
+
+	UPROPERTY(EditAnywhere, Category = Debug)
+	UMaterialInterface* MaterialOverride = nullptr;
+
+	UPROPERTY(EditAnywhere, Category = Debug)
+	float Scale = 1.f;
+};
+
+USTRUCT()
+struct FSimDebugVisComponent : public FLWComponentData
+{
+	GENERATED_BODY()
+	int32 InstanceIndex = INDEX_NONE;
+	int16 VisualType = INDEX_NONE;
+};
+
+UENUM()
+enum class EMassEntityDebugShape : uint8
+{
+	Box,
+	Cone,
+	Cylinder,
+	Capsule,
+	MAX
+};
+
+USTRUCT()
+struct FDataFragment_DebugVis : public FLWComponentData
+{
+	GENERATED_BODY()
+	UPROPERTY(EditAnywhere, Category = Debug)
+	EMassEntityDebugShape Shape = EMassEntityDebugShape::Box;
+};
+
+USTRUCT()
+struct FAgentDebugVisualization : public FTableRowBase
+{
+	GENERATED_BODY()
+
+	UPROPERTY(EditAnywhere, Category = "Mass|Debug")
+	UStaticMesh* Mesh = nullptr;
+
+	UPROPERTY(EditAnywhere, Category = "Mass|Debug")
+	UMaterialInterface* MaterialOverride = nullptr;
+
+	/** Near cull distance to override default value for that agent type */
+	UPROPERTY(EditAnywhere, Category = "Mass|Debug")
+	float VisualNearCullDistance = 5000.f;
+
+	/** Far cull distance to override default value for that agent type */
+	UPROPERTY(EditAnywhere, Category = "Mass|Debug")
+	float VisualFarCullDistance = 7500.f;
+
+	/** If Mesh is not set this WireShape will be used for debug drawing via GameplayDebugger */
+	UPROPERTY(EditAnywhere, Category = "Mass|Debug")
+	EMassEntityDebugShape WireShape = EMassEntityDebugShape::Box;
+};
+
+/** @todo comment functionality */
+USTRUCT(DisplayName="MassSpawnProps_Debug")
+struct FMassSpawnProps : public FMassSpawnConfigBase
+{
+	GENERATED_BODY()
+
+#if WITH_EDITORONLY_DATA
+	UPROPERTY(EditAnywhere, Category = "Mass|Spawn")
+	FAgentDebugVisualization DebugShape;
+#endif // WITH_EDITORONLY_DATA
+
+	friend uint32 GetTypeHash(const FMassSpawnProps& Instance)
+	{
+#if WITH_EDITORONLY_DATA
+		return HashCombine(GetTypeHash((const FMassSpawnConfigBase&)Instance)
+				, HashCombine(GetTypeHash(Instance.DebugShape.Mesh), GetTypeHash(Instance.DebugShape.MaterialOverride)));
+#else
+		return GetTypeHash((const FMassSpawnConfigBase&)Instance);
+#endif // WITH_EDITORONLY_DATA
+	}
+};
\ No newline at end of file
diff --git a/Engine/Plugins/Runtime/MassGameplay/Source/MassLOD/Private/MassLODCollectorProcessor.cpp b/Engine/Plugins/Runtime/MassGameplay/Source/MassLOD/Private/MassLODCollectorProcessor.cpp
new file mode 100644
--- /dev/null
+++ b/Engine/Plugins/Runtime/MassGameplay/Source/MassLOD/Private/MassLODCollectorProcessor.cpp
@@ -0,0 +1,129 @@
+// Copyright Epic Games, Inc. All Rights Reserved.
+
+#include "MassLODCollectorProcessor.h"
+#include "MassLODUtils.h"
+#include "MassCommonFragments.h"
+#include "Engine/World.h"
+#include "MassSimulationLOD.h"
+
+UMassLODCollectorProcessor::UMassLODCollectorProcessor()
+{
+	ExecutionFlags = (int32)EProcessorExecutionFlags::All;
+
+	ExecutionOrder.ExecuteInGroup = UE::Mass::ProcessorGroupNames::LODCollector;
+	ExecutionOrder.ExecuteAfter.Add(UE::Mass::ProcessorGroupNames::SyncWorldToMass);
+}
+
+void UMassLODCollectorProcessor::Initialize(UObject& InOwner)
+{
+	Super::Initialize(InOwner);
+
+	for (FMassCollectorLODConfig& LODConfig : LODConfigs)
+	{
+		LODConfig.RepresentationLODCollector.Initialize(LODConfig.FOVAnglesToDriveVisibility, LODConfig.BufferHysteresisOnFOVPercentage / 100.f);
+		LODConfig.SimulationLODCollector.Initialize();
+		LODConfig.CombinedLODCollector.Initialize(LODConfig.FOVAnglesToDriveVisibility, LODConfig.BufferHysteresisOnFOVPercentage / 100.f);
+	}
+}
+
+void UMassLODCollectorProcessor::ConfigureQueries()
+{
+	for (FMassCollectorLODConfig& LODConfig : LODConfigs)
+	{
+		FLWComponentQuery BaseQuery;
+		if (LODConfig.TagFilter.GetScriptStruct())
+		{
+			BaseQuery.AddTagRequirement(*LODConfig.TagFilter.GetScriptStruct(), ELWComponentPresence::All);
+		}
+		BaseQuery.AddRequirement<FDataFragment_Transform>(ELWComponentAccess::ReadOnly);
+		BaseQuery.AddRequirement<FMassLODInfoFragment>(ELWComponentAccess::ReadWrite);
+
+		LODConfig.OnLODEntityQuery = BaseQuery;
+		LODConfig.OnLODEntityQuery.AddTagRequirement<FMassOffLODTag>(ELWComponentPresence::None);
+		LODConfig.OffLODEntityQuery_Conditional = BaseQuery;
+		LODConfig.OffLODEntityQuery_Conditional.AddTagRequirement<FMassOffLODTag>(ELWComponentPresence::All);
+		LODConfig.OffLODEntityQuery_Conditional.AddChunkRequirement<FMassSimulationVariableTickChunkFragment>(ELWComponentAccess::ReadOnly);
+		LODConfig.OffLODEntityQuery_Conditional.SetChunkFilter(&FMassSimulationVariableTickChunkFragment::ShouldTickChunkThisFrame);
+
+		LODConfig.CloseEntityQuery = BaseQuery;
+		LODConfig.CloseEntityQuery.AddTagRequirement<FMassVisibilityCulledByDistanceTag>(ELWComponentPresence::None);
+		LODConfig.FarEntityQuery_Conditional = BaseQuery;
+		LODConfig.FarEntityQuery_Conditional.AddTagRequirement<FMassVisibilityCulledByDistanceTag>(ELWComponentPresence::All);
+		LODConfig.FarEntityQuery_Conditional.AddChunkRequirement<FMassVisualizationChunkFragment>(ELWComponentAccess::ReadOnly);
+		LODConfig.FarEntityQuery_Conditional.SetChunkFilter(&FMassVisualizationChunkFragment::ShouldUpdateVisualizationForChunk);
+
+		// @todo: Cleanup once we can do more complex logic with entity queries. There is no other way to do "Notall", we need to create by hand 3 queries to achieve it.
+		LODConfig.CloseAndOnLODEntityQuery = BaseQuery;
+		LODConfig.CloseAndOnLODEntityQuery.AddTagRequirement<FMassVisibilityCulledByDistanceTag>(ELWComponentPresence::None);
+		LODConfig.CloseAndOnLODEntityQuery.AddTagRequirement<FMassOffLODTag>(ELWComponentPresence::None);
+		LODConfig.CloseAndOffLODEntityQuery = BaseQuery;
+		LODConfig.CloseAndOffLODEntityQuery.AddTagRequirement<FMassVisibilityCulledByDistanceTag>(ELWComponentPresence::None);
+		LODConfig.CloseAndOffLODEntityQuery.AddTagRequirement<FMassOffLODTag>(ELWComponentPresence::All);
+		LODConfig.FarAndOnLODEntityQuery = BaseQuery;
+		LODConfig.FarAndOnLODEntityQuery.AddTagRequirement<FMassVisibilityCulledByDistanceTag>(ELWComponentPresence::All);
+		LODConfig.FarAndOnLODEntityQuery.AddTagRequirement<FMassOffLODTag>(ELWComponentPresence::None);
+		LODConfig.FarAndOffLODEntityQuery_Conditional = BaseQuery;
+		LODConfig.FarAndOffLODEntityQuery_Conditional.AddTagRequirement<FMassVisibilityCulledByDistanceTag>(ELWComponentPresence::All);
+		LODConfig.FarAndOffLODEntityQuery_Conditional.AddTagRequirement<FMassOffLODTag>(ELWComponentPresence::All);
+		LODConfig.FarAndOffLODEntityQuery_Conditional.AddChunkRequirement<FMassSimulationVariableTickChunkFragment>(ELWComponentAccess::ReadOnly);
+		LODConfig.FarAndOffLODEntityQuery_Conditional.AddChunkRequirement<FMassVisualizationChunkFragment>(ELWComponentAccess::ReadOnly);
+		LODConfig.FarAndOffLODEntityQuery_Conditional.SetChunkFilter([](const FLWComponentSystemExecutionContext& Context)
+			{
+				return Context.GetChunkComponent<FMassVisualizationChunkFragment>().ShouldUpdateVisualization()
+					|| Context.GetChunkComponent<FMassSimulationVariableTickChunkFragment>().ShouldTickThisFrame();
+			});
+	}
+}
+
+template <typename TCollector>
+void CollectLOD(UEntitySubsystem& EntitySubsystem, FLWComponentSystemExecutionContext& Context, const TArray<FViewerInfo>& Viewers, TCollector& Collector, TArrayView<FLWComponentQuery> HighFrequencyTickingEntityQueries, FLWComponentQuery& LowFrequencyTickingEntityQuery)
+{
+	Collector.PrepareExecution(Viewers);
+
+	auto CollectLODInfo = [&Collector](FLWComponentSystemExecutionContext& Context)
+	{
+		TConstArrayView<FDataFragment_Transform> LocationList = Context.GetComponentView<FDataFragment_Transform>();
+		TArrayView<FMassLODInfoFragment> ViewerInfoList = Context.GetMutableComponentView<FMassLODInfoFragment>();
+
+		Collector.CollectLODInfo(Context, LocationList, ViewerInfoList);
+	};
+
+	{
+		TRACE_CPUPROFILER_EVENT_SCOPE(TEXT("HighFrequency"))
+		for(FLWComponentQuery& HighFrequencyTickingEntityQuery : HighFrequencyTickingEntityQueries)
+		{
+			HighFrequencyTickingEntityQuery.ForEachEntityChunk(EntitySubsystem, Context, CollectLODInfo);
+		}
+	}
+	{
+		TRACE_CPUPROFILER_EVENT_SCOPE(TEXT("LowFrequency"))
+		LowFrequencyTickingEntityQuery.ForEachEntityChunk(EntitySubsystem, Context, CollectLODInfo);
+	}
+}
+
+void UMassLODCollectorProcessor::Execute(UEntitySubsystem& EntitySubsystem, FLWComponentSystemExecutionContext& Context)
+{
+	check(LODManager);
+	const TArray<FViewerInfo>& Viewers = LODManager->GetViewers();
+
+	check(World);
+
+	for (FMassCollectorLODConfig& LODConfig : LODConfigs)
+	{
+		if (World->IsNetMode(NM_Client))
+		{
+			TRACE_CPUPROFILER_EVENT_SCOPE(TEXT("RepresentationLODCollector"))
+			CollectLOD(EntitySubsystem, Context, Viewers, LODConfig.RepresentationLODCollector, MakeArrayView(&LODConfig.CloseEntityQuery,1), LODConfig.FarEntityQuery_Conditional);
+		}
+		else if (World->IsNetMode(NM_DedicatedServer))
+		{
+			TRACE_CPUPROFILER_EVENT_SCOPE(TEXT("SimulationLODCollector"))
+			CollectLOD(EntitySubsystem, Context, Viewers, LODConfig.SimulationLODCollector, MakeArrayView(&LODConfig.OnLODEntityQuery,1), LODConfig.OffLODEntityQuery_Conditional);
+		}
+		else
+		{
+			TRACE_CPUPROFILER_EVENT_SCOPE(TEXT("CombinedLODCollector"))
+			CollectLOD(EntitySubsystem, Context, Viewers, LODConfig.CombinedLODCollector, MakeArrayView(&LODConfig.CloseAndOnLODEntityQuery,3), LODConfig.FarAndOffLODEntityQuery_Conditional);
+		}
+	}
+}
diff --git a/Engine/Plugins/Runtime/MassGameplay/Source/MassLOD/Private/MassLODLogic.cpp b/Engine/Plugins/Runtime/MassGameplay/Source/MassLOD/Private/MassLODLogic.cpp
new file mode 100644
--- /dev/null
+++ b/Engine/Plugins/Runtime/MassGameplay/Source/MassLOD/Private/MassLODLogic.cpp
@@ -0,0 +1,29 @@
+// Copyright Epic Games, Inc. All Rights Reserved.
+
+#include "MassLODLogic.h"
+#include "Engine/World.h"
+#include "GameFramework/PlayerController.h"
+
+void FMassLODBaseLogic::CacheViewerInformation(TConstArrayView<FViewerInfo> Viewers, const bool bLocalViewersOnly)
+{
+	NumOfViewers = Viewers.Num();
+	check(NumOfViewers <= UE::MassLOD::MaxNumOfViewers);
+
+	// Cache viewer info
+	for (int ViewerIdx = 0; ViewerIdx < NumOfViewers; ++ViewerIdx)
+	{
+		const FViewerInfo& Viewer = Viewers[ViewerIdx];
+		const APlayerController* PlayerController = Viewer.PlayerController;
+		const FMassViewerHandle ViewerHandle = Viewer.bEnabled && (!bLocalViewersOnly || Viewer.IsLocal()) ? Viewer.Handle : FMassViewerHandle();
+
+		// Check if it is the same client as before
+		bClearViewerData[ViewerIdx] = ViewersHandles[ViewerIdx] != ViewerHandle;
+		ViewersHandles[ViewerIdx] = ViewerHandle;
+
+		if (ViewerHandle.IsValid())
+		{
+			ViewersLocation[ViewerIdx] = Viewer.Location;
+			ViewersDirection[ViewerIdx] = Viewer.Direction;
+		}
+	}
+}
\ No newline at end of file
diff --git a/Engine/Plugins/Runtime/MassGameplay/Source/MassLOD/Private/MassLODManager.cpp b/Engine/Plugins/Runtime/MassGameplay/Source/MassLOD/Private/MassLODManager.cpp
new file mode 100644
--- /dev/null
+++ b/Engine/Plugins/Runtime/MassGameplay/Source/MassLOD/Private/MassLODManager.cpp
@@ -0,0 +1,337 @@
+// Copyright Epic Games, Inc. All Rights Reserved.
+
+#include "MassLODManager.h"
+#include "GameFramework/PlayerController.h"
+#include "WorldPartition/WorldPartition.h"
+#include "WorldPartition/WorldPartitionStreamingSource.h"
+#include "Engine/World.h"
+
+namespace UE
+{
+namespace MassLOD
+{
+	FColor LODColors[] =
+	{
+		FColor::Red,
+		FColor::Yellow,
+		FColor::Emerald,
+		FColor::White,
+	};
+}
+}
+
+void UMassProcessor_LODBase::Initialize(UObject& Owner)
+{
+	Super::Initialize(Owner);
+
+	World = Owner.GetWorld();
+	LODManager = UWorld::GetSubsystem<UMassLODManager>(World);
+}
+
+void UMassLODManager::Initialize(FSubsystemCollectionBase& Collection)
+{
+	Super::Initialize(Collection);
+	SynchronizeViewers();
+}
+
+void UMassLODManager::Tick(float DeltaTime)
+{
+	SynchronizeViewers();
+}
+
+TStatId UMassLODManager::GetStatId() const
+{
+	RETURN_QUICK_DECLARE_CYCLE_STAT(UMassLODManager, STATGROUP_Tickables);
+}
+
+void UMassLODManager::Deinitialize()
+{
+	// Remove all viewers from the viewer info list
+	for (FViewerInfo& ViewerInfo : Viewers)
+	{
+		if (ViewerInfo.Handle.IsValid())
+		{
+			// Safe to remove while iterating as it is a sparse array with a free list
+			RemoveViewer(ViewerInfo.Handle);
+		}
+	}
+}
+
+int32 UMassLODManager::GetValidViewerIdx(const FMassViewerHandle& ViewerHandle) const
+{
+	// Check against invalid handle
+	if (!ViewerHandle.IsValid())
+	{
+		return INDEX_NONE;
+	}
+
+	// Check against removed viewers
+	const int32 ViewerIdx = ViewerHandle.GetIndex();
+	if (ViewerIdx >= Viewers.Num() || ViewerHandle != Viewers[ViewerIdx].Handle)
+	{
+		return INDEX_NONE;
+	}
+
+	return ViewerIdx;
+}
+
+const TArray<FViewerInfo>& UMassLODManager::GetSynchronizedViewers()
+{
+	SynchronizeViewers();
+
+	return Viewers;
+}
+
+FMassViewerHandle UMassLODManager::GetViewerHandleFromPlayerController(const APlayerController* PlayerController) const
+{
+	const FMassViewerHandle* Handle = ViewerMap.Find(GetTypeHash(PlayerController->GetFName()));
+	return Handle ? *Handle : FMassViewerHandle();
+}
+
+
+FMassViewerHandle UMassLODManager::GetViewerHandleFromStreamingSource(const FName StreamingSourceName) const
+{
+	const FMassViewerHandle* Handle = ViewerMap.Find(GetTypeHash(StreamingSourceName));
+	return Handle ? *Handle : FMassViewerHandle();
+}
+
+APlayerController* UMassLODManager::GetPlayerControllerFromViewerHandle(const FMassViewerHandle& ViewerHandle) const
+{
+	const int32 ViewerIdx = GetValidViewerIdx(ViewerHandle);
+	return ViewerIdx != INDEX_NONE ? Viewers[ViewerIdx].PlayerController : nullptr;
+}
+
+void UMassLODManager::SynchronizeViewers()
+{
+	if (LastSynchronizedFrame == GFrameCounter)
+	{
+		return;
+	}
+	LastSynchronizedFrame = GFrameCounter;
+
+	bool bNeedShrinking = false;
+
+	const UWorld* World = GetWorld();
+	UWorldPartition* WorldPartition = World ? World->GetWorldPartition() : nullptr;
+	const TArray<FWorldPartitionStreamingSource>& StreamingSources = WorldPartition ? WorldPartition->GetStreamingSources() : TArray<FWorldPartitionStreamingSource>();
+
+	// Go through the list and check validity and store the valid one into a map
+	TMap<uint32, FMassViewerHandle> LocalViewerMap;
+	for (int32 ViewerIdx = 0; ViewerIdx < Viewers.Num(); ++ViewerIdx)
+	{
+		FViewerInfo& ViewerInfo = Viewers[ViewerIdx];
+		if (!ViewerInfo.Handle.IsValid())
+		{
+			check(ViewerFreeIndices.Find(ViewerIdx) != INDEX_NONE);
+			continue;
+		}
+
+		if (ViewerInfo.PlayerController != nullptr)
+		{
+			LocalViewerMap.Add(GetTypeHash(ViewerInfo.PlayerController->GetFName()), ViewerInfo.Handle);
+		}
+		else if (!ViewerInfo.StreamingSourceName.IsNone() && StreamingSources.FindByPredicate([&ViewerInfo](const FWorldPartitionStreamingSource& Source){ return Source.Name == ViewerInfo.StreamingSourceName; }) != nullptr)
+		{
+			LocalViewerMap.Add(GetTypeHash(ViewerInfo.StreamingSourceName), ViewerInfo.Handle);
+		}
+		else
+		{
+			// Safe to remove while iterating as it is a sparse array with a free list
+			RemoveViewer(ViewerInfo.Handle);
+			bNeedShrinking |= ViewerIdx == Viewers.Num() - 1;
+		}
+	}
+
+	if (World)
+	{
+		// Now go through all current player controllers and add if they do not exist
+		for (FConstPlayerControllerIterator PlayerIterator = World->GetPlayerControllerIterator(); PlayerIterator; ++PlayerIterator)
+		{
+			APlayerController* PlayerController = (*PlayerIterator).Get();
+			check(PlayerController);
+
+			// Check if the controller already exists by trying to remove it from the map which was filled up with controllers we were tracking
+			if (LocalViewerMap.Remove(GetTypeHash(PlayerController->GetFName())) == 0)
+			{
+				// If not add it to the list
+				AddViewer(PlayerController);
+			}
+		}
+
+		// Now go through all current streaming source and add if they do not exist
+		for (const FWorldPartitionStreamingSource& StreamingSource : StreamingSources)
+		{
+			if (LocalViewerMap.Remove(GetTypeHash(StreamingSource.Name)) == 0)
+			{
+				AddViewer(nullptr, StreamingSource.Name);
+			}
+		}
+	}
+
+	// Anything left in the map need to be removed from the list
+	for (TMap<uint32, FMassViewerHandle>::TIterator Itr = LocalViewerMap.CreateIterator(); Itr; ++Itr)
+	{
+		const int32 ViewerIdx = Itr->Value.GetIndex();
+		RemoveViewer(Viewers[ViewerIdx].Handle);
+		bNeedShrinking |= ViewerIdx == Viewers.Num() - 1;
+	}
+
+	if (bNeedShrinking)
+	{
+		// Check to shrink the array of viewers if possible
+		while (Viewers.Num() > 0 && ViewerFreeIndices.Num() > 0 && ViewerFreeIndices.Remove(Viewers.Num() - 1))
+		{
+			Viewers.Pop(false);
+		}
+	}
+
+	// Update location and direction for every type of viewer
+	for (FViewerInfo& ViewerInfo : Viewers)
+	{
+		if (!ViewerInfo.Handle.IsValid())
+		{
+			continue;
+		}
+
+		if (ViewerInfo.PlayerController)
+		{
+			ViewerInfo.bEnabled = !WorldPartition || ViewerInfo.PlayerController->bEnableStreamingSource;
+
+			FVector PlayerCameraLocation(ForceInitToZero);
+			FRotator PlayerCameraRotation(FRotator::ZeroRotator);
+			ViewerInfo.PlayerController->GetPlayerViewPoint(PlayerCameraLocation, PlayerCameraRotation);
+			ViewerInfo.Location = PlayerCameraLocation;
+			ViewerInfo.Direction = PlayerCameraRotation.Vector();
+		}
+		else
+		{
+			checkf(!ViewerInfo.StreamingSourceName.IsNone(), TEXT("Expecting to have a streamingsourcename if the playercontroller is null"));
+			const FWorldPartitionStreamingSource* StreamingSource = StreamingSources.FindByPredicate([&ViewerInfo](const FWorldPartitionStreamingSource& Source) { return Source.Name == ViewerInfo.StreamingSourceName; });
+			checkf(StreamingSource, TEXT("Expecting to be pointing to a valid streaming source"));
+			ViewerInfo.bEnabled = StreamingSource != nullptr;
+			if (StreamingSource)
+			{
+				ViewerInfo.Location = StreamingSource->Location;
+				ViewerInfo.Direction = StreamingSource->Rotation.Vector();
+			}
+		}
+	}
+}
+
+void UMassLODManager::AddViewer(APlayerController* PlayerController, FName StreamingSourceName /* = NAME_None*/)
+{
+	const int32 HashValue = PlayerController ? GetTypeHash(PlayerController->GetFName()) : GetTypeHash(StreamingSourceName);
+
+	FMassViewerHandle& ViewerHandle = ViewerMap.FindOrAdd(HashValue, FMassViewerHandle());
+	if (ViewerHandle.IsValid())
+	{
+		// We are only interested to set the player controller if it was not already set.
+		if (PlayerController)
+		{
+			const int32 ViewerHandleIdx = GetValidViewerIdx(ViewerHandle);
+			check(ViewerHandleIdx != INDEX_NONE);
+
+			FViewerInfo& ViewerInfo = Viewers[ViewerHandleIdx];
+			check(ViewerInfo.PlayerController == nullptr);
+			ViewerInfo.PlayerController = PlayerController;
+			PlayerController->OnEndPlay.AddUniqueDynamic(this, &UMassLODManager::OnPlayerControllerEndPlay);
+		}
+
+		return;
+	}
+
+	// Add new viewer
+#if UE_ALLOW_DEBUG_REPLICATION_DUPLICATE_VIEWERS_PER_CONTROLLER
+	//for debugging / profiling purposes create DebugNumberViwersPerController
+	//in this case ViewerMap will only contain a hash to the most recent viewer handle created.
+	for (int Idx = 0; Idx < UE::MassLOD::DebugNumberViewersPerController; ++Idx)
+#endif //UE_ALLOW_DEBUG_REPLICATION_DUPLICATE_VIEWERS_PER_CONTROLLER
+	{
+		const bool bAddNew = ViewerFreeIndices.Num() == 0;
+		const int NewIdx = bAddNew ? Viewers.Num() : ViewerFreeIndices.Pop();
+		FViewerInfo& NewViewerInfo = bAddNew ? Viewers.AddDefaulted_GetRef() : Viewers[NewIdx];
+		NewViewerInfo.PlayerController = PlayerController;
+		NewViewerInfo.StreamingSourceName = StreamingSourceName;
+		NewViewerInfo.Handle.Index = NewIdx;
+		NewViewerInfo.Handle.SerialNumber = GetNextViewerSerialNumber();
+		NewViewerInfo.HashValue = HashValue;
+
+		ViewerHandle = NewViewerInfo.Handle;
+
+		OnViewerAddedDelegate.Broadcast(NewViewerInfo.Handle, NewViewerInfo.PlayerController, NewViewerInfo.StreamingSourceName);
+	}
+
+	if(PlayerController)
+	{
+		PlayerController->OnEndPlay.AddUniqueDynamic(this, &UMassLODManager::OnPlayerControllerEndPlay);
+	}
+}
+
+void UMassLODManager::RemoveViewer(const FMassViewerHandle& ViewerHandle)
+{
+#if UE_ALLOW_DEBUG_REPLICATION_DUPLICATE_VIEWERS_PER_CONTROLLER
+
+	const int32 ViewerHandleIdx = GetValidViewerIdx(ViewerHandle);
+	check(ViewerHandleIdx != INDEX_NONE);
+
+	//find all the viewer handles the slow way and remove them
+	for (int32 ViewerIdx = 0; ViewerIdx < Viewers.Num(); ++ViewerIdx)
+	{
+		FViewerInfo& ViewerInfo = Viewers[ViewerIdx];
+
+		if (ViewerInfo.PlayerController == Viewers[ViewerHandleIdx].PlayerController)
+		{
+			RemoveViewerInternal(ViewerInfo.Handle);
+		}
+	}
+
+#else
+
+	RemoveViewerInternal(ViewerHandle);
+
+#endif //UE_ALLOW_DEBUG_REPLICATION_DUPLICATE_VIEWERS_PER_CONTROLLER
+}
+
+void UMassLODManager::RemoveViewerInternal(const FMassViewerHandle& ViewerHandle)
+{
+	const int32 ViewerIdx = GetValidViewerIdx(ViewerHandle);
+	check(ViewerIdx != INDEX_NONE);
+	FViewerInfo& ViewerInfo = Viewers[ViewerIdx];
+
+	OnViewerRemovedDelegate.Broadcast(ViewerInfo.Handle, ViewerInfo.PlayerController, ViewerInfo.StreamingSourceName);
+
+	if (ViewerInfo.PlayerController)
+	{
+		ViewerInfo.PlayerController->OnEndPlay.RemoveDynamic(this, &UMassLODManager::OnPlayerControllerEndPlay);
+	}
+
+	ViewerMap.Remove(ViewerInfo.HashValue);
+
+	ViewerInfo.Reset();
+	ViewerFreeIndices.Push(ViewerIdx);
+}
+
+void UMassLODManager::OnPlayerControllerEndPlay(AActor* Actor, EEndPlayReason::Type EndPlayReason)
+{
+	APlayerController* PlayerController = Cast<APlayerController>(Actor);
+	if (ensure(PlayerController))
+	{
+		const FMassViewerHandle ViewerHandle = GetViewerHandleFromPlayerController(PlayerController);
+		if (ensure(ViewerHandle.IsValid()))
+		{
+			RemoveViewer(ViewerHandle);
+		}
+	}
+}
+
+void FViewerInfo::Reset()
+{
+	Handle.Invalidate();
+	PlayerController = nullptr;
+	HashValue = 0;
+}
+
+bool FViewerInfo::IsLocal() const
+{
+	return (PlayerController && PlayerController->IsLocalController()) || !StreamingSourceName.IsNone();
+}
diff --git a/Engine/Plugins/Runtime/MassGameplay/Source/MassLOD/Private/MassLODModule.cpp b/Engine/Plugins/Runtime/MassGameplay/Source/MassLOD/Private/MassLODModule.cpp
new file mode 100644
--- /dev/null
+++ b/Engine/Plugins/Runtime/MassGameplay/Source/MassLOD/Private/MassLODModule.cpp
@@ -0,0 +1,46 @@
+// Copyright Epic Games, Inc. All Rights Reserved.
+
+#include "CoreMinimal.h"
+#include "IMassLODModule.h"
+#include "UObject/CoreRedirects.h"
+
+
+class FMassLODModule : public IMassLODModule
+{
+	/** IModuleInterface implementation */
+	virtual void StartupModule() override;
+	virtual void ShutdownModule() override;
+};
+
+IMPLEMENT_MODULE(FMassLODModule, MassLOD)
+
+
+
+void FMassLODModule::StartupModule()
+{
+	// This code will execute after your module is loaded into memory (but after global variables are initialized, of course.)
+
+	TArray<FCoreRedirect> Redirects;
+	Redirects.Emplace(ECoreRedirectFlags::Type_Struct, TEXT("/Script/MassCommon.DataFragment_MassSimulationLODInfo"), TEXT("/Script/MassLOD.DataFragment_MassSimulationLODInfo"));
+	Redirects.Emplace(ECoreRedirectFlags::Type_Struct, TEXT("/Script/MassCommon.DataFragment_MassSimulationLOD"), TEXT("/Script/MassLOD.DataFragment_MassSimulationLOD"));
+	Redirects.Emplace(ECoreRedirectFlags::Type_Struct, TEXT("/Script/MassCommon.DataFragment_MassViewerLOD"), TEXT("/Script/MassLOD.DataFragment_MassViewerLOD"));
+
+	Redirects.Emplace(ECoreRedirectFlags::Type_Class, TEXT("/Script/MassCommon.MassProcessor_LODBase"), TEXT("/Script/MassLOD.MassProcessor_LODBase"));
+
+	Redirects.Emplace(ECoreRedirectFlags::Type_Class, TEXT("/Script/MassCommon.MassProcessor_MassSimulationLODViewersInfo"), TEXT("/Script/MassLOD.MassProcessor_MassSimulationLODViewersInfo"));
+	Redirects.Emplace(ECoreRedirectFlags::Type_Class, TEXT("/Script/MassCommon.MassProcessor_MassSimulationLOD"), TEXT("/Script/MassLOD.MassProcessor_MassSimulationLOD"));
+
+	Redirects.Emplace(ECoreRedirectFlags::Type_Class, TEXT("/Script/MassCommon.MassProcessor_MassViewerLOD"), TEXT("/Script/MassLOD.MassProcessor_MassViewerLOD"));
+	
+	FCoreRedirects::AddRedirectList(Redirects, TEXT("MassMovement"));
+}
+
+
+void FMassLODModule::ShutdownModule()
+{
+	// This function may be called during shutdown to clean up your module.  For modules that support dynamic reloading,
+	// we call this function before unloading the module.
+}
+
+
+
