AggregateTicking and Pipe merged into Engine/Plugins/Runtime/MassEntity

diff --git a/Engine/Plugins/Runtime/MassGameplay/Source/MassLOD/Private/MassRepresentationLODCollectorProcessor.cpp b/Engine/Plugins/Runtime/MassGameplay/Source/MassLOD/Private/MassRepresentationLODCollectorProcessor.cpp
new file mode 100644
--- /dev/null
+++ b/Engine/Plugins/Runtime/MassGameplay/Source/MassLOD/Private/MassRepresentationLODCollectorProcessor.cpp
@@ -0,0 +1,79 @@
+// Copyright Epic Games, Inc. All Rights Reserved.
+
+#include "MassRepresentationLODCollectorProcessor.h"
+#include "Engine/World.h"
+#include "MassCommonFragments.h"
+#include "GameFramework/PlayerController.h"
+
+UMassRepresentationLODCollectorProcessor::UMassRepresentationLODCollectorProcessor()
+{
+	ExecutionFlags = (int32)EProcessorExecutionFlags::All;
+}
+
+void UMassRepresentationLODCollectorProcessor::Initialize(UObject& InOwner)
+{
+	Super::Initialize(InOwner);
+
+	for (FMassRepresentationLODCollectorConfig& LODConfig : LODConfigs)
+	{
+		LODConfig.RepresentationLODCollector.Initialize(LODConfig.FOVAnglesToDriveVisibility, LODConfig.BufferHysteresisOnFOVPercentage / 100.f);
+	}
+}
+
+void UMassRepresentationLODCollectorProcessor::ConfigureQueries()
+{
+	for (FMassRepresentationLODCollectorConfig& LODConfig : LODConfigs)
+	{
+		FLWComponentQuery BaseQuery;
+		if (LODConfig.TagFilter.GetScriptStruct())
+		{
+			BaseQuery.AddTagRequirement(*LODConfig.TagFilter.GetScriptStruct(), ELWComponentPresence::All);
+		}
+		BaseQuery.AddRequirement<FDataFragment_Transform>(ELWComponentAccess::ReadOnly);
+		BaseQuery.AddRequirement<FMassLODInfoFragment>(ELWComponentAccess::ReadWrite);
+		BaseQuery.AddChunkRequirement<FMassVisualizationChunkFragment>(ELWComponentAccess::ReadOnly);
+
+		LODConfig.CloseEntityQuery = BaseQuery;
+		LODConfig.CloseEntityQuery.AddTagRequirement<FMassVisibilityCulledByDistanceTag>(ELWComponentPresence::None);
+		LODConfig.FarEntityQuery_Conditional = BaseQuery;
+		LODConfig.FarEntityQuery_Conditional.AddTagRequirement<FMassVisibilityCulledByDistanceTag>(ELWComponentPresence::All);
+		LODConfig.FarEntityQuery_Conditional.SetChunkFilter(&FMassVisualizationChunkFragment::ShouldUpdateVisualizationForChunk);
+	}
+}
+
+template <typename TCollector>
+void CollectLODInfo(UEntitySubsystem& EntitySubsystem, FLWComponentSystemExecutionContext& Context, const TArray<FViewerInfo>& Viewers, TCollector& Collector, FLWComponentQuery& HighFrequencyTickingEntityQuery, FLWComponentQuery& LowFrequencyTickingEntityQuery)
+{
+	Collector.PrepareExecution(Viewers);
+
+	auto InternalCollectLODInfo = [&Collector](FLWComponentSystemExecutionContext& Context)
+	{
+		TConstArrayView<FDataFragment_Transform> LocationList = Context.GetComponentView<FDataFragment_Transform>();
+		TArrayView<FMassLODInfoFragment> ViewerInfoList = Context.GetMutableComponentView<FMassLODInfoFragment>();
+
+		Collector.CollectLODInfo(Context, LocationList, ViewerInfoList);
+	};
+
+	{
+		TRACE_CPUPROFILER_EVENT_SCOPE(TEXT("HighFrequency"))
+		HighFrequencyTickingEntityQuery.ForEachEntityChunk(EntitySubsystem, Context, InternalCollectLODInfo);
+	}
+	{
+		TRACE_CPUPROFILER_EVENT_SCOPE(TEXT("LowFrequency"))
+		LowFrequencyTickingEntityQuery.ForEachEntityChunk(EntitySubsystem, Context, InternalCollectLODInfo);
+	}
+}
+
+void UMassRepresentationLODCollectorProcessor::Execute(UEntitySubsystem& EntitySubsystem, FLWComponentSystemExecutionContext& Context)
+{
+	check(LODManager);
+	const TArray<FViewerInfo>& Viewers = LODManager->GetViewers();
+
+	check(World);
+
+	for (FMassRepresentationLODCollectorConfig& LODConfig : LODConfigs)
+	{
+		TRACE_CPUPROFILER_EVENT_SCOPE(TEXT("ViewerLODCollector"))
+		CollectLODInfo(EntitySubsystem, Context, Viewers, LODConfig.RepresentationLODCollector, LODConfig.CloseEntityQuery, LODConfig.FarEntityQuery_Conditional);
+	}
+}
diff --git a/Engine/Plugins/Runtime/MassGameplay/Source/MassLOD/Private/MassSimulationLOD.cpp b/Engine/Plugins/Runtime/MassGameplay/Source/MassLOD/Private/MassSimulationLOD.cpp
new file mode 100644
--- /dev/null
+++ b/Engine/Plugins/Runtime/MassGameplay/Source/MassLOD/Private/MassSimulationLOD.cpp
@@ -0,0 +1,183 @@
+// Copyright Epic Games, Inc. All Rights Reserved.
+
+#include "MassSimulationLOD.h"
+
+#include "MassCommonFragments.h"
+#include "GameFramework/PlayerController.h"
+#include "Engine/World.h"
+
+//-----------------------------------------------------------------------------
+// UMassProcessor_MassSimulationLODViewersInfo
+//-----------------------------------------------------------------------------
+
+UMassProcessor_MassSimulationLODViewersInfo::UMassProcessor_MassSimulationLODViewersInfo()
+{
+	bAutoRegisterWithProcessingPhases = false;
+}
+
+void UMassProcessor_MassSimulationLODViewersInfo::Initialize(UObject& Owner)
+{
+	Super::Initialize(Owner);
+
+	LODCollector.Initialize();
+}
+
+void UMassProcessor_MassSimulationLODViewersInfo::ConfigureQueries()
+{
+	EntityQuery.AddRequirement<FDataFragment_Transform>(ELWComponentAccess::ReadOnly);
+	EntityQuery.AddRequirement<FDataFragment_MassSimulationLODInfo>(ELWComponentAccess::ReadWrite);
+}
+
+void UMassProcessor_MassSimulationLODViewersInfo::Execute(UEntitySubsystem& EntitySubsystem, FLWComponentSystemExecutionContext& Context)
+{
+	check(LODManager);
+	const TArray<FViewerInfo>& Viewers = LODManager->GetViewers();
+	LODCollector.PrepareExecution(Viewers);
+
+	EntityQuery.ForEachEntityChunk(EntitySubsystem, Context, [this](FLWComponentSystemExecutionContext& Context)
+	{
+		const TConstArrayView<FDataFragment_Transform> LocationList = Context.GetComponentView<FDataFragment_Transform>();
+		const TArrayView<FDataFragment_MassSimulationLODInfo> ViewersInfoList = Context.GetMutableComponentView<FDataFragment_MassSimulationLODInfo>();
+		LODCollector.CollectLODInfo(Context, LocationList, ViewersInfoList);
+	});
+}
+
+//-----------------------------------------------------------------------------
+// FMassSimulationLODConfig
+//-----------------------------------------------------------------------------
+FMassSimulationLODConfig::FMassSimulationLODConfig()
+{
+	LODDistance[EMassLOD::High] = 0.0f;
+	LODDistance[EMassLOD::Medium] = 5000.0f;
+	LODDistance[EMassLOD::Low] = 10000.0f;
+	LODDistance[EMassLOD::Off] = 30000.0f;
+
+	LODMaxCount[EMassLOD::High] = 100;
+	LODMaxCount[EMassLOD::Medium] = 200;
+	LODMaxCount[EMassLOD::Low] = 300;
+	LODMaxCount[EMassLOD::Off] = INT_MAX;
+
+	TickRates[EMassLOD::High] = 0.0f;
+	TickRates[EMassLOD::Medium] = 0.5f;
+	TickRates[EMassLOD::Low] = 1.0f;
+	TickRates[EMassLOD::Off] = 1.5f;
+}
+
+//-----------------------------------------------------------------------------
+// UMassSimulationLODProcessor
+//-----------------------------------------------------------------------------
+
+namespace UE::MassLOD
+{
+	int32 bDebugSimulationLOD = 0;
+	FAutoConsoleVariableRef CVarDebugSimulationLODTest(TEXT("ai.debug.SimulationLOD"), bDebugSimulationLOD, TEXT("Debug Simulation LOD"), ECVF_Cheat);
+} // UE::MassLOD
+
+UMassSimulationLODProcessor::UMassSimulationLODProcessor()
+{
+	ExecutionFlags = (int32)EProcessorExecutionFlags::All;
+	ExecutionOrder.ExecuteInGroup = UE::Mass::ProcessorGroupNames::LOD;
+	ExecutionOrder.ExecuteAfter.Add(UE::Mass::ProcessorGroupNames::LODCollector);
+}
+
+void UMassSimulationLODProcessor::ConfigureQueries()
+{
+	for (FMassSimulationLODConfig& LODConfig : LODConfigs)
+	{
+		if(LODConfig.TagFilter.GetScriptStruct())
+		{
+			LODConfig.EntityQuery.AddTagRequirement(*LODConfig.TagFilter.GetScriptStruct(), ELWComponentPresence::All);
+		}
+		LODConfig.EntityQuery.AddRequirement<FDataFragment_Transform>(ELWComponentAccess::ReadOnly);
+		LODConfig.EntityQuery.AddRequirement<FMassLODInfoFragment>(ELWComponentAccess::ReadOnly);
+		LODConfig.EntityQuery.AddRequirement<FMassSimulationLODFragment>(ELWComponentAccess::ReadWrite);
+		LODConfig.EntityQuery.AddChunkRequirement<FMassSimulationVariableTickChunkFragment>(ELWComponentAccess::ReadOnly);
+	}
+}
+
+void UMassSimulationLODProcessor::Initialize(UObject& InOwner)
+{
+	for (FMassSimulationLODConfig& LODConfig : LODConfigs)
+	{
+		LODConfig.LODCalculator.Initialize(LODConfig.LODDistance, LODConfig.BufferHysteresisOnDistancePercentage / 100.0f, LODConfig.LODMaxCount);
+		LODConfig.LODTickRateController.Initialize(LODConfig.TickRates);
+	}
+
+	Super::Initialize(InOwner);
+}
+
+void UMassSimulationLODProcessor::Execute(UEntitySubsystem& EntitySubsystem, FLWComponentSystemExecutionContext& Context)
+{
+	TRACE_CPUPROFILER_EVENT_SCOPE(TEXT("SimulationLOD"))
+
+	for (FMassSimulationLODConfig& LODConfig : LODConfigs)
+	{
+		CalculateLODForConfig(EntitySubsystem, Context, LODConfig);
+	}
+}
+
+void UMassSimulationLODProcessor::CalculateLODForConfig(UEntitySubsystem& EntitySubsystem, FLWComponentSystemExecutionContext& Context, FMassSimulationLODConfig& LODConfig)
+{
+	check(LODManager);
+	const TArray<FViewerInfo>& Viewers = LODManager->GetViewers();
+	LODConfig.LODCalculator.PrepareExecution(Viewers);
+
+	{
+		TRACE_CPUPROFILER_EVENT_SCOPE(TEXT("CalculateLOD"));
+
+		LODConfig.EntityQuery.ForEachEntityChunk(EntitySubsystem, Context, [&LODConfig](FLWComponentSystemExecutionContext& Context)
+		{
+			if (!LODConfig.LODTickRateController.ShouldCalculateLODForChunk(Context))
+			{
+				return;
+			}
+
+			const TConstArrayView<FMassLODInfoFragment> ViewersInfoList = Context.GetComponentView<FMassLODInfoFragment>();
+			const TArrayView<FMassSimulationLODFragment> SimulationLODFragments = Context.GetMutableComponentView<FMassSimulationLODFragment>();
+			LODConfig.LODCalculator.CalculateLOD(Context, SimulationLODFragments, ViewersInfoList);
+		});
+	}
+
+	{
+		TRACE_CPUPROFILER_EVENT_SCOPE(TEXT("AdjustDistancesAndLODFromCount"));
+
+		if (LODConfig.LODCalculator.AdjustDistancesFromCount())
+		{
+			LODConfig.EntityQuery.ForEachEntityChunk(EntitySubsystem, Context, [&LODConfig](FLWComponentSystemExecutionContext& Context)
+			{
+				if (!LODConfig.LODTickRateController.ShouldAdjustLODFromCountForChunk(Context))
+				{
+					return;
+				}
+
+				const TConstArrayView<FMassLODInfoFragment> ViewersInfoList = Context.GetComponentView<FMassLODInfoFragment>();
+				const TArrayView<FMassSimulationLODFragment> SimulationLODFragments = Context.GetMutableComponentView<FMassSimulationLODFragment>();
+				LODConfig.LODCalculator.AdjustLODFromCount(Context, SimulationLODFragments, ViewersInfoList);
+			});
+		}
+	}
+
+	{
+		TRACE_CPUPROFILER_EVENT_SCOPE(TEXT("VariableTickRates"))
+		check(World);
+		const float Time = World->GetTimeSeconds();
+		LODConfig.EntityQuery.ForEachEntityChunk(EntitySubsystem, Context, [&LODConfig, Time](FLWComponentSystemExecutionContext& Context)
+		{
+			TArrayView<FMassSimulationLODFragment> SimulationLODFragments = Context.GetMutableComponentView<FMassSimulationLODFragment>();
+			LODConfig.LODTickRateController.UpdateTickRateFromLOD(Context, SimulationLODFragments, Time);
+		});
+	}
+
+	// Optional debug display
+	if (UE::MassLOD::bDebugSimulationLOD)
+	{
+		TRACE_CPUPROFILER_EVENT_SCOPE(TEXT("DebugDisplayLOD"));
+
+		LODConfig.EntityQuery.ForEachEntityChunk(EntitySubsystem, Context, [this, &LODConfig](FLWComponentSystemExecutionContext& Context)
+		{
+			const TConstArrayView<FDataFragment_Transform> LocationList = Context.GetComponentView<FDataFragment_Transform>();
+			const TConstArrayView<FMassSimulationLODFragment> SimulationLODList = Context.GetComponentView<FMassSimulationLODFragment>();
+			LODConfig.LODCalculator.DebugDisplayLOD(Context, SimulationLODList, LocationList, World);
+		});
+	}
+}
\ No newline at end of file
diff --git a/Engine/Plugins/Runtime/MassGameplay/Source/MassLOD/Public/MassLODCalculator.h b/Engine/Plugins/Runtime/MassGameplay/Source/MassLOD/Public/MassLODCalculator.h
new file mode 100644
--- /dev/null
+++ b/Engine/Plugins/Runtime/MassGameplay/Source/MassLOD/Public/MassLODCalculator.h
@@ -0,0 +1,584 @@
+// Copyright Epic Games, Inc. All Rights Reserved.
+
+#pragma once
+
+#include "MassLODLogic.h"
+#include "MassLODUtils.h"
+#include "DrawDebugHelpers.h"
+
+/**
+ * Helper struct to calculate LOD for each agent and maximize count per LOD
+ *   Requires FMassLODSourceInfo fragment collected by the TMassLODCollector.
+ *   Stores information in FMassLODResultInfo fragment.
+*/
+template <typename FLODLogic = FLODDefaultLogic >
+struct TMassLODCalculator : public FMassLODBaseLogic
+{
+public:
+
+	/**
+	 * Initializes the LOD calculator, needed to be called once at initialization time
+	 * @Param InBaseLODDistance distances used to calculate LOD
+	 * @Param InBufferHysteresisOnFOVRatio distance hysteresis used to calculate LOD
+	 * @Param InLODMaxCount the maximum count for each LOD
+	 * @Param InLODMaxCountPerViewer the maximum count for each LOD per viewer (Only when FLODLogic::bMaximizeCountPerViewer is enabled)
+	 * @Param InVisibleLODDistance the maximum count for each LOD per viewer (Only when FLODLogic::bDoVisibilityLogic is enabled)
+	 */
+	void Initialize(const float InBaseLODDistance[EMassLOD::Max], 
+					const float InBufferHysteresisOnDistanceRatio, 
+					const int32 InLODMaxCount[EMassLOD::Max], 
+					const int32 InLODMaxCountPerViewer[EMassLOD::Max] = nullptr,
+					const float InVisibleLODDistance[EMassLOD::Max] = nullptr);
+
+	/**
+	 * Prepares execution for the current frame, needed to be called before every execution
+	 * @Param Viewers is the array of all the known viewers
+	 */
+	void PrepareExecution(TConstArrayView<FViewerInfo> Viewers);
+
+	/**
+	 * Calculate LOD, called for each entity chunks
+	 * @Param Context of the chunk execution
+	 * @Param LODList is the fragment where calculation are stored
+	 * @Param ViewersInfoList is the source information fragment for LOD calculation
+	 */
+	template< typename FMassLODResultInfo, typename FMassLODSourceInfo >
+	void CalculateLOD(FLWComponentSystemExecutionContext& Context, TArrayView<FMassLODResultInfo> LODList, TConstArrayView<FMassLODSourceInfo> ViewersInfoList);
+
+	/**
+	 * Adjust LOD distances by clamping them to respect the maximum LOD count
+	 * @Return true if any LOD distances clamping was done
+	 */
+	bool AdjustDistancesFromCount();
+
+	/**
+	 * Adjust LOD from newly adjusted distances, only needed to be called when AdjustDistancesFromCount return true, called for each entity chunks
+	 * @Param Context of the chunk execution
+	 * @Param LODList is the fragment where calculation are stored
+	 * @Param ViewersInfoList is the source information fragment for LOD calculation
+	 */
+	template< typename FMassLODResultInfo, typename FMassLODSourceInfo >
+	void AdjustLODFromCount(FLWComponentSystemExecutionContext& Context, TArrayView<FMassLODResultInfo> LODList, TConstArrayView<FMassLODSourceInfo> ViewersInfoList);
+
+	/**
+	 * Turn Off all LOD, called for each entity chunks
+	 * @Param Context of the chunk execution
+	 * @Param LODList is the fragment where calculation are stored
+	 */
+	template <typename FMassLODResultInfo>
+	void ForceOffLOD(FLWComponentSystemExecutionContext& Context, TArrayView<FMassLODResultInfo> LODList);
+
+	/**
+	 * Debug draw the current state of each agent as a color coded square
+	 * @Param Context of the chunk execution
+	 * @Param LODList is the fragment where calculation are stored
+	 * @Param LocationList is the fragment transforms of the entities
+	 * @Param World where the debug display should be drawn
+	 */
+	template< typename FMassLODResultInfo, typename TFragmentLocation >
+	void DebugDisplayLOD(FLWComponentSystemExecutionContext& Context, TConstArrayView<FMassLODResultInfo> LODList, TConstArrayView<TFragmentLocation> LocationList, UWorld* World);
+
+protected:
+
+	struct FMassLODRuntimeData
+	{
+		/** Reset values to default */
+		void Reset(const TStaticArray<float, EMassLOD::Max>& InBaseLODDistance, const TStaticArray<float, EMassLOD::Max>& InVisibleLODDistance)
+		{
+			// Reset the AdjustedLODDistances as they might have been changed by the max count calculation previous frame
+			for (int32 LODDistIdx = 0; LODDistIdx < EMassLOD::Max; LODDistIdx++)
+			{
+				AdjustedBaseLODDistance[LODDistIdx] = InBaseLODDistance[LODDistIdx];
+				AdjustedBaseLODDistanceSq[LODDistIdx] = FMath::Square(AdjustedBaseLODDistance[LODDistIdx]);
+				if (FLODLogic::bDoVisibilityLogic)
+				{
+					AdjustedVisibleLODDistance[LODDistIdx] = InVisibleLODDistance[LODDistIdx];
+					AdjustedVisibleLODDistanceSq[LODDistIdx] = FMath::Square(AdjustedVisibleLODDistance[LODDistIdx]);
+				}
+			}
+			FMemory::Memzero(BaseBucketCounts);
+			if (FLODLogic::bDoVisibilityLogic)
+			{
+				FMemory::Memzero(VisibleBucketCounts);
+			}
+		}
+
+		/** Distance where each LOD becomes relevant (Squared and Normal) */
+		TStaticArray<float, EMassLOD::Max> AdjustedBaseLODDistanceSq;
+		TStaticArray<float, EMassLOD::Max> AdjustedVisibleLODDistanceSq;
+		TStaticArray<float, EMassLOD::Max> AdjustedBaseLODDistance;
+		TStaticArray<float, EMassLOD::Max> AdjustedVisibleLODDistance;
+
+		/** Count of entities in each subdivision */
+		TStaticArray< TStaticArray<int32, UE::MassLOD::MaxBucketsPerLOD>, EMassLOD::Max > BaseBucketCounts;
+		TStaticArray< TStaticArray<int32, UE::MassLOD::MaxBucketsPerLOD>, EMassLOD::Max > VisibleBucketCounts;
+	};
+
+
+	template <bool bCalculateLODSignificance>
+	float AccumulateCountInRuntimeData(const EMassLOD::Type LOD, const float ViewerDistanceSq, const bool bIsVisible, FMassLODRuntimeData& Data) const;
+
+	bool AdjustDistancesFromCountForRuntimeData(const TStaticArray<int32, EMassLOD::Max>& MaxCount, FMassLODRuntimeData& RuntimeData) const;
+
+	EMassLOD::Type ComputeLODFromSettings(const EMassLOD::Type PrevLOD, const float DistanceToViewerSq, const bool bIsVisible, bool* bIsInAVisibleRange, const FMassLODRuntimeData& Data) const;
+
+	/** LOD distances */
+	TStaticArray<float, EMassLOD::Max> BaseLODDistance;
+	TStaticArray<float, EMassLOD::Max> VisibleLODDistance;
+
+	/** MaxCount total */
+	TStaticArray<int32, EMassLOD::Max> LODMaxCount;
+
+	/** MaxCount total per viewers*/
+	TStaticArray<int32, EMassLOD::Max> LODMaxCountPerViewer;
+
+	/** Ratio for Buffer Distance Hysteresis */
+	float BufferHysteresisOnDistanceRatio = 0.1f;
+
+	/** The size of each subdivision per LOD (LOD Size/MaxBucketsPerLOD) */
+	TStaticArray<float, EMassLOD::Max> BaseBucketSize;
+	TStaticArray<float, EMassLOD::Max> VisibleBucketSize;
+
+	/** Maximum LOD Distance  */
+	float MaxLODDistance = 0.0f;
+
+	/** Runtime data for LOD calculation */
+	FMassLODRuntimeData RuntimeData;
+
+	/** Runtime data for each viewer specific LOD calculation, used only when bMaximizeCountPerViewer is true */
+	TStaticArray<FMassLODRuntimeData, UE::MassLOD::MaxNumOfViewers> RuntimeDataPerViewer;
+};
+
+template <typename FLODLogic>
+void TMassLODCalculator<FLODLogic>::Initialize(const float InBaseLODDistance[EMassLOD::Max],
+											   const float InBufferHysteresisOnDistanceRatio,
+											   const int32 InLODMaxCount[EMassLOD::Max],
+											   const int32 InLODMaxCountPerViewer[EMassLOD::Max] /*= nullptr*/,
+											   const float InVisibleLODDistance[EMassLOD::Max] /*= nullptr*/)
+{
+	checkf(FLODLogic::bMaximizeCountPerViewer == (InLODMaxCountPerViewer != nullptr), TEXT("Missmatched between expected parameter InLODMaxCountPerViewer and LOD logic trait bMaximizeCountPerViewer."));
+	checkf(FLODLogic::bDoVisibilityLogic == (InVisibleLODDistance != nullptr), TEXT("Missmatched between expected parameter InVisibleLODDistance and LOD logic trait bDoVisibilityLogic."));
+
+	// Make a copy of all the settings
+	for (int x = 0; x < EMassLOD::Max; x++)
+	{
+		BaseLODDistance[x] = InBaseLODDistance[x];
+		LODMaxCount[x] = InLODMaxCount[x];
+		if (FLODLogic::bDoVisibilityLogic && InVisibleLODDistance)
+		{
+			VisibleLODDistance[x] = InVisibleLODDistance[x];
+		}
+		if (FLODLogic::bMaximizeCountPerViewer && InLODMaxCountPerViewer)
+		{
+			LODMaxCountPerViewer[x] = InLODMaxCountPerViewer[x];
+		}
+	}
+
+	// Some values should always be constant
+	BaseLODDistance[EMassLOD::High] = 0.0f;
+	BaseBucketSize[EMassLOD::Off] = FLT_MAX;
+	VisibleLODDistance[EMassLOD::High] = 0.0f;
+	VisibleBucketSize[EMassLOD::Off] = FLT_MAX;
+	LODMaxCount[EMassLOD::Off] = INT_MAX;
+	LODMaxCountPerViewer[EMassLOD::Off] = INT_MAX;
+
+	// Calculate the size for each LOD buckets
+	float BasePrevLODDistance = BaseLODDistance[0];
+	float VisiblePrevLODDistance = VisibleLODDistance[0];
+	for (int32 LODDistIdx = 1; LODDistIdx < EMassLOD::Max; LODDistIdx++)
+	{
+		BaseBucketSize[LODDistIdx - 1] = (BaseLODDistance[LODDistIdx] - BasePrevLODDistance) / UE::MassLOD::MaxBucketsPerLOD;
+		BasePrevLODDistance = BaseLODDistance[LODDistIdx];
+
+		if (FLODLogic::bDoVisibilityLogic)
+		{
+			VisibleBucketSize[LODDistIdx - 1] = (VisibleLODDistance[LODDistIdx] - VisiblePrevLODDistance) / UE::MassLOD::MaxBucketsPerLOD;
+			VisiblePrevLODDistance = VisibleLODDistance[LODDistIdx];
+		}
+	}
+
+	// Assuming that off is the farthest distance, calculate the max LOD distance
+	MaxLODDistance = !FLODLogic::bDoVisibilityLogic || BaseLODDistance[EMassLOD::Off] >= VisibleLODDistance[EMassLOD::Off] ? BaseLODDistance[EMassLOD::Off] : VisibleLODDistance[EMassLOD::Off];
+}
+
+template <typename FLODLogic>
+void TMassLODCalculator<FLODLogic>::PrepareExecution(TConstArrayView<FViewerInfo> Viewers)
+{
+	CacheViewerInformation(Viewers, FLODLogic::bLocalViewersOnly);
+
+	if (FLODLogic::bMaximizeCountPerViewer)
+	{
+		for (int ViewerIdx = 0; ViewerIdx < NumOfViewers; ++ViewerIdx)
+		{
+			// Reset viewer data
+			if (ViewersHandles[ViewerIdx].IsValid())
+			{
+				RuntimeDataPerViewer[ViewerIdx].Reset(BaseLODDistance, VisibleLODDistance);
+			}
+		}
+	}
+
+	RuntimeData.Reset(BaseLODDistance, VisibleLODDistance);
+}
+
+
+template <typename FLODLogic>
+template <bool bCalculateLODSignificance>
+float TMassLODCalculator<FLODLogic>::AccumulateCountInRuntimeData(const EMassLOD::Type LOD, const float ViewerDistanceSq, const bool bIsVisible, FMassLODRuntimeData& Data) const
+{
+	TStaticArray< TStaticArray<int32, UE::MassLOD::MaxBucketsPerLOD>, EMassLOD::Max>& BucketCounts = FLODLogic::bDoVisibilityLogic && bIsVisible ? Data.VisibleBucketCounts : Data.BaseBucketCounts;
+
+	// Cumulate LOD in buckets for Max LOD count calculation
+	if (LOD == EMassLOD::Off)
+	{
+		// A single bucket for Off LOD
+		BucketCounts[EMassLOD::Off][0]++;
+		if (bCalculateLODSignificance)
+		{
+			return float(EMassLOD::Off);
+		}
+	}
+	else
+	{
+		const TStaticArray<float, EMassLOD::Max>& BucketSize = FLODLogic::bDoVisibilityLogic && bIsVisible ? VisibleBucketSize : BaseBucketSize;
+		const TStaticArray<float, EMassLOD::Max>& AdjustedLODDistance = FLODLogic::bDoVisibilityLogic && bIsVisible ? Data.AdjustedVisibleLODDistance : Data.AdjustedBaseLODDistance;
+
+		const int32 LODDistIdx = (int32)LOD;
+
+		// Need to clamp as the Sqrt is not precise enough and always end up with floating calculation errors
+		const int32 BucketIdx = FMath::Clamp((int32)((FMath::Sqrt(ViewerDistanceSq) - AdjustedLODDistance[LODDistIdx]) / BucketSize[LODDistIdx]), 0, UE::MassLOD::MaxBucketsPerLOD - 1);
+		BucketCounts[LODDistIdx][BucketIdx]++;
+
+		if (bCalculateLODSignificance)
+		{
+			// Derive significance from LODDistIdx combined with BucketIdx
+			const float PartialLODSignificance = float(BucketIdx) / float(UE::MassLOD::MaxBucketsPerLOD);
+			return float(LODDistIdx) + PartialLODSignificance;
+		}
+	}
+	return 0.0f;
+}
+
+template <typename FLODLogic>
+template< typename FMassLODResultInfo, typename FMassLODSourceInfo >
+void TMassLODCalculator<FLODLogic>::CalculateLOD(FLWComponentSystemExecutionContext& Context, TArrayView<FMassLODResultInfo> LODList, TConstArrayView<FMassLODSourceInfo> ViewersInfoList)
+{
+#if WITH_MASS_DEBUG
+	if (UE::MassLOD::Debug::bLODCalculationsPaused)
+	{
+		return;
+	}
+#endif // WITH_MASS_DEBUG
+
+	const int32 NumEntities = Context.GetEntitiesNum();
+	for (int EntityIdx = 0; EntityIdx < NumEntities; EntityIdx++)
+	{
+		// Calculate the LOD purely upon distances
+		const FMassLODSourceInfo& EntityViewersInfo = ViewersInfoList[EntityIdx];
+		FMassLODResultInfo& EntityLOD = LODList[EntityIdx];
+		EntityLOD.ClosestViewerDistanceSq = FLT_MAX;
+		bool bIsVisibleByAViewer = false;
+		bool bIsInAVisibleRange = false;
+
+		for (int ViewerIdx = 0; ViewerIdx < NumOfViewers; ++ViewerIdx)
+		{
+			if (bClearViewerData[ViewerIdx])
+			{
+				SetLODPerViewer<FLODLogic::bStoreLODPerViewer>(EntityLOD, ViewerIdx, EMassLOD::Max);
+				SetPrevLODPerViewer<FLODLogic::bStoreLODPerViewer>(EntityLOD, ViewerIdx, EMassLOD::Max);
+			}
+			if (ViewersHandles[ViewerIdx].IsValid())
+			{
+				const bool bIsVisibleByViewer = GetbIsVisibleByViewer<FLODLogic::bDoVisibilityLogic>(EntityViewersInfo, ViewerIdx, false);
+				bIsVisibleByAViewer |= bIsVisibleByViewer;
+
+				if (EntityLOD.ClosestViewerDistanceSq > EntityViewersInfo.DistanceToViewerSq[ViewerIdx])
+				{
+					EntityLOD.ClosestViewerDistanceSq = EntityViewersInfo.DistanceToViewerSq[ViewerIdx];
+				}
+
+				if (FLODLogic::bStoreLODPerViewer || FLODLogic::bMaximizeCountPerViewer)
+				{
+					const EMassLOD::Type PrevLODPerViewer = GetLODPerViewer<FLODLogic::bStoreLODPerViewer>(EntityLOD, ViewerIdx, EntityLOD.LOD);
+
+					// Find new LOD
+					const EMassLOD::Type LODPerViewer = ComputeLODFromSettings(PrevLODPerViewer, EntityViewersInfo.DistanceToViewerSq[ViewerIdx], bIsVisibleByViewer, &bIsInAVisibleRange, RuntimeData);
+
+					// Set Per viewer LOD
+					SetPrevLODPerViewer<FLODLogic::bStoreLODPerViewer>(EntityLOD, ViewerIdx, PrevLODPerViewer);
+					SetLODPerViewer<FLODLogic::bStoreLODPerViewer>(EntityLOD, ViewerIdx, LODPerViewer);
+
+					if (FLODLogic::bMaximizeCountPerViewer)
+					{
+						// Accumulate in buckets
+						AccumulateCountInRuntimeData<false>(LODPerViewer, EntityViewersInfo.DistanceToViewerSq[ViewerIdx], bIsVisibleByViewer, RuntimeDataPerViewer[ViewerIdx]);
+					}
+				}
+			}
+		}
+
+		// Find new LOD
+		EntityLOD.PrevLOD = EntityLOD.LOD;
+		EntityLOD.LOD = ComputeLODFromSettings(EntityLOD.PrevLOD, EntityLOD.ClosestViewerDistanceSq, bIsVisibleByAViewer, &bIsInAVisibleRange, RuntimeData);
+
+		// Set visibility
+		SetbWasVisibleByAViewer<FLODLogic::bDoVisibilityLogic>(EntityLOD, GetbIsVisibleByAViewer<FLODLogic::bDoVisibilityLogic>(EntityLOD, false));
+		SetbIsVisibleByAViewer<FLODLogic::bDoVisibilityLogic>(EntityLOD, bIsVisibleByAViewer);
+		SetbWasInVisibleRange<FLODLogic::bDoVisibilityLogic>(EntityLOD, GetbIsInVisibleRange<FLODLogic::bDoVisibilityLogic>(EntityLOD, false));
+		SetbIsInVisibleRange<FLODLogic::bDoVisibilityLogic>(EntityLOD, bIsInAVisibleRange);
+
+		// Accumulate in buckets
+		const float LODSignificance = AccumulateCountInRuntimeData<FLODLogic::bCalculateLODSignificance>(EntityLOD.LOD, EntityLOD.ClosestViewerDistanceSq, bIsVisibleByAViewer, RuntimeData);
+		SetLODSignificance<FLODLogic::bCalculateLODSignificance>(EntityLOD, LODSignificance);
+	}
+}
+
+template <typename FLODLogic>
+bool TMassLODCalculator<FLODLogic>::AdjustDistancesFromCountForRuntimeData(const TStaticArray<int32, EMassLOD::Max>& MaxCount, FMassLODRuntimeData& Data) const
+{
+	int32 Count = 0;
+	int32 ProcessingLODIdx = EMassLOD::High;
+
+	bool bNeedAdjustments = false;
+
+	// Go through all LOD can start counting from the high LOD
+	for (int32 BucketLODIdx = 0; BucketLODIdx < EMassLOD::Max - 1; ++BucketLODIdx)
+	{
+		// Switch to next LOD if we have not reach the max
+		if (ProcessingLODIdx < BucketLODIdx)
+		{
+			// Switch to next LOD
+			ProcessingLODIdx = BucketLODIdx;
+
+			// Restart the count
+			Count = 0;
+		}
+
+		// Count entities through all buckets of this LOD
+		for (int32 BucketIdx = 0; BucketIdx < UE::MassLOD::MaxBucketsPerLOD; ++BucketIdx)
+		{
+			if (FLODLogic::bDoVisibilityLogic)
+			{
+				// Do visible count first to prioritize them over none visible for that bucket idx
+				Count += Data.VisibleBucketCounts[BucketLODIdx][BucketIdx];
+
+				while (Count > MaxCount[ProcessingLODIdx])
+				{
+					// Switch to next LOD
+					ProcessingLODIdx++;
+
+					// Adjust distance for this LOD
+					Data.AdjustedBaseLODDistance[ProcessingLODIdx] = BaseLODDistance[BucketLODIdx] + (BucketIdx * BaseBucketSize[BucketLODIdx]);
+					Data.AdjustedBaseLODDistanceSq[ProcessingLODIdx] = FMath::Square(Data.AdjustedBaseLODDistance[ProcessingLODIdx]);
+					Data.AdjustedVisibleLODDistance[ProcessingLODIdx] = VisibleLODDistance[BucketLODIdx] + (BucketIdx * VisibleBucketSize[BucketLODIdx]);
+					Data.AdjustedVisibleLODDistanceSq[ProcessingLODIdx] = FMath::Square(Data.AdjustedVisibleLODDistance[ProcessingLODIdx]);
+
+					// Check if we are done
+					if (ProcessingLODIdx == EMassLOD::Off)
+					{
+						return true;
+					}
+
+					// Start the next LOD count with the bucket count that made it go over
+					Count = Data.VisibleBucketCounts[BucketLODIdx][BucketIdx];
+
+					bNeedAdjustments = true;
+				}
+			}
+
+			// Add base count
+			Count += Data.BaseBucketCounts[BucketLODIdx][BucketIdx];
+
+			// Check if the count is going over max
+			while (Count > MaxCount[ProcessingLODIdx])
+			{
+				// Switch to next LOD
+				ProcessingLODIdx++;
+
+				// Adjust distance for this LOD
+				Data.AdjustedBaseLODDistance[ProcessingLODIdx] = BaseLODDistance[BucketLODIdx] + (BucketIdx * BaseBucketSize[BucketLODIdx]);
+				Data.AdjustedBaseLODDistanceSq[ProcessingLODIdx] = FMath::Square(Data.AdjustedBaseLODDistance[ProcessingLODIdx]);
+				if (FLODLogic::bDoVisibilityLogic)
+				{
+					Data.AdjustedVisibleLODDistance[ProcessingLODIdx] = VisibleLODDistance[BucketLODIdx] + ((BucketIdx + 1) * VisibleBucketSize[BucketLODIdx]);
+					Data.AdjustedVisibleLODDistanceSq[ProcessingLODIdx] = FMath::Square(Data.AdjustedVisibleLODDistance[ProcessingLODIdx]);
+				}
+
+				// Check if we are done
+				if (ProcessingLODIdx == EMassLOD::Off)
+				{
+					return true;
+				}
+
+				// Start the next LOD count with the bucket count that made it go over
+				Count = Data.BaseBucketCounts[BucketLODIdx][BucketIdx];
+
+				bNeedAdjustments = true;
+			}
+		}
+	}
+
+	return bNeedAdjustments;
+}
+
+
+template <typename FLODLogic>
+bool TMassLODCalculator<FLODLogic>::AdjustDistancesFromCount()
+{
+	bool bDistanceAdjusted = false;
+	if (FLODLogic::bMaximizeCountPerViewer)
+	{
+		for (int ViewerIdx = 0; ViewerIdx < NumOfViewers; ++ViewerIdx)
+		{
+			if (!ViewersHandles[ViewerIdx].IsValid())
+			{
+				continue;
+			}
+
+			bDistanceAdjusted |= AdjustDistancesFromCountForRuntimeData(LODMaxCountPerViewer, RuntimeDataPerViewer[ViewerIdx]);
+		}
+	}
+
+	bDistanceAdjusted |= AdjustDistancesFromCountForRuntimeData(LODMaxCount, RuntimeData);
+	return bDistanceAdjusted;
+}
+
+template <typename FLODLogic>
+template< typename FMassLODResultInfo, typename FMassLODSourceInfo >
+void TMassLODCalculator<FLODLogic>::AdjustLODFromCount(FLWComponentSystemExecutionContext& Context, TArrayView<FMassLODResultInfo> LODList, TConstArrayView<FMassLODSourceInfo> ViewersInfoList)
+{
+	const int32 NumEntities = Context.GetEntitiesNum();
+	// Adjust LOD for each viewer and remember the new highest
+	for (int EntityIdx = 0; EntityIdx < NumEntities; EntityIdx++)
+	{
+		const FMassLODSourceInfo& EntityViewersInfo = ViewersInfoList[EntityIdx];
+		FMassLODResultInfo& EntityLOD = LODList[EntityIdx];
+		EMassLOD::Type HighestViewerLOD = EMassLOD::Off;
+		if (FLODLogic::bMaximizeCountPerViewer)
+		{
+			for (int ViewerIdx = 0; ViewerIdx < NumOfViewers; ++ViewerIdx)
+			{
+				if (!ViewersHandles[ViewerIdx].IsValid())
+				{
+					continue;
+				}
+
+				const bool bIsVisibleByViewer = GetbIsVisibleByViewer<FLODLogic::bDoVisibilityLogic>(EntityViewersInfo, ViewerIdx, false);
+				EMassLOD::Type LODPerViewer = GetLODPerViewer<FLODLogic::bStoreLODPerViewer>(EntityLOD, ViewerIdx, EntityLOD.LOD);
+
+				LODPerViewer = ComputeLODFromSettings(LODPerViewer, EntityViewersInfo.DistanceToViewerSq[ViewerIdx], bIsVisibleByViewer, nullptr, RuntimeDataPerViewer[ViewerIdx]);
+
+				if (HighestViewerLOD < LODPerViewer)
+				{
+					HighestViewerLOD = LODPerViewer;
+				}
+
+				SetLODPerViewer<FLODLogic::bStoreLODPerViewer>(EntityLOD, ViewerIdx, LODPerViewer);
+			}
+		}
+
+		const bool bIsVisibleByAViewer = GetbIsVisibleByAViewer<FLODLogic::bDoVisibilityLogic>(EntityLOD, false);
+		EMassLOD::Type NewLOD = ComputeLODFromSettings(EntityLOD.PrevLOD, EntityLOD.ClosestViewerDistanceSq, bIsVisibleByAViewer, nullptr, RuntimeData);
+
+		// Maybe the highest of all the viewers is now lower than the global entity LOD, make sure to update the it accordingly
+		if (FLODLogic::bMaximizeCountPerViewer && NewLOD < HighestViewerLOD)
+		{
+			NewLOD = HighestViewerLOD;
+		}
+		if (EntityLOD.LOD != NewLOD)
+		{
+			EntityLOD.LOD = NewLOD;
+			if (FLODLogic::bCalculateLODSignificance)
+			{
+				float LODSignificance = 0.f;
+				if (NewLOD == EMassLOD::Off)
+				{
+					LODSignificance = float(EMassLOD::Off);
+				}
+				else
+				{
+					const TStaticArray<float, EMassLOD::Max>& AdjustedLODDistance = FLODLogic::bDoVisibilityLogic && bIsVisibleByAViewer ? RuntimeData.AdjustedVisibleLODDistance : RuntimeData.AdjustedBaseLODDistance;
+
+					// Need to clamp as the Sqrt is not precise enough and always end up with floating calculation errors
+					const float DistanceBetweenLODThresholdAndEntity = FMath::Max(FMath::Sqrt(EntityLOD.ClosestViewerDistanceSq) - AdjustedLODDistance[NewLOD], 0.f);
+
+					// Derive significance from distance between viewer and where the agent stands between both LOD threshold
+					const float AdjustedDistanceBetweenCurrentLODAndNext = AdjustedLODDistance[NewLOD + 1] - AdjustedLODDistance[NewLOD];
+					const float PartialLODSignificance = DistanceBetweenLODThresholdAndEntity / AdjustedDistanceBetweenCurrentLODAndNext;
+					LODSignificance = float(NewLOD) + PartialLODSignificance;
+				}
+
+				SetLODSignificance<FLODLogic::bCalculateLODSignificance>(EntityLOD, LODSignificance);
+			}
+		}
+	}
+}
+
+template <typename FLODLogic>
+template <typename FMassLODResultInfo>
+void TMassLODCalculator<FLODLogic>::ForceOffLOD(FLWComponentSystemExecutionContext& Context, TArrayView<FMassLODResultInfo> LODList)
+{
+	const int32 NumEntities = Context.GetEntitiesNum();
+	for (int EntityIdx = 0; EntityIdx < NumEntities; EntityIdx++)
+	{
+		FMassLODResultInfo& EntityLOD = LODList[EntityIdx];
+
+		// Force off LOD
+		EntityLOD.PrevLOD = EntityLOD.LOD;
+		EntityLOD.LOD = EMassLOD::Off;
+
+		SetbWasVisibleByAViewer<FLODLogic::bDoVisibilityLogic>(EntityLOD, GetbIsVisibleByAViewer<FLODLogic::bDoVisibilityLogic>(EntityLOD, false));
+		SetbIsVisibleByAViewer<FLODLogic::bDoVisibilityLogic>(EntityLOD, false);
+		SetbWasInVisibleRange<FLODLogic::bDoVisibilityLogic>(EntityLOD, GetbIsInVisibleRange<FLODLogic::bDoVisibilityLogic>(EntityLOD, false));
+		SetbIsInVisibleRange<FLODLogic::bDoVisibilityLogic>(EntityLOD, false);
+		SetLODSignificance<FLODLogic::bCalculateLODSignificance>(EntityLOD, float(EMassLOD::Off));
+	}
+}
+
+template <typename FLODLogic>
+template< typename FMassLODResultInfo, typename TFragmentLocation >
+void TMassLODCalculator<FLODLogic>::DebugDisplayLOD(FLWComponentSystemExecutionContext& Context, TConstArrayView<FMassLODResultInfo> LODList, TConstArrayView<TFragmentLocation> LocationList, UWorld* World)
+{
+	const int32 NumEntities = Context.GetEntitiesNum();
+	for (int EntityIdx = 0; EntityIdx < NumEntities; EntityIdx++)
+	{
+		const TFragmentLocation& EntityLocation = LocationList[EntityIdx];
+		const FMassLODResultInfo& EntityLOD = LODList[EntityIdx];
+		int32 LODIdx = (int32)EntityLOD.LOD;
+		DrawDebugSolidBox(World, EntityLocation.GetTransform().GetLocation() + FVector(0.0f, 0.0f, 120.0f), FVector(25.0f), UE::MassLOD::LODColors[LODIdx]);
+	}
+}
+
+template <typename FLODLogic>
+EMassLOD::Type TMassLODCalculator<FLODLogic>::ComputeLODFromSettings(const EMassLOD::Type PrevLOD, const float DistanceToViewerSq, const bool bIsVisible, bool* bIsInAVisibleRange, const FMassLODRuntimeData& Data) const
+{
+	const TStaticArray<float, EMassLOD::Max>& AdjustedLODDistanceSq = FLODLogic::bDoVisibilityLogic && bIsVisible ? Data.AdjustedVisibleLODDistanceSq : Data.AdjustedBaseLODDistanceSq;
+	int32 LODDistIdx = EMassLOD::Max - 1;
+	for (; LODDistIdx > 0; LODDistIdx--)
+	{
+		if (DistanceToViewerSq >= AdjustedLODDistanceSq[LODDistIdx])
+		{
+			// Validate that we allow going to a single higher LOD level after considering an extended buffer hysteresis on distance for the LOD level we are about to quit to prevent oscillating LOD states
+			if (PrevLOD != EMassLOD::Max && (PrevLOD - (EMassLOD::Type)LODDistIdx) == 1)
+			{
+				const TStaticArray<float, EMassLOD::Max>& AdjustedLODDistance = FLODLogic::bDoVisibilityLogic && bIsVisible ? Data.AdjustedVisibleLODDistance : Data.AdjustedBaseLODDistance;
+				float HysteresisDistance = FMath::Lerp(AdjustedLODDistance[LODDistIdx], AdjustedLODDistance[LODDistIdx + 1], 1.f - BufferHysteresisOnDistanceRatio);
+				if (DistanceToViewerSq >= FMath::Square(HysteresisDistance))
+				{
+					// Keep old
+					LODDistIdx = PrevLOD;
+				}
+			}
+
+			break;
+		}
+	}
+
+	EMassLOD::Type NewLOD = (EMassLOD::Type)LODDistIdx;
+	if(FLODLogic::bDoVisibilityLogic && bIsInAVisibleRange)
+	{
+		*bIsInAVisibleRange = bIsVisible ? (NewLOD != EMassLOD::Off) : (DistanceToViewerSq < Data.AdjustedVisibleLODDistanceSq[EMassLOD::Off]);
+	}
+
+	return NewLOD;
+}
\ No newline at end of file
diff --git a/Engine/Plugins/Runtime/MassGameplay/Source/MassLOD/Public/MassLODCollector.h b/Engine/Plugins/Runtime/MassGameplay/Source/MassLOD/Public/MassLODCollector.h
new file mode 100644
--- /dev/null
+++ b/Engine/Plugins/Runtime/MassGameplay/Source/MassLOD/Public/MassLODCollector.h
@@ -0,0 +1,124 @@
+// Copyright Epic Games, Inc. All Rights Reserved.
+
+#pragma once
+
+#include "MassLODLogic.h"
+#include "DrawDebugHelpers.h"
+
+/**
+ * This is the expected member variables for the TArrayView<FMassLODSourceInfo> when calling TMassLODCollector methods
+ */
+struct FMassLODSourceInfo
+{
+	// Distance to viewer (Always needed)
+	TStaticArray<float, UE::MassLOD::MaxNumOfViewers> DistanceToViewerSq;
+
+	// Visibility information (Only when FLODLogic::bDoVisibilityLogic is enabled)
+	TStaticArray<bool, UE::MassLOD::MaxNumOfViewers> bIsVisibleByViewer;
+};
+
+/**
+ * Helper struct to collect needed information on each agent that will be needed later for LOD calculation
+ *   Requires FMassTransform fragment.
+ *   Stores information in FMassLODSourceInfo fragment.
+ */
+template <typename FLODLogic = FLODDefaultLogic >
+struct TMassLODCollector : public FMassLODBaseLogic
+{
+	/**
+	 * Initializes the LOD collector, needed to be called once at initialization time
+	 * @Param InFOVAngleToDriveVisibility the FOV angle that will be use for visibility check in degrees (Only when FLODLogic::bDoVisibilityLogic is enabled)
+	 * @Param InBufferHysteresisOnFOVRatio the FOV angle hysteresis that will be use for visibility check (Only when FLODLogic::bDoVisibilityLogic is enabled)
+	 */
+	void Initialize(const float InFOVAngleToDriveVisibility = -1.0f, const float InBufferHysteresisOnFOVRatio = -1.0f);
+
+	/**
+	 * Prepares execution for the current frame, needed to be called before every execution
+	 * @Param Viewers is the array of all the known viewers
+	 */
+	void PrepareExecution(TConstArrayView<FViewerInfo> Viewers);
+
+	/**
+	 * Collects the information for LOD calculation, called for each entity chunks
+	 * @Param Context of the chunk execution
+	 * @Param LocationList is the fragment transforms of the entities
+	 * @Param ViewersInfoList is the fragment where to store source information for LOD calculation
+	 */
+	template< typename FMassTransform, typename FMassLODSourceInfo >
+	void CollectLODInfo(FLWComponentSystemExecutionContext& Context, TConstArrayView<FMassTransform> LocationList, TArrayView<FMassLODSourceInfo> ViewersInfoList);
+
+protected:
+
+	float FOVAngleToDriveVisibilityRad = 0.0f;
+	float CosFOVAngleToDriveVisibility = 0.0f;
+	float BufferHysteresisOnFOVRatio = 0.1f;
+};
+
+template <typename FLODLogic>
+void TMassLODCollector<FLODLogic>::Initialize(const float InFOVAngleToDriveVisibility/*= -1.0f*/, const float InBufferHysteresisOnFOVRatio /*= -1.0f*/)
+{
+	checkf(!FLODLogic::bDoVisibilityLogic || (InFOVAngleToDriveVisibility > 0.0f && BufferHysteresisOnFOVRatio > 0.0f), TEXT("Expecting FOVAngleToDriveVisibility and BufferHysteresisOnFOVRatio when bDoVisibilityLogic is true"));
+	FOVAngleToDriveVisibilityRad = FMath::DegreesToRadians(InFOVAngleToDriveVisibility);
+	CosFOVAngleToDriveVisibility = FMath::Cos(FOVAngleToDriveVisibilityRad);
+	BufferHysteresisOnFOVRatio = InBufferHysteresisOnFOVRatio;
+}
+
+template <typename FLODLogic>
+void TMassLODCollector<FLODLogic>::PrepareExecution(TConstArrayView<FViewerInfo> Viewers)
+{
+	CacheViewerInformation(Viewers, FLODLogic::bLocalViewersOnly);
+}
+
+template <typename FLODLogic>
+template< typename FMassTransform, typename FMassLODSourceInfo >
+void TMassLODCollector<FLODLogic>::CollectLODInfo(FLWComponentSystemExecutionContext& Context, TConstArrayView<FMassTransform> LocationList, TArrayView<FMassLODSourceInfo> ViewersInfoList)
+{
+	const int32 NumEntities = Context.GetEntitiesNum();
+	for (int EntityIdx = 0; EntityIdx < NumEntities; EntityIdx++)
+	{
+		const FMassTransform& EntityLocation = LocationList[EntityIdx];
+		FMassLODSourceInfo& EntityViewerInfo = ViewersInfoList[EntityIdx];
+		for (int ViewerIdx = 0; ViewerIdx < NumOfViewers; ++ViewerIdx)
+		{
+			if (bClearViewerData[ViewerIdx])
+			{
+				EntityViewerInfo.DistanceToViewerSq[ViewerIdx] = FLT_MAX;
+				SetbIsVisibleByViewer<FLODLogic::bDoVisibilityLogic>(EntityViewerInfo, ViewerIdx, false);
+			}
+			if (ViewersHandles[ViewerIdx].IsValid())
+			{
+				const FVector ViewerToEntity = EntityLocation.GetTransform().GetLocation() - ViewersLocation[ViewerIdx];
+				EntityViewerInfo.DistanceToViewerSq[ViewerIdx] = ViewerToEntity.SizeSquared();
+
+				if (FLODLogic::bDoVisibilityLogic)
+				{
+					const bool bWasVisibleByViewer = GetbIsVisibleByViewer<FLODLogic::bDoVisibilityLogic>(EntityViewerInfo, ViewerIdx, false);
+					bool bIsVisibleByViewer = false;
+
+					if (EntityViewerInfo.DistanceToViewerSq[ViewerIdx] <= KINDA_SMALL_NUMBER)
+					{
+						bIsVisibleByViewer = true;
+					}
+					else
+					{
+						const FVector ViewerToEntityNorm = ViewerToEntity * FMath::InvSqrt(EntityViewerInfo.DistanceToViewerSq[ViewerIdx]);
+						const float Dot = ViewersDirection[ViewerIdx] | ViewerToEntityNorm;
+
+						// If we were visible and are now exiting defined FOV angle, consider a buffer hysteresis to give some room before toggling off visibility to prevent oscillating LOD states
+						if (bWasVisibleByViewer && Dot < CosFOVAngleToDriveVisibility)
+						{
+							const float FOVToCos = FMath::Cos(FMath::Clamp(FOVAngleToDriveVisibilityRad * (1.f + BufferHysteresisOnFOVRatio), 0.f, PI));
+							bIsVisibleByViewer = Dot > FOVToCos;
+						}
+						else
+						{
+							bIsVisibleByViewer = Dot >= CosFOVAngleToDriveVisibility;
+						}
+					}
+
+					SetbIsVisibleByViewer<FLODLogic::bDoVisibilityLogic>(EntityViewerInfo, ViewerIdx, bIsVisibleByViewer);
+				}
+			}
+		}
+	}
+}
\ No newline at end of file
diff --git a/Engine/Plugins/Runtime/MassGameplay/Source/MassLOD/Public/MassLODCollectorProcessor.h b/Engine/Plugins/Runtime/MassGameplay/Source/MassLOD/Public/MassLODCollectorProcessor.h
new file mode 100644
--- /dev/null
+++ b/Engine/Plugins/Runtime/MassGameplay/Source/MassLOD/Public/MassLODCollectorProcessor.h
@@ -0,0 +1,60 @@
+// Copyright Epic Games, Inc. All Rights Reserved.
+
+#pragma once
+
+#include "LWComponentTypes.h"
+#include "MassLODTypes.h"
+#include "MassLODManager.h"
+#include "MassLODCollector.h"
+
+#include "MassLODCollectorProcessor.generated.h"
+
+USTRUCT()
+struct FMassCollectorLODConfig
+{
+	GENERATED_BODY()
+
+	/** Component Tag that will be used to associate LOD config */
+	UPROPERTY(EditAnywhere, Category = "Mass|LOD", config, meta = (BaseStruct = "ComponentTag"))
+	FInstancedStruct TagFilter;
+
+	/** Set FOV Angles to define the limits between using Visible or Base LOD Distances */
+	UPROPERTY(EditAnywhere, Category = "Mass|LOD", meta = (ClampMin = "0.0", ClampMax = "180.0", UIMin = "0.0", UIMax = "180.0"), config)
+	float FOVAnglesToDriveVisibility = 45.0f;
+	UPROPERTY(EditAnywhere, Category = "Mass|LOD", meta = (ClampMin = "0.0", UIMin = "0.0"), config)
+	float BufferHysteresisOnFOVPercentage = 10.0f;
+
+	/** Runtime data for matching the LOD config */
+	TMassLODCollector<FMassRepresentationLODLogic> RepresentationLODCollector;
+	FLWComponentQuery CloseEntityQuery;
+	FLWComponentQuery FarEntityQuery_Conditional;
+	TMassLODCollector<FMassSimulationLODLogic> SimulationLODCollector;
+	FLWComponentQuery OnLODEntityQuery;
+	FLWComponentQuery OffLODEntityQuery_Conditional;
+	TMassLODCollector<FMassCombinedLODLogic> CombinedLODCollector;
+	// Keep these in this order as we are creating ArrayView from the first 3 queries.
+	FLWComponentQuery CloseAndOnLODEntityQuery;
+	FLWComponentQuery CloseAndOffLODEntityQuery;
+	FLWComponentQuery FarAndOnLODEntityQuery;
+	FLWComponentQuery FarAndOffLODEntityQuery_Conditional;
+};
+
+/*
+ * LOD collector which combines collection of LOD information for both Viewer and Simulation LODing when possible.
+ */
+UCLASS(meta = (DisplayName = "LOD Collector"))
+class MASSLOD_API UMassLODCollectorProcessor : public UMassProcessor_LODBase
+{
+	GENERATED_BODY()
+
+public:
+	UMassLODCollectorProcessor();
+
+protected:
+	virtual void Initialize(UObject& Owner) override;
+	virtual void ConfigureQueries() override;
+	virtual void Execute(UEntitySubsystem& EntitySubsystem, FLWComponentSystemExecutionContext& Context) override;
+
+	UPROPERTY(EditAnywhere, Category = "Mass|LOD", config)
+	TArray<FMassCollectorLODConfig> LODConfigs;
+};
diff --git a/Engine/Plugins/Runtime/MassGameplay/Source/MassLOD/Public/MassLODLogic.h b/Engine/Plugins/Runtime/MassGameplay/Source/MassLOD/Public/MassLODLogic.h
new file mode 100644
--- /dev/null
+++ b/Engine/Plugins/Runtime/MassGameplay/Source/MassLOD/Public/MassLODLogic.h
@@ -0,0 +1,128 @@
+// Copyright Epic Games, Inc. All Rights Reserved.
+
+#pragma once
+
+#include "MassLODManager.h"
+#include "Containers/StaticArray.h"
+
+
+#define DECLARE_CONDITIONAL_MEMBER_ACCESSORS( Condition, MemberType, MemberName ) \
+template <bool Condition, typename TemplateClass> static FORCEINLINE typename TEnableIf< Condition, MemberType>::Type Get##MemberName(TemplateClass& Obj, MemberType DefaultValue) { return Obj.MemberName; } \
+template <bool Condition, typename TemplateClass> static FORCEINLINE typename TEnableIf<!Condition, MemberType>::Type Get##MemberName(TemplateClass& Obj, MemberType DefaultValue) { return DefaultValue; } \
+template <bool Condition, typename TemplateClass> static FORCEINLINE void Set##MemberName(TemplateClass& Obj, typename TEnableIf< Condition, MemberType>::Type Value) { Obj.MemberName = Value; } \
+template <bool Condition, typename TemplateClass> static FORCEINLINE void Set##MemberName(TemplateClass& Obj, typename TEnableIf<!Condition, MemberType>::Type Value) {}
+
+#define DECLARE_CONDITIONAL_MEMBER_ARRAY_ACCESSORS( Condition, MemberType, MemberName ) \
+template <bool Condition, typename TemplateClass> static FORCEINLINE typename TEnableIf< Condition, MemberType>::Type Get##MemberName(TemplateClass& Obj, int32 Index, MemberType DefaultValue) { return Obj.MemberName[Index]; } \
+template <bool Condition, typename TemplateClass> static FORCEINLINE typename TEnableIf<!Condition, MemberType>::Type Get##MemberName(TemplateClass& Obj, int32 Index, MemberType DefaultValue) { return DefaultValue; } \
+template <bool Condition, typename TemplateClass> static FORCEINLINE void Set##MemberName(TemplateClass& Obj, int32 Index, typename TEnableIf< Condition, MemberType>::Type Value) { Obj.MemberName[Index] = Value; } \
+template <bool Condition, typename TemplateClass> static FORCEINLINE void Set##MemberName(TemplateClass& Obj, int32 Index, typename TEnableIf<!Condition, MemberType>::Type Value) {}
+
+/**
+ * Traits for LOD logic calculation behaviors
+ */
+struct FLODDefaultLogic
+{
+	enum
+	{
+		bStoreLODPerViewer = false, // Enable to calculate and store the result LOD per viewer in the FMassLODResultInfo::LODPerViewer and FMassLODResultInfo::PrevLODPerViewer.
+		bMaximizeCountPerViewer = false, // Enable to maximize count per viewer, requires a valid InLODMaxCountPerViewer parameter during initialization of TMassLODCalculator.
+		bDoVisibilityLogic = false, // Enable to calculate visibility and apply its own LOD distances. Requires a valid InVisibleLODDistance parameter during initialization of TMassLODCalculator.
+		bCalculateLODSignificance = false, // Enable to calculate and set the a more precise LOD floating point significance in member FMassLODResultInfo::LODSignificance.
+		bLocalViewersOnly = false, // Enable to calculate LOD from LocalViewersOnly, otherwise will be done on all viewers.
+		bDoVariableTickRate = false, // Enable to update entity variable tick rate calculation
+	};
+};
+
+struct FMassSimulationLODLogic : public FLODDefaultLogic
+{
+	enum
+	{
+		bDoVariableTickRate = true,
+	};
+};
+
+struct FMassRepresentationLODLogic : public FLODDefaultLogic
+{
+	enum
+	{
+		bDoVisibilityLogic = true,
+		bCalculateLODSignificance = true,
+		bLocalViewersOnly = true,
+	};
+};
+
+struct FMassCombinedLODLogic : public FLODDefaultLogic
+{
+	enum
+	{
+		bDoVariableTickRate = true,
+		bDoVisibilityLogic = true,
+		bCalculateLODSignificance = true,
+		bLocalViewersOnly = true,
+	};
+};
+
+/**
+ * This is the expected member variables for the TArrayView<FMassLODFragment> when calling the TMassLODCalculator and TMassLODtickRateController methods
+ */
+struct FMassLODFragment
+{
+	// Saved closest viewer distance
+	float ClosestViewerDistanceSq;
+
+	// LOD information
+	TEnumAsByte<EMassLOD::Type> LOD;
+	TEnumAsByte<EMassLOD::Type> PrevLOD;
+
+	// Per viewer LOD information (Required only when FLODLogic::bStoreLODPerViewer is enabled)
+	TStaticArray<EMassLOD::Type, UE::MassLOD::MaxNumOfViewers> LODPerViewer;
+	TStaticArray<EMassLOD::Type, UE::MassLOD::MaxNumOfViewers> PrevLODPerViewer;
+
+	// Visibility information (Required only when FLODLogic::bDoVisibilityLogic is enabled)
+	bool bIsVisibleByAViewer;
+	bool bWasVisibleByAViewer;
+	bool bIsInVisibleRange;
+	bool bWasInVisibleRange;
+
+	// Floating point significance (Required only when FLODLogic::bCalculateLODSignificance is enabled)
+	// scaling from 0 to 3, 0 highest LOD and 3 being completely off LOD
+	float LODSignificance = 0.0f;
+
+	// Accumulated DeltaTime (Required only when FLODLogic::bDoVariableTickRate is enalbed)
+	float DeltaTime = 0.0f;
+	float LastTickedTime = 0.0f;
+};
+
+/**
+ * Base struct for the LOD calculation helpers
+ */
+struct MASSLOD_API FMassLODBaseLogic
+{
+protected:
+	void CacheViewerInformation(TConstArrayView<FViewerInfo> Viewers, const bool bLocalViewersOnly);
+
+	DECLARE_CONDITIONAL_MEMBER_ACCESSORS(Condition, bool, bIsVisibleByAViewer);
+	DECLARE_CONDITIONAL_MEMBER_ACCESSORS(Condition, bool, bWasVisibleByAViewer);
+	DECLARE_CONDITIONAL_MEMBER_ACCESSORS(Condition, bool, bIsInVisibleRange);
+	DECLARE_CONDITIONAL_MEMBER_ACCESSORS(Condition, bool, bWasInVisibleRange);
+	DECLARE_CONDITIONAL_MEMBER_ACCESSORS(Condition, float, LODSignificance);
+	DECLARE_CONDITIONAL_MEMBER_ACCESSORS(Condition, float, DeltaTime);
+	DECLARE_CONDITIONAL_MEMBER_ACCESSORS(Condition, float, LastTickedTime);
+	DECLARE_CONDITIONAL_MEMBER_ARRAY_ACCESSORS(Condition, bool, bIsVisibleByViewer);
+	DECLARE_CONDITIONAL_MEMBER_ARRAY_ACCESSORS(Condition, EMassLOD::Type, LODPerViewer);
+	DECLARE_CONDITIONAL_MEMBER_ARRAY_ACCESSORS(Condition, EMassLOD::Type, PrevLODPerViewer);
+
+	/* Number of viewers handled */
+	int32 NumOfViewers = 0;
+
+	/* Boolean indicating the viewer data needs to be cleared */
+	TStaticArray<bool, UE::MassLOD::MaxNumOfViewers> bClearViewerData;
+
+	/** The handle to each viewers */
+	TStaticArray<FMassViewerHandle, UE::MassLOD::MaxNumOfViewers> ViewersHandles;
+
+	/** Viewer location and looking direction */
+	TStaticArray<FVector, UE::MassLOD::MaxNumOfViewers> ViewersLocation;
+	TStaticArray<FVector, UE::MassLOD::MaxNumOfViewers> ViewersDirection;
+};
\ No newline at end of file
diff --git a/Engine/Plugins/Runtime/MassGameplay/Source/MassLOD/Public/MassLODManager.h b/Engine/Plugins/Runtime/MassGameplay/Source/MassLOD/Public/MassLODManager.h
new file mode 100644
--- /dev/null
+++ b/Engine/Plugins/Runtime/MassGameplay/Source/MassLOD/Public/MassLODManager.h
@@ -0,0 +1,154 @@
+// Copyright Epic Games, Inc. All Rights Reserved.
+
+#pragma once
+
+#include "MassCommonTypes.h"
+#include "MassProcessor.h"
+#include "IndexedHandle.h"
+#include "MassLODTypes.h"
+#include "Subsystems/WorldSubsystem.h"
+
+#include "MassLODManager.generated.h"
+
+class UMassLODManager;
+
+/*
+ * Base mass LOD processor to store common information for all LOD processors
+ */
+UCLASS(abstract)
+class MASSLOD_API UMassProcessor_LODBase : public UPipeProcessor
+{
+	GENERATED_BODY()
+
+public:
+	virtual void Initialize(UObject& Owner) override;
+
+protected:
+	UPROPERTY()
+	UWorld* World = nullptr;
+
+	UPROPERTY()
+	UMassLODManager* LODManager = nullptr;
+};
+
+/*
+ * Handle that lets you reference the concept of a viewer
+ */
+USTRUCT()
+struct MASSLOD_API FMassViewerHandle : public FIndexedHandleBase
+{
+	GENERATED_BODY()
+
+	friend class UMassLODManager;
+};
+
+USTRUCT()
+struct FViewerInfo
+{
+	GENERATED_BODY()
+
+	UPROPERTY(transient)
+	APlayerController* PlayerController = nullptr;
+	
+	FName StreamingSourceName;
+
+	FMassViewerHandle Handle;
+	uint32 HashValue = 0;
+
+	FVector Location;
+	FVector Direction;
+
+	bool bEnabled = true;
+
+	void Reset();
+
+	bool IsLocal() const;
+};
+
+DECLARE_MULTICAST_DELEGATE_ThreeParams(FOnViewerAdded, FMassViewerHandle ViewerHandle, APlayerController* PlayerController, FName StreamingSourceName);
+DECLARE_MULTICAST_DELEGATE_ThreeParams(FOnViewerRemoved, FMassViewerHandle ViewerHandle, APlayerController* PlayerController, FName StreamingSourceName);
+
+/*
+ * Manager responsible to manage and synchronized available viewers
+ */
+UCLASS()
+class MASSLOD_API UMassLODManager : public UTickableWorldSubsystem
+{
+	GENERATED_BODY()
+
+public:
+	/** Checks the validity of a viewer handle */
+	bool IsValidViewer(const FMassViewerHandle& ViewerHandle) const { return GetValidViewerIdx(ViewerHandle) != INDEX_NONE; }
+
+	/** Returns the index of the viewer if valid, otherwise INDEX_NONE is return */
+	int32 GetValidViewerIdx(const FMassViewerHandle& ViewerHandle) const;
+
+	/** Returns the array of viewers */
+	const TArray<FViewerInfo>& GetViewers() const { return Viewers; }
+
+	/** Synchronize the viewers if not done this frame and returns the updated array */
+	const TArray<FViewerInfo>& GetSynchronizedViewers();
+
+	/** Returns viewer handle from the PlayerController pointer */
+	FMassViewerHandle GetViewerHandleFromPlayerController(const APlayerController* PlayerController) const;
+
+	/** Returns viewer handle from the streaming source name */
+	FMassViewerHandle GetViewerHandleFromStreamingSource(const FName StreamingSourceName) const;
+
+	/** Returns PlayerController pointer from the viewer handle */
+	APlayerController* GetPlayerControllerFromViewerHandle(const FMassViewerHandle& ViewerHandle) const;
+
+	/** Returns the delegate called when new viewer are added to the list */
+	FOnViewerAdded& GetOnViewerAddedDelegate() { return OnViewerAddedDelegate;  }
+
+	/** Returns the delegate called when viewer are removed from the list */
+	FOnViewerRemoved& GetOnViewerRemovedDelegate() { return OnViewerRemovedDelegate; }
+
+protected:
+	virtual void Initialize(FSubsystemCollectionBase& Collection) override;
+	virtual void Tick(float DeltaTime) override;
+	virtual TStatId GetStatId() const override;
+	virtual void Deinitialize() override;
+
+	/** Synchronizes the viewers from the engine PlayerController list */
+	void SynchronizeViewers();
+
+	/** Adds a viewer to the list and sends notification about addition */
+	void AddViewer(APlayerController* PlayerController, FName StreamingSourceName = NAME_None);
+
+	/** Removes a viewer to the list and send notification about removal */
+	void RemoveViewer(const FMassViewerHandle& ViewerHandle);
+
+	/** Returns the next new viewer serial number */
+	uint32 GetNextViewerSerialNumber() { return ViewerSerialNumberCounter++; }
+
+	/** Player controller EndPlay callback, removing viewers from the list */
+	UFUNCTION()
+	void OnPlayerControllerEndPlay(AActor* Actor, EEndPlayReason::Type EndPlayReason);
+
+private:
+	/** Removes a viewer to the list and send notification about removal */
+	void RemoveViewerInternal(const FMassViewerHandle& ViewerHandle);
+
+	/** The actual array of viewer's information*/
+	UPROPERTY(Transient)
+	TArray<FViewerInfo> Viewers;
+
+	/** The map that do reverse look up to get ViewerHandle */
+	UPROPERTY(Transient)
+	TMap<uint32, FMassViewerHandle> ViewerMap;
+
+	uint64 LastSynchronizedFrame = 0;
+
+	/** Viewer serial number counter */
+	uint32 ViewerSerialNumberCounter = 0;
+
+	/** Free list of indices in the sparse viewer array */
+	TArray<int32> ViewerFreeIndices;
+
+	/** Delegates to notify anyone who needs to know about viewer changes */
+	FOnViewerAdded OnViewerAddedDelegate;
+	FOnViewerRemoved OnViewerRemovedDelegate;
+
+};
+
diff --git a/Engine/Plugins/Runtime/MassGameplay/Source/MassLOD/Public/MassLODTickRateController.h b/Engine/Plugins/Runtime/MassGameplay/Source/MassLOD/Public/MassLODTickRateController.h
new file mode 100644
--- /dev/null
+++ b/Engine/Plugins/Runtime/MassGameplay/Source/MassLOD/Public/MassLODTickRateController.h
@@ -0,0 +1,147 @@
+// Copyright Epic Games, Inc. All Rights Reserved.
+
+#pragma once
+
+#include "MassLODLogic.h"
+#include "MassLODUtils.h"
+#include "LWCCommandBuffer.h"
+
+/**
+ * Helper struct to control LOD tick rate for each agent, 
+ * It will add a fragment tag to group the agent of the same LOD together, that way the user can do tick rate logic per chunk.
+ */
+template <typename FVariableTickChunkFragment, typename FLODLogic = FLODDefaultLogic >
+struct TMassLODTickRateController : public FMassLODBaseLogic
+{
+public:
+
+	/**
+	 * Initializes the LOD trick rate controller, needed to be called once at initialization time (Only when FLODLogic::bDoVariableTickRate is enabled)
+	 * @Param InTickRate the rate at which entities should be ticked per LOD
+	 */
+	void Initialize(const float InTickRate[EMassLOD::Max]);
+
+	/**
+	 * Retrieve if it is needed to calculate the LOD for this chunk
+	 *
+	 * @return if the LOD needs to be calculated
+	 */
+	bool ShouldCalculateLODForChunk(const FLWComponentSystemExecutionContext& Context) const;
+
+	/**
+	 * Retrieve if it is needed to adjust LOD from the newly calculated count for this chunk
+	 * 
+	 * @return if the LOD needs to be adjusted
+	 */
+	bool ShouldAdjustLODFromCountForChunk(const FLWComponentSystemExecutionContext& Context) const;
+
+	/**
+	 * Updates tick rate for this chunk and its entities
+	 * @param Context of the chunk execution 
+	 * @param LODList is the fragment where calculation are stored 
+	 * @param Time of the simulation to use for this update
+	 * @return bool return if the chunk should be tick this frame
+	 */
+	template <typename FMassLODResultInfo>
+	bool UpdateTickRateFromLOD(FLWComponentSystemExecutionContext& Context, TArrayView<FMassLODResultInfo>& LODList, const float Time);
+
+protected:
+
+	/** Tick rate for each LOD */
+	TStaticArray<int32, EMassLOD::Max> TickRates;
+};
+
+template <typename FVariableTickChunkFragment, typename FLODLogic>
+void TMassLODTickRateController<FVariableTickChunkFragment,FLODLogic>::Initialize(const float InTickRates[EMassLOD::Max])
+{
+	checkf(InTickRates, TEXT("You need to provide tick rate values to use this class."));
+	checkf(FLODLogic::bDoVariableTickRate, TEXT("You need to enalbe bDoVariableTickRate to use this class."));
+
+	// Make a copy of all the settings
+	for (int x = 0; x < EMassLOD::Max; x++)
+	{
+		TickRates[x] = InTickRates[x];
+	}
+}
+
+template <typename FVariableTickChunkFragment, typename FLODLogic>
+bool TMassLODTickRateController<FVariableTickChunkFragment, FLODLogic>::ShouldCalculateLODForChunk(const FLWComponentSystemExecutionContext& Context) const
+{
+	// EMassLOD::Off does not need to handle max count, so we can use ticking rate for them if available
+	const FMassVariableTickChunkFragment& ChunkData = Context.GetChunkComponent<FVariableTickChunkFragment>();
+	return ChunkData.GetLOD() != EMassLOD::Off || ChunkData.ShouldTickThisFrame();
+}
+
+template <typename FVariableTickChunkFragment, typename FLODLogic>
+bool TMassLODTickRateController<FVariableTickChunkFragment, FLODLogic>::ShouldAdjustLODFromCountForChunk(const FLWComponentSystemExecutionContext& Context) const
+{
+	// EMassLOD::Off does not need to handle max count, so we can skip it
+	const FMassVariableTickChunkFragment& ChunkData = Context.GetChunkComponent<FVariableTickChunkFragment>();
+	return ChunkData.GetLOD() != EMassLOD::Off;
+}
+
+template <typename FVariableTickChunkFragment, typename FLODLogic>
+template <typename FMassLODResultInfo>
+bool TMassLODTickRateController<FVariableTickChunkFragment, FLODLogic>::UpdateTickRateFromLOD(FLWComponentSystemExecutionContext& Context, TArrayView<FMassLODResultInfo>& LODList, const float Time)
+{
+
+	bool bShouldTickThisFrame = true;
+	bool bWasChunkTicked = true;
+	const float DeltaTime = Context.GetDeltaTimeSeconds();
+
+	FMassVariableTickChunkFragment& ChunkData = Context.GetMutableChunkComponent<FVariableTickChunkFragment>();
+	EMassLOD::Type ChunkLOD = ChunkData.GetLOD();
+	if (ChunkLOD == EMassLOD::Max)
+	{
+		// The LOD on the chunk component data isn't set yet, let see if the Archetype has an LOD tag and set it on the ChunkData
+		ChunkLOD = UE::MassLOD::GetLODFromArchetype(Context);
+		ChunkData.SetLOD(ChunkLOD);
+	}
+	else
+	{
+		checkfSlow(UE::MassLOD::IsLODTagSet(Context, ChunkLOD), TEXT("Expecting the same LOD as what we saved in the chunk data, maybe external code is modifying the tags"))
+	}
+
+	if (ChunkLOD != EMassLOD::Max)
+	{
+		float TimeUntilNextTick = ChunkData.GetTimeUntilNextTick();
+		bWasChunkTicked = ChunkData.ShouldTickThisFrame();
+
+		// Reset DeltaTime if we ticked last frame
+		if (bWasChunkTicked)
+		{
+			const float TickRate = TickRates[ChunkLOD];
+			TimeUntilNextTick = TickRate * (1.0f + FMath::RandRange(-0.1f, 0.1f));
+		}
+		else
+		{
+			// Decrement delta time
+			TimeUntilNextTick -= DeltaTime;
+		}
+
+		// Should we tick this frame?
+		const int32 ChunkSerialModificationNumber = Context.GetChunkSerialModificationNumber();
+		bShouldTickThisFrame = ChunkData.GetLastChunkSerialModificationNumber() != ChunkSerialModificationNumber || TimeUntilNextTick <= 0.0f;
+		ChunkData.Update(bShouldTickThisFrame, TimeUntilNextTick, ChunkSerialModificationNumber);
+	}
+
+	if (bWasChunkTicked)
+	{
+		const int32 NumEntities = Context.GetEntitiesNum();
+		for (int32 Index = 0; Index < NumEntities; ++Index)
+		{
+			FMassLODResultInfo& EntityLOD = LODList[Index];
+			const float LastTickedTime = GetLastTickedTime<FLODLogic::bDoVariableTickRate>(EntityLOD, 0.0f);
+			const float NewDeltaTime = LastTickedTime != 0.0f ? Time - LastTickedTime : DeltaTime;
+			SetDeltaTime<FLODLogic::bDoVariableTickRate>(EntityLOD, NewDeltaTime);
+			SetLastTickedTime<FLODLogic::bDoVariableTickRate>(EntityLOD, Time);
+			if (EntityLOD.LOD != ChunkLOD)
+			{
+				FLWEntity Entity = Context.GetEntity(Index);
+				Context.Defer().PushCommand(FCommandSwapTags(Entity, UE::MassLOD::GetLODTagFromLOD(ChunkLOD), UE::MassLOD::GetLODTagFromLOD(EntityLOD.LOD)));
+			}
+		}
+	}
+
+	return bShouldTickThisFrame;
+}
diff --git a/Engine/Plugins/Runtime/MassGameplay/Source/MassLOD/Public/MassLODTypes.h b/Engine/Plugins/Runtime/MassGameplay/Source/MassLOD/Public/MassLODTypes.h
new file mode 100644
--- /dev/null
+++ b/Engine/Plugins/Runtime/MassGameplay/Source/MassLOD/Public/MassLODTypes.h
@@ -0,0 +1,251 @@
+// Copyright Epic Games, Inc. All Rights Reserved.
+
+#pragma once
+
+#include "LWComponentTypes.h"
+#include "Containers/StaticArray.h"
+#include "EntitySubsystem.h"
+
+#include "MassLODTypes.generated.h"
+
+/** Debug option to allow multiple viewers per controller. Useful for testing and profiling purposes */
+#define UE_DEBUG_REPLICATION_DUPLICATE_VIEWERS_PER_CONTROLLER 0
+
+#define UE_ALLOW_DEBUG_REPLICATION_DUPLICATE_VIEWERS_PER_CONTROLLER (UE_DEBUG_REPLICATION_DUPLICATE_VIEWERS_PER_CONTROLLER && !UE_BUILD_SHIPPING)
+
+namespace UE::MassLOD
+{
+#if UE_ALLOW_DEBUG_REPLICATION_DUPLICATE_VIEWERS_PER_CONTROLLER
+	constexpr int32 DebugNumberViewersPerController = 50;
+	constexpr int32 MaxNumOfViewers = 50;
+#else
+	constexpr int32 MaxNumOfViewers = 16;
+#endif // UE_ALLOW_DEBUG_REPLICATION_DUPLICATE_VIEWERS_PER_CONTROLLER
+
+	constexpr int32 MaxBucketsPerLOD = 250;
+
+	extern MASSLOD_API FColor LODColors[];
+} // UE::MassLOD
+
+namespace UE::Mass::ProcessorGroupNames
+{
+	const FName LODCollector = FName(TEXT("LODCollector"));
+	const FName LOD = FName(TEXT("LOD"));
+}
+
+// We are not using enum class here because we are doing so many arithmetic operation and comparison on them 
+// that it is not worth polluting int32 casts everywhere in the code.
+UENUM()
+namespace EMassLOD 
+{
+	enum Type
+	{
+		High,
+		Medium,
+		Low,
+		Off,
+		Max
+	};
+}
+
+USTRUCT()
+struct MASSLOD_API FMassHighLODTag : public FComponentTag
+{
+	GENERATED_BODY()
+};
+
+USTRUCT()
+struct MASSLOD_API FMassMediumLODTag : public FComponentTag
+{
+	GENERATED_BODY()
+};
+
+USTRUCT()
+struct MASSLOD_API FMassLowLODTag : public FComponentTag
+{
+	GENERATED_BODY()
+};
+
+USTRUCT()
+struct MASSLOD_API FMassOffLODTag : public FComponentTag
+{
+	GENERATED_BODY()
+};
+
+/*
+ * Data fragment to store the calculated distances to viewers
+ */
+USTRUCT()
+struct MASSLOD_API FMassLODInfoFragment : public FLWComponentData
+{
+	GENERATED_BODY()
+
+	// Distance between each viewer and entity
+	TStaticArray<float, UE::MassLOD::MaxNumOfViewers> DistanceToViewerSq;
+
+	// Whether the entity is visible to each viewer
+	TStaticArray<bool, UE::MassLOD::MaxNumOfViewers> bIsVisibleByViewer;
+};
+
+USTRUCT()
+struct MASSLOD_API FMassVariableTickChunkFragment : public FLWChunkComponent
+{
+	GENERATED_BODY();
+
+	bool ShouldTickThisFrame() const
+	{
+		return bShouldTickThisFrame;
+	}
+
+	float GetTimeUntilNextTick() const
+	{
+		return TimeUntilNextTick;
+	}
+
+	int32 GetLastChunkSerialModificationNumber() const
+	{
+		return LastChunkSerialModificationNumber;
+	}
+
+	EMassLOD::Type GetLOD() const
+	{
+		return LOD;
+	}
+
+	void SetLOD(EMassLOD::Type InLOD)
+	{
+		checkf(LOD == EMassLOD::Max, TEXT("Chunk LOD should never change, it is allowed to only set it once"))
+		LOD = InLOD;
+	}
+
+	void Update(const bool bInShouldTickThisFrame, const float InTimeUntilNextTick, int32 ChunkSerialModificationNumber)
+	{
+		bShouldTickThisFrame = bInShouldTickThisFrame;
+		TimeUntilNextTick = InTimeUntilNextTick;
+		LastChunkSerialModificationNumber = ChunkSerialModificationNumber;
+	}
+
+private:
+	bool bShouldTickThisFrame = true;
+	EMassLOD::Type LOD = EMassLOD::Max;
+	float TimeUntilNextTick = 0.0f;
+	int32 LastChunkSerialModificationNumber = -1;
+};
+
+enum class EMassVisibility : uint8
+{
+	CanBeSeen, // Not too far and within camera frustum
+	CulledByFrustum, // Not in camera frustum but within visibility distance
+	CulledByDistance, // Too far whether in or out of frustum
+	Max
+};
+
+USTRUCT()
+struct MASSLOD_API FMassVisibilityCanBeSeenTag : public FComponentTag
+{
+	GENERATED_BODY()
+};
+
+USTRUCT()
+struct MASSLOD_API FMassVisibilityCulledByFrustumTag : public FComponentTag
+{
+	GENERATED_BODY()
+};
+
+USTRUCT()
+struct MASSLOD_API FMassVisibilityCulledByDistanceTag : public FComponentTag
+{
+	GENERATED_BODY()
+};
+
+USTRUCT()
+struct MASSLOD_API FMassVisualizationChunkFragment : public FLWChunkComponent
+{
+	GENERATED_BODY();
+
+	/**
+	 * Fetched existing chunk fragment to know if there is a possibility of a an entity that is visible
+	 * In the case that there is not chunk information, we cannot assume that all entities are not visible.
+	 *
+	 * @param Context of the execution from the entity sub system
+	 * @return true if there is a possibility that the chunk contains a visible entity
+	 */
+	static bool AreAnyEntitiesVisibleInChunk(const FLWComponentSystemExecutionContext& Context)
+	{
+		return Context.GetChunkComponent<FMassVisualizationChunkFragment>().AreAnyEntitiesVisible();
+	}
+
+	/**
+	 * Returns if there could be a visible entities in that chunk
+	 *
+	 * @return true if that is the case
+	 */
+	bool AreAnyEntitiesVisible() const
+	{
+		return (Visibility != EMassVisibility::CulledByDistance && Visibility != EMassVisibility::CulledByFrustum) || bContainsNewlyVisibleEntity;
+	}
+
+	/**
+	 * ShouldUpdateVisualizationForChunk
+	 *
+	* @param Context of the execution from the entity sub system
+	 * @return bool true if we should update the representation type for this chunk
+	 */
+	static bool ShouldUpdateVisualizationForChunk(const FLWComponentSystemExecutionContext& Context)
+	{
+		return Context.GetChunkComponent<FMassVisualizationChunkFragment>().ShouldUpdateVisualization();
+	}
+
+	/**
+	 * Representation type of all currently visible entities are always updated
+	 * But as an optimization, we use a frequency check on the not visible one.
+	 *
+	 * @return true if we should update the representation type for this chunk
+	 */
+	bool ShouldUpdateVisualization() const
+	{
+		return Visibility != EMassVisibility::CulledByDistance || DeltaTime <= 0.0f;
+	}
+
+	void SetContainsNewlyVisibleEntity(bool bInContainsNewlyVisibleEntity)
+	{
+		if (bInContainsNewlyVisibleEntity)
+		{
+			checkfSlow(Visibility != EMassVisibility::CanBeSeen, TEXT("Something is not adding up, how can an entity be newly visible in a can be seen chunk?"));
+			bContainsNewlyVisibleEntity = true;
+		}
+	}
+
+	void SetVisibility(const EMassVisibility InVisibility)
+	{
+		checkf(Visibility == EMassVisibility::Max, TEXT("Chunk visibility should never change, it is allowed to only set it once"));
+		Visibility = InVisibility;
+	}
+
+	EMassVisibility GetVisibility() const
+	{
+		return Visibility;
+	}
+
+	float GetDeltaTime() const
+	{
+		return DeltaTime;
+	}
+
+	void Update(float InDeltaTime)
+	{
+		bContainsNewlyVisibleEntity = false;
+		DeltaTime = InDeltaTime;
+	}
+
+protected:
+
+	/** Visibility of the current chunk, should never change */
+	EMassVisibility Visibility = EMassVisibility::Max;
+
+	/** Not visible chunks, might contains entity that are newly visible and not yet moved. */
+	bool bContainsNewlyVisibleEntity = true;
+
+	/** Not visible chunks delta time until next update */
+	float DeltaTime = 0;
+};
diff --git a/Engine/Plugins/Runtime/MassGameplay/Source/MassLOD/Public/MassRepresentationLODCollectorProcessor.h b/Engine/Plugins/Runtime/MassGameplay/Source/MassLOD/Public/MassRepresentationLODCollectorProcessor.h
new file mode 100644
--- /dev/null
+++ b/Engine/Plugins/Runtime/MassGameplay/Source/MassLOD/Public/MassRepresentationLODCollectorProcessor.h
@@ -0,0 +1,52 @@
+// Copyright Epic Games, Inc. All Rights Reserved.
+
+#pragma once
+
+#include "MassLODManager.h"
+#include "MassLODCollector.h"
+#include "MassLODCalculator.h"
+#include "LWComponentTypes.h"
+
+#include "MassRepresentationLODCollectorProcessor.generated.h"
+
+
+USTRUCT()
+struct FMassRepresentationLODCollectorConfig
+{
+	GENERATED_BODY()
+
+	/** Component Tag that will be used to associate LOD config */
+	UPROPERTY(EditAnywhere, Category = "Mass|LOD", config, meta = (BaseStruct = "ComponentTag"))
+	FInstancedStruct TagFilter;
+
+	/** Set FOV Angles to define the limits between using Visible or Base LOD Distances */
+	UPROPERTY(EditAnywhere, Category = "Mass|LOD", meta = (ClampMin = "0.0", ClampMax = "180.0", UIMin = "0.0", UIMax = "180.0"), config)
+	float FOVAnglesToDriveVisibility = 45.0f;
+	UPROPERTY(EditAnywhere, Category = "Mass|LOD", meta = (ClampMin = "0.0", UIMin = "0.0"), config)
+	float BufferHysteresisOnFOVPercentage = 10.0f;
+
+	/** Runtime data for matching the LOD config */
+	TMassLODCollector<FMassRepresentationLODLogic> RepresentationLODCollector;
+	FLWComponentQuery CloseEntityQuery;
+	FLWComponentQuery FarEntityQuery_Conditional;
+};
+
+/*
+* LOD collector which collects LOD information for Viewer LODing when possible.
+*/
+UCLASS(meta = (DisplayName = "LOD Collector"))
+class MASSLOD_API UMassRepresentationLODCollectorProcessor : public UMassProcessor_LODBase
+{
+	GENERATED_BODY()
+
+	public:
+	UMassRepresentationLODCollectorProcessor();
+
+	protected:
+	virtual void Initialize(UObject& Owner) override;
+	virtual void ConfigureQueries() override;
+	virtual void Execute(UEntitySubsystem& EntitySubsystem, FLWComponentSystemExecutionContext& Context) override;
+
+	UPROPERTY(EditAnywhere, Category = "Mass|LOD", config)
+	TArray<FMassRepresentationLODCollectorConfig> LODConfigs;
+};
diff --git a/Engine/Plugins/Runtime/MassGameplay/Source/MassLOD/Public/MassSimulationLOD.h b/Engine/Plugins/Runtime/MassGameplay/Source/MassLOD/Public/MassSimulationLOD.h
new file mode 100644
--- /dev/null
+++ b/Engine/Plugins/Runtime/MassGameplay/Source/MassLOD/Public/MassSimulationLOD.h
@@ -0,0 +1,141 @@
+// Copyright Epic Games, Inc. All Rights Reserved.
+
+#pragma once
+
+#include "MassLODManager.h"
+#include "MassLODCollector.h"
+#include "MassLODCalculator.h"
+#include "MassLODTickRateController.h"
+#include "MassLODLogic.h"
+#include "LWComponentTypes.h"
+
+#include "MassSimulationLOD.generated.h"
+
+/*
+ * Data fragment to store the calculated distances to viewers
+ * @todo To be removed and convert to new Simulation LOD stuff
+ */
+USTRUCT()
+struct MASSLOD_API FDataFragment_MassSimulationLODInfo : public FLWComponentData
+{
+	GENERATED_BODY()
+
+	// Distance between each viewer and entity
+	TStaticArray<float, UE::MassLOD::MaxNumOfViewers> DistanceToViewerSq;
+};
+
+struct FSimulationLODLogic : public FLODDefaultLogic
+{
+};
+
+/*
+ * Mass processor to calculate square distances from entity to viewers.
+ * This is done once per entity and will be used for all subsequent LOD calculation.
+ * @todo To be removed and convert to new Simulation LOD stuff
+ */
+UCLASS()
+class MASSLOD_API UMassProcessor_MassSimulationLODViewersInfo : public UMassProcessor_LODBase
+{
+	GENERATED_BODY()
+public:
+	UMassProcessor_MassSimulationLODViewersInfo();
+
+protected:
+	virtual void Initialize(UObject& Owner) override;
+	virtual void ConfigureQueries() override;
+	virtual void Execute(UEntitySubsystem& EntitySubsystem, FLWComponentSystemExecutionContext& Context) override;
+
+	TMassLODCollector<FSimulationLODLogic> LODCollector;
+	FLWComponentQuery EntityQuery;
+};
+
+USTRUCT()
+struct MASSLOD_API FMassSimulationLODFragment : public FLWComponentData
+{
+	GENERATED_BODY()
+
+	/** Saved closest ViewerDistance */
+	float ClosestViewerDistanceSq = FLT_MAX;
+
+	/**LOD information */
+	TEnumAsByte<EMassLOD::Type> LOD = EMassLOD::Max;
+
+	/** Previous LOD information*/
+	TEnumAsByte<EMassLOD::Type> PrevLOD = EMassLOD::Max;
+
+	/** Accumulated delta time to use upon next tick */
+	float LastTickedTime = 0.0f;
+	float DeltaTime = 0.0f;
+};
+
+USTRUCT()
+struct FMassSimulationVariableTickChunkFragment : public FMassVariableTickChunkFragment
+{
+	GENERATED_BODY();
+
+	static bool ShouldTickChunkThisFrame(const FLWComponentSystemExecutionContext& Context)
+	{
+		const FMassSimulationVariableTickChunkFragment& ChunkFragment = Context.GetChunkComponent<FMassSimulationVariableTickChunkFragment>();
+		return ChunkFragment.ShouldTickThisFrame();
+	}
+
+	static EMassLOD::Type GetChunkLOD(const FLWComponentSystemExecutionContext& Context)
+	{
+		const FMassSimulationVariableTickChunkFragment& ChunkFragment = Context.GetChunkComponent<FMassSimulationVariableTickChunkFragment>();
+		return ChunkFragment.GetLOD();
+	}
+
+};
+
+USTRUCT()
+struct FMassSimulationLODConfig
+{
+	GENERATED_BODY()
+
+	FMassSimulationLODConfig();
+
+	/** Component Tag that will be used to associate LOD config */
+	UPROPERTY(EditAnywhere, Category = "Mass|LOD", config, meta = (BaseStruct = "ComponentTag"))
+	FInstancedStruct TagFilter;
+
+	/** Distance where each LOD becomes relevant */
+	UPROPERTY(EditAnywhere, Category = "Mass|LOD", config)
+	float LODDistance[EMassLOD::Max];
+
+	/** Hysteresis percentage on delta between the LOD distances */
+	UPROPERTY(EditAnywhere, Category = "Mass|LOD", meta = (ClampMin = "0.0", UIMin = "0.0"), config)
+	float BufferHysteresisOnDistancePercentage = 10.0f;
+
+	/** Maximum limit of entity per LOD */
+	UPROPERTY(EditAnywhere, Category = "Mass|LOD", config)
+	int32 LODMaxCount[EMassLOD::Max];
+
+	/** Rate in seconds at which entities should update when in this LOD */
+	UPROPERTY(EditAnywhere, Category = "Mass|LOD", config)
+	float TickRates[EMassLOD::Max];
+
+	/** Runtime data for matching the LOD config */
+	FLWComponentQuery EntityQuery;
+	TMassLODCalculator<FMassSimulationLODLogic> LODCalculator;
+	TMassLODTickRateController<FMassSimulationVariableTickChunkFragment, FMassSimulationLODLogic> LODTickRateController;
+};
+
+UCLASS(meta = (DisplayName = "Simulation LOD"))
+class MASSLOD_API UMassSimulationLODProcessor : public UMassProcessor_LODBase
+{
+	GENERATED_BODY()
+
+public:
+	UMassSimulationLODProcessor();
+
+protected:
+
+	virtual void ConfigureQueries() override;
+	virtual void Initialize(UObject& InOwner) override;
+	virtual void Execute(UEntitySubsystem& EntitySubsystem, FLWComponentSystemExecutionContext& Context) override;
+
+	void CalculateLODForConfig(UEntitySubsystem& EntitySubsystem, FLWComponentSystemExecutionContext& Context, FMassSimulationLODConfig& LODConfig);
+
+	UPROPERTY(EditAnywhere, Category = "Mass|LOD", config)
+	TArray<FMassSimulationLODConfig> LODConfigs;
+};
diff --git a/Engine/Plugins/Runtime/MassGameplay/Source/MassMovement/Private/MassMovementFragments.cpp b/Engine/Plugins/Runtime/MassGameplay/Source/MassMovement/Private/MassMovementFragments.cpp
new file mode 100644
--- /dev/null
+++ b/Engine/Plugins/Runtime/MassGameplay/Source/MassMovement/Private/MassMovementFragments.cpp
@@ -0,0 +1,36 @@
+// Copyright Epic Games, Inc. All Rights Reserved.
+
+#include "MassMovementFragments.h"
+#include "GameFramework/GameStateBase.h"
+#include "Engine/World.h"
+
+//void FMassMoveTargetFragment::CreateNewAction(const EMassMovementAction InAction, const UWorld& InWorld)
+//{
+//	ensureMsgf(InWorld.GetNetMode() != NM_Client, TEXT("This version of SetDesiredAction should only be called on the authority."));
+//	CurrentAction = InAction;
+//	CurrentActionID++;
+//	CurrentActionWorldStartTime = InWorld.TimeSeconds;
+//
+//	const AGameStateBase* GameState = InWorld.GetGameState();
+//	CurrentActionServerStartTime = GameState != nullptr ? GameState->GetServerWorldTimeSeconds() : CurrentActionWorldStartTime;
+//
+//	MarkNetDirty();
+//}
+//
+//void FMassMoveTargetFragment::CreateReplicatedAction(const EMassMovementAction InAction, const uint16 InActionID, const float InWorldStartTime, const float InServerStartTime)
+//{
+//	CurrentAction = InAction;
+//	CurrentActionID = InActionID;
+//	CurrentActionWorldStartTime = InWorldStartTime;
+//	CurrentActionServerStartTime = InServerStartTime;
+//}
+//
+//FString FMassMoveTargetFragment::ToString() const
+//{
+//	return FString::Printf(TEXT("ActionID:%d Action:%s StartTime: World:%.1f Server:%.1f DesiredSpeed:%.1f "),
+//		CurrentActionID,
+//		*UEnum::GetValueAsString(CurrentAction),
+//		CurrentActionWorldStartTime,
+//		CurrentActionServerStartTime,
+//		DesiredSpeed.Get());
+//}
\ No newline at end of file
diff --git a/Engine/Plugins/Runtime/MassGameplay/Source/MassMovement/Private/MassMovementModule.cpp b/Engine/Plugins/Runtime/MassGameplay/Source/MassMovement/Private/MassMovementModule.cpp
new file mode 100644
--- /dev/null
+++ b/Engine/Plugins/Runtime/MassGameplay/Source/MassMovement/Private/MassMovementModule.cpp
@@ -0,0 +1,42 @@
+// Copyright Epic Games, Inc. All Rights Reserved.
+
+#include "CoreMinimal.h"
+#include "IMassMovementModule.h"
+#include "UObject/CoreRedirects.h"
+
+
+class FMassMovementModule : public IMassMovementModule
+{
+	/** IModuleInterface implementation */
+	virtual void StartupModule() override;
+	virtual void ShutdownModule() override;
+};
+
+IMPLEMENT_MODULE(FMassMovementModule, MassMovement)
+
+
+
+void FMassMovementModule::StartupModule()
+{
+	// This code will execute after your module is loaded into memory (but after global variables are initialized, of course.)
+
+	TArray<FCoreRedirect> Redirects;
+	Redirects.Emplace(ECoreRedirectFlags::Type_Struct, TEXT("VelocityComponent"), TEXT("MassVelocityFragment"));
+	Redirects.Emplace(ECoreRedirectFlags::Type_Struct, TEXT("AvoidanceComponent"), TEXT("MassAvoidanceFragment"));
+	Redirects.Emplace(ECoreRedirectFlags::Type_Struct, TEXT("DataFragment_GridCellLocation"), TEXT("MassGridCellLocationFragment"));
+	Redirects.Emplace(ECoreRedirectFlags::Type_Struct, TEXT("DataFragment_PathFollow"), TEXT("MassPathFollowFragment"));
+	Redirects.Emplace(ECoreRedirectFlags::Type_Struct, TEXT("DataFragment_Wander"), TEXT("MassWanderFragment"));
+	Redirects.Emplace(ECoreRedirectFlags::Type_Struct, TEXT("DataFragment_MoveTo"), TEXT("MassMoveToFragment"));
+	
+	FCoreRedirects::AddRedirectList(Redirects, TEXT("MassMovement"));
+}
+
+
+void FMassMovementModule::ShutdownModule()
+{
+	// This function may be called during shutdown to clean up your module.  For modules that support dynamic reloading,
+	// we call this function before unloading the module.
+}
+
+
+
diff --git a/Engine/Plugins/Runtime/MassGameplay/Source/MassMovement/Public/MassMovementTypes.h b/Engine/Plugins/Runtime/MassGameplay/Source/MassMovement/Public/MassMovementTypes.h
new file mode 100644
--- /dev/null
+++ b/Engine/Plugins/Runtime/MassGameplay/Source/MassMovement/Public/MassMovementTypes.h
@@ -0,0 +1,32 @@
+// Copyright Epic Games, Inc. All Rights Reserved.
+
+#pragma once
+
+#include "LWComponentTypes.h"
+#include "HierarchicalHashGrid2D.h"
+
+namespace UE::Mass::ProcessorGroupNames
+{
+	const FName Avoidance = FName(TEXT("Avoidance"));
+}
+
+UENUM()
+enum class EMassAvoidanceObstacleItemFlags : uint8
+{
+	None			= 0,
+	HasColliderData = 1 << 0,
+};
+ENUM_CLASS_FLAGS(EMassAvoidanceObstacleItemFlags)
+
+struct FMassAvoidanceObstacleItem
+{
+	bool operator==(const FMassAvoidanceObstacleItem& Other) const
+	{
+		return Entity == Other.Entity;
+	}
+
+	FLWEntity Entity;
+	EMassAvoidanceObstacleItemFlags ItemFlags = EMassAvoidanceObstacleItemFlags::None;
+};
+
+typedef THierarchicalHashGrid2D<2, 4, FMassAvoidanceObstacleItem> FAvoidanceObstacleHashGrid2D;	// 2 levels of hierarchy, 4 ratio between levels
\ No newline at end of file
