AggregateTicking and Pipe merged into Engine/Plugins/Runtime/MassEntity

diff --git a/Engine/Plugins/Runtime/MassGameplay/Source/MassReplication/Private/MassReplicationManager.cpp b/Engine/Plugins/Runtime/MassGameplay/Source/MassReplication/Private/MassReplicationManager.cpp
new file mode 100644
--- /dev/null
+++ b/Engine/Plugins/Runtime/MassGameplay/Source/MassReplication/Private/MassReplicationManager.cpp
@@ -0,0 +1,697 @@
+// Copyright Epic Games, Inc. All Rights Reserved.
+
+#include "MassReplicationManager.h"
+#include "Engine/World.h"
+#include "Engine/ChildConnection.h"
+#include "GameFramework/GameModeBase.h"
+#include "MassCommonTypes.h"
+#include "MassEntitySystem.h"
+#include "MassClientBubbleHandler.h"
+#include "MassClientBubbleInfoBase.h"
+
+
+uint32 UMassReplicationManager::CurrentNetMassCounter = 0;
+
+void UMassReplicationManager::Initialize(FSubsystemCollectionBase& Collection)
+{
+	Super::Initialize(Collection);
+
+	World = GetWorld();
+
+	check(World);
+
+	MassLODManager = Collection.InitializeDependency<UMassLODManager>();
+
+	check(MassLODManager);
+
+	EntitySystem = UPipeEntitySubsystem::GetCurrent(GetWorld());
+}
+
+void UMassReplicationManager::Deinitialize()
+{
+	// remove all Clients
+	for (const FMassClientHandle& Handle : ClientHandleManager.GetHandles())
+	{
+		if (Handle.IsValid())
+		{
+			RemoveClient(Handle);
+		}
+	}
+
+	// make sure all the other data is reset
+	ClientHandleManager.Reset();
+	BubbleInfoArray.Reset();
+	ClientsViewerHandles.Reset();
+	ClientToViewerHandleArray.Reset();
+	ViewerToClientHandleArray.Reset();
+
+	World = nullptr;
+	MassLODManager = nullptr;
+	EntitySystem = nullptr;
+}
+
+FMassNetworkID UMassReplicationManager::GetNetIDFromHandle(const FMassHandle& Handle) const
+{
+	check(EntitySystem);
+	const FMassNetworkIDFragment& Data = EntitySystem->GetComponentDataChecked<FMassNetworkIDFragment>(Handle.GetLWEntity());
+	return Data.NetID;
+}
+
+namespace UE { namespace Mass { namespace Replication {
+// Checks the parent net connection has APlayerController OwningActor, unfortunately GetParentConnection() prevents UChildConnection being const here.
+APlayerController* GetValidParentNetConnection(UChildConnection& ChildConnection)
+{
+	return (ChildConnection.GetParentConnection() != nullptr) ? Cast<APlayerController>(ChildConnection.GetParentConnection()->OwningActor) : nullptr;
+}
+
+enum class ENetConnectionType : uint8
+{
+	None,
+	Parent,
+	Child,
+};
+
+// Checks Controllers net connection is a parent net connection
+bool HasParentNetConnection(const APlayerController* Controller)
+{
+	return (Controller != nullptr) && Controller->NetConnection && (Cast<UChildConnection>(Controller->NetConnection) == nullptr);
+}
+
+// If Controller has a UChildConnection then this gets the parent net connection's APlayerController OwningActor otherwise nullptr
+APlayerController* GetParentControllerFromChildNetConnection(const APlayerController* Controller)
+{
+	//ChildConnection can't be const
+	UChildConnection* ChildConnection = Controller ? Cast<UChildConnection>(Controller->NetConnection) : nullptr;
+
+	return ChildConnection ? GetValidParentNetConnection(*ChildConnection) : nullptr;
+}
+
+// Gets the ENetConnectionType for Controller. Tests if Controller's net connection is a parent or if Controller's net connection is a child and its parent net connection has a APlayerController OwningActor, 
+ENetConnectionType HasParentOrChildWithValidParentNetConnection(const APlayerController* Controller)
+{
+	ENetConnectionType ConnectionType = ENetConnectionType::None;
+
+	if (Controller && Controller->NetConnection)
+	{
+		UChildConnection* ChildConnection = Cast<UChildConnection>(Controller->NetConnection);
+
+		// if the NetConnection is either a parent connection or a child connection with APlayerController owning actor its valid 
+		if (ChildConnection == nullptr)
+		{
+			ConnectionType = ENetConnectionType::Parent;
+		}
+		else
+		{
+			ConnectionType = (GetValidParentNetConnection(*ChildConnection) != nullptr) ? ENetConnectionType::Child : ENetConnectionType::None;
+		}		
+	}
+
+	return ConnectionType;
+}
+}}};
+
+bool UMassReplicationManager::SynchronizeClients(const TArray<FViewerInfo>& Viewers)
+{
+	typedef TMap<FMassViewerHandle, FMassClientHandle> FControllerMap;
+	struct FClientAddData
+	{
+		FClientAddData(FMassViewerHandle InHandle, APlayerController* InController)
+			: Handle(InHandle)
+			, Controller(InController)
+		{}
+
+		FMassViewerHandle Handle;
+		APlayerController* Controller = nullptr;
+	};
+
+	bool bNeedShrinking = false;
+
+	//we can only replicate if we have some BubbleInfos to use for replication
+	if (BubbleInfoArray.Num() == 0)
+	{
+		return bNeedShrinking;
+	}
+
+	// Arbitrarily use index 0, if there are more items the corresponding InfoData, Bubbles will be the same length and corresponding indices between the Bubbles
+	// will have the same player controller owner.
+	const FMassClientBubbleInfoData& InfoData = BubbleInfoArray[0];
+
+	check(MassLODManager);
+
+	FControllerMap ClientConnectionMap;
+
+	TArray<FClientAddData> ClientsToAdd;
+
+	{
+		// Go through the stored Clients, add valid ones to the ClientConnectionMap and remove invalids.
+		// Note this is only valid until we next add a client handle.
+		const TArray<FMassClientHandle>& ClientHandles = ClientHandleManager.GetHandles();
+
+		for (int32 Idx = 0; Idx < ClientHandles.Num(); ++Idx)
+		{
+			const FMassClientHandle& ClientHandle = ClientHandles[Idx];
+
+			if (ClientHandle.IsValid())
+			{
+				check(ClientToViewerHandleArray.IsValidIndex(ClientHandle.GetIndex()));
+
+				const FViewerClientPair ViewerClientPair = ClientToViewerHandleArray[ClientHandle.GetIndex()];
+
+				check(ViewerClientPair.ViewerHandle.IsValid());
+				check(ViewerClientPair.ClientHandle == ClientHandle);
+				check(InfoData.Bubbles.IsValidIndex(ClientHandle.GetIndex()));
+
+				const AMassClientBubbleInfoBase* ClientBubble = InfoData.Bubbles[ClientHandle.GetIndex()];
+				const APlayerController* Controller = ClientBubble ? Cast<APlayerController>(ClientBubble->GetOwner()) : nullptr;
+
+//no need to check netconnection if UE_ALLOW_DEBUG_REPLICATION_BUBBLES_STANDALONE
+#if !UE_ALLOW_DEBUG_REPLICATION_BUBBLES_STANDALONE
+				// must have a valid non child UNetConnection
+				if (UE::Mass::Replication::HasParentNetConnection(Controller))
+#endif
+				{
+					ClientConnectionMap.Add(ViewerClientPair.ViewerHandle, ClientHandle);
+				}
+#if !UE_ALLOW_DEBUG_REPLICATION_BUBBLES_STANDALONE
+				else
+				{
+					// this is safe as Clients and their handles are free list arrays
+					RemoveClient(ClientHandle);
+					bNeedShrinking |= Idx == ClientHandles.Num() - 1;
+				}
+#endif
+			}
+		}
+
+		// now go through all current Viewers and add if they do not exist
+		for (const FViewerInfo& Viewer : Viewers)
+		{
+			// must have valid Viewer and parent net connection
+			if (Viewer.Handle.IsValid())
+			{
+//no need to check netconnection if UE_ALLOW_DEBUG_REPLICATION_BUBBLES_STANDALONE
+#if !UE_ALLOW_DEBUG_REPLICATION_BUBBLES_STANDALONE
+				if (UE::Mass::Replication::HasParentNetConnection(Viewer.PlayerController))
+#endif
+				{
+					// check if the controller already exists by trying to remove it from the map which was filled up with controllers we were tracking
+					if (ClientConnectionMap.Remove(Viewer.Handle) == 0)
+					{
+						// If not add it to ClientsToAdd. Its important AddClient isn't called until necessary Clients are removed, as we may reuse 
+						// array indices that are already in use.
+						ClientsToAdd.Emplace(Viewer.Handle, Viewer.PlayerController);
+					}
+				}
+			}
+		}
+
+		// anything left in the map needs to be removed from the list
+		for (FControllerMap::TIterator Itr = ClientConnectionMap.CreateIterator(); Itr; ++Itr)
+		{
+			const int32 ViewerIdx = Itr->Value.GetIndex();
+
+			RemoveClient(Itr->Value);
+
+			//InfoData must be valid here as ClientConnectionMap can only have items if ClientHandleManager.GetHandles() was greater than 0 
+			//when we entered the function.
+			bNeedShrinking |= (ViewerIdx == (InfoData.Bubbles.Num() - 1));
+		}
+	}
+
+	// resize the ViewerToClientHandleArray array to match the viewer handles array (for consistency)
+	if (ViewerToClientHandleArray.Num() > Viewers.Num())
+	{
+		ViewerToClientHandleArray.RemoveAt(Viewers.Num(), ViewerToClientHandleArray.Num() - Viewers.Num(), /* bAllowShrinking */ false);
+	}
+
+	for (FClientAddData& ClientAdd : ClientsToAdd)
+	{
+		//TODO make AddClients (ie plural) functionality
+		AddClient(ClientAdd.Handle, *ClientAdd.Controller);
+	}
+
+	return bNeedShrinking;
+}
+
+// synchronize the ClientViewers
+void UMassReplicationManager::SynchronizeClientViewers(const TArray<FViewerInfo>& Viewers)
+{
+	check(MassLODManager);
+
+	struct FMassHandleData
+	{
+		FMassHandleData(FMassClientHandle InClientHandle, FMassViewerHandle InViewerHandle)
+			: ClientHandle(InClientHandle)
+			, ViewerHandle(InViewerHandle)
+		{}
+
+		FMassClientHandle ClientHandle;
+		FMassViewerHandle ViewerHandle;
+	};
+
+	// go through the ClientsViewerHandles and check validity and store the valid ones into a map
+	typedef TMap<APlayerController*, FMassHandleData> FViewerMap;
+
+	FViewerMap ClientViewerMap;
+	{
+		const TArray<FMassClientHandle> ClientHandles = ClientHandleManager.GetHandles();
+
+		for (const FMassClientHandle& ClientHandle : ClientHandles)
+		{
+			//as this is a fresh handle then we only need to check against !IsInvalid
+			if (ClientHandle.IsValid())
+			{
+				FClientViewerHandles& ClientViewers = ClientsViewerHandles[ClientHandle.GetIndex()];
+
+				int32 ViewerIdx = 0;
+
+				while (ViewerIdx < ClientViewers.Handles.Num())
+				{
+					const FMassViewerHandle& ClientViewer = ClientViewers.Handles[ViewerIdx];
+
+					APlayerController* Controller = MassLODManager->GetPlayerControllerFromViewerHandle(ClientViewer);
+
+					const UE::Mass::Replication::ENetConnectionType ConnectionType = UE::Mass::Replication::HasParentOrChildWithValidParentNetConnection(Controller);
+
+					if (ConnectionType != UE::Mass::Replication::ENetConnectionType::None)
+					{
+						++ViewerIdx;
+
+						// we don't verify or remove Parent Connections here as that was done when we synchronized the client bubbles
+						if (ConnectionType == UE::Mass::Replication::ENetConnectionType::Child)
+						{
+							ClientViewerMap.Add(Controller, FMassHandleData(ClientHandle, ClientViewer));
+						}
+					}
+					else //remove invalid ClientViewer, but dont increment the ViewerIdx
+					{
+						ClientViewers.Handles.RemoveAt(ViewerIdx, 1, /* bAllowShrinking */ false);
+					}
+				}
+			}
+		}
+
+		// now go through all current Viewers and add if they are a valid ClientViewer with a child UNetconnection if they do not exist
+		for (const FViewerInfo& Viewer : Viewers)
+		{
+			// we are only interested in valid Viewers
+			if (Viewer.Handle.IsValid())
+			{
+				// we are only processing child UNetConnections that have a valid APlayerController OwningActor
+				const APlayerController* ParentController = UE::Mass::Replication::GetParentControllerFromChildNetConnection(Viewer.PlayerController);
+
+				// check if the parent controller is valid and already exists
+				if (ParentController && (ClientViewerMap.Find(Viewer.PlayerController) == nullptr))
+				{
+					FMassViewerHandle ParentViewerHandle = MassLODManager->GetViewerHandleFromPlayerController(ParentController);
+
+					if (ensureMsgf(ParentViewerHandle.IsValid(), TEXT("MassLODManager handles are out of sync with PlayerController NetConnections!")))
+					{
+						// note all clients (parent NetConnections) should already be set up so we would expect valid handles here
+						check(ViewerToClientHandleArray.IsValidIndex(ParentViewerHandle.GetIndex()));
+
+						const FViewerClientPair& ParentViewerClientPair = ViewerToClientHandleArray[ParentViewerHandle.GetIndex()];
+
+						check(MassLODManager->IsValidViewer(ParentViewerClientPair.ViewerHandle));
+						check(ClientHandleManager.IsValidHandle(ParentViewerClientPair.ClientHandle));
+
+						// remove APlayerController from the ClientViewerMap and Add the viewer to the ClientsViewerHandles
+						ClientViewerMap.Remove(Viewer.PlayerController);
+
+						FClientViewerHandles& ClientViewers = ClientsViewerHandles[ParentViewerClientPair.ClientHandle.GetIndex()];
+
+						ClientViewers.Handles.Add(Viewer.Handle);
+					}
+				}
+			}
+		}
+	}
+
+	{
+		// anything left in the map needs to be removed from the list
+		for (FViewerMap::TIterator Itr = ClientViewerMap.CreateIterator(); Itr; ++Itr)
+		{
+			const FMassHandleData& HandleData = Itr->Value;
+
+			FClientViewerHandles& ClientViewers = ClientsViewerHandles[HandleData.ClientHandle.GetIndex()];
+
+			ClientViewers.Handles.RemoveSingle(HandleData.ViewerHandle);
+		}
+	}
+}
+
+void UMassReplicationManager::SynchronizeClientsAndViewers()
+{
+	// only execute this code at most once per frame
+	if (LastSynchronizedFrame == GFrameCounter)
+	{
+		return;
+	}
+	LastSynchronizedFrame = GFrameCounter;
+
+	if (MassLODManager == nullptr)
+	{
+		checkNoEntry();
+		return;
+	}
+
+	// makes sure the LOD manager Viewers are synced before we process them
+	const TArray<FViewerInfo>& Viewers = MassLODManager->GetSynchronizedViewers();
+
+	const bool bNeedShrinking = SynchronizeClients(Viewers);
+
+	//only synchronize the client viewers outside of the debug functionality for now
+#if !UE_ALLOW_DEBUG_REPLICATION_BUBBLES_STANDALONE
+	SynchronizeClientViewers(Viewers);
+#endif //UE_ALLOW_REPLICATION_DEBUG_BUBBLES_IN_STANDALONE
+
+	if (bNeedShrinking)
+	{
+		const int32 NumItems = ClientHandleManager.ShrinkHandles();
+		 
+		ClientsViewerHandles.RemoveAt(NumItems, ClientsViewerHandles.Num() - NumItems, /* bAllowShrinking */ false);
+		ClientToViewerHandleArray.RemoveAt(NumItems, ClientToViewerHandleArray.Num() - NumItems, /* bAllowShrinking */ false);
+
+		for (FMassClientBubbleInfoData& InfoData : BubbleInfoArray)
+		{
+			InfoData.Bubbles.RemoveAt(NumItems, InfoData.Bubbles.Num() - NumItems, /* bAllowShrinking */ false);
+		}
+	}
+}
+
+FMassBubbleInfoClassHandle UMassReplicationManager::RegisterBubbleInfoClass(const TSubclassOf<AMassClientBubbleInfoBase>& BubbleInfoClass)
+{
+	checkf(BubbleInfoClass.Get() != nullptr, TEXT("BubbleInfoClass must have been set!"));
+
+	if (ClientHandleManager.GetHandles().Num() > 0)
+	{
+		checkf(false, TEXT("RegisterBubbleInfoClass() must not be called after AddClient, BubbleInfoClass has not been registered"));
+		return FMassBubbleInfoClassHandle();
+	}
+
+	const int32 IdxFound = BubbleInfoArray.IndexOfByPredicate([BubbleInfoClass](const FMassClientBubbleInfoData& Data)
+		{
+			return BubbleInfoClass == Data.BubbleClass;
+		});
+
+	if (IdxFound != INDEX_NONE)
+	{
+		UE_LOG(LogMassReplication, Log, TEXT("UMassReplicationManager: Trying to RegisterBubbleInfoClass() twice with the same BubbleInfoClass, Only one BubbleInfoClass will be registered for this type"));
+		return FMassBubbleInfoClassHandle(IdxFound);
+	}
+
+	const int32 Idx = BubbleInfoArray.Emplace(BubbleInfoClass);
+
+	return FMassBubbleInfoClassHandle(Idx);
+}
+
+FMassBubbleInfoClassHandle UMassReplicationManager::GetBubbleInfoClassHandle(const TSubclassOf<AMassClientBubbleInfoBase>& BubbleInfoClass) const
+{
+	FMassBubbleInfoClassHandle Handle;
+
+	for (int32 Idx = 0; Idx < BubbleInfoArray.Num(); ++Idx)
+	{
+		const FMassClientBubbleInfoData& BubbleData = BubbleInfoArray[Idx];
+
+		if (BubbleData.BubbleClass == BubbleInfoClass)
+		{
+			Handle.SetIndex(Idx);
+			break;
+		}
+	}
+	return Handle;
+}
+
+#if UE_REPLICATION_COMPILE_CLIENT_CODE
+void UMassReplicationManager::SetEntity(const FMassNetworkID NetworkID, const FLWEntity Entity)
+{
+	FMassReplicationEntityInfo* EntityInfo = FindMassEntityInfoMutable(NetworkID);
+	check(EntityInfo && !EntityInfo->Entity.IsSet());
+	EntityInfo->Entity = Entity;
+
+	OnMassAgentAdded.Broadcast(NetworkID, Entity);
+}
+#endif // UE_REPLICATION_COMPILE_CLIENT_CODE
+
+#if UE_REPLICATION_COMPILE_CLIENT_CODE
+FLWEntity UMassReplicationManager::ResetEntityIfValid(const FMassNetworkID NetworkID, int32 ReplicationID)
+{
+	FLWEntity EntityReset;
+
+	FMassReplicationEntityInfo* EntityInfo = FindMassEntityInfoMutable(NetworkID);
+
+	checkf(EntityInfo, TEXT("EntityData must have been added to EntityInfoMap!"));
+	checkf(EntityInfo->ReplicationID >= ReplicationID, TEXT("We must not be removing an item we've never added!"));
+
+	// Only reset the item if its currently Set / Valid and its the most recent ReplicationID. Stale removes after more recent adds are ignored
+	// We do need to check the ReplicationID in this case
+	if (EntityInfo->Entity.IsSet() && (EntityInfo->ReplicationID == ReplicationID))
+	{
+		OnRemovingMassAgent.Broadcast(NetworkID, EntityReset);
+
+		EntityReset = EntityInfo->Entity;
+
+		//Unset the Entity handle, this indicates that its currently removed from the bubble
+		EntityInfo->Entity = FLWEntity();
+	}
+
+	return EntityReset;
+}
+#endif // UE_REPLICATION_COMPILE_CLIENT_CODE
+
+#if UE_REPLICATION_COMPILE_CLIENT_CODE
+void UMassReplicationManager::RemoveFromEntityInfoMap(const FMassNetworkID NetworkID)
+{
+	check(NetworkID.IsValid());
+	FMassReplicationEntityInfo Info;
+	ensureMsgf(EntityInfoMap.RemoveAndCopyValue(NetworkID, Info), TEXT("Removing a non-existant NetworkID(%s)!"), *NetworkID.Describe());
+}
+#endif // UE_REPLICATION_COMPILE_CLIENT_CODE
+
+#if UE_REPLICATION_COMPILE_CLIENT_CODE
+UMassReplicationManager::EFindOrAddMassEntityInfo UMassReplicationManager::FindAndUpdateOrAddMassEntityInfo(const FMassNetworkID NetworkID, int32 ReplicationID, const FMassReplicationEntityInfo*& OutMassEntityInfo)
+{
+	FMassReplicationEntityInfo* MassEntityInfo = FindMassEntityInfoMutable(NetworkID);
+	EFindOrAddMassEntityInfo FindOrAddStatus = EFindOrAddMassEntityInfo::FoundOlderReplicationID;
+
+	if (MassEntityInfo)
+	{
+		// Currently we don't think this should be needed, but are leaving it in for bomb proofing
+		if (ensure(MassEntityInfo->ReplicationID < ReplicationID))
+		{
+			// Update the replication ID to the latest
+			MassEntityInfo->ReplicationID = ReplicationID;
+		}
+		else
+		{
+			FindOrAddStatus = EFindOrAddMassEntityInfo::FoundNewerReplicationID;
+		}
+	}
+	else
+	{
+		MassEntityInfo = &EntityInfoMap.Add(NetworkID, FMassReplicationEntityInfo(FLWEntity(), ReplicationID));
+		FindOrAddStatus = EFindOrAddMassEntityInfo::Added;
+	}
+
+	OutMassEntityInfo = MassEntityInfo;
+	return FindOrAddStatus;
+}
+#endif // UE_REPLICATION_COMPILE_CLIENT_CODE
+
+#if UE_REPLICATION_COMPILE_CLIENT_CODE
+const FMassReplicationEntityInfo* UMassReplicationManager::FindMassEntityInfo(const FMassNetworkID NetworkID) const
+{
+	check(NetworkID.IsValid());
+	return EntityInfoMap.Find(NetworkID);
+}
+#endif // UE_REPLICATION_COMPILE_CLIENT_CODE
+
+#if UE_REPLICATION_COMPILE_CLIENT_CODE
+FMassReplicationEntityInfo* UMassReplicationManager::FindMassEntityInfoMutable(const FMassNetworkID NetworkID)
+{
+	check(NetworkID.IsValid());
+	return EntityInfoMap.Find(NetworkID);
+}
+#endif // UE_REPLICATION_COMPILE_CLIENT_CODE
+
+#if UE_REPLICATION_COMPILE_CLIENT_CODE
+FLWEntity UMassReplicationManager::FindEntity(const FMassNetworkID NetworkID) const
+{
+	check(NetworkID.IsValid());
+	const FMassReplicationEntityInfo* Info = EntityInfoMap.Find(NetworkID);
+	return Info ? Info->Entity : FLWEntity();
+}
+#endif // UE_REPLICATION_COMPILE_CLIENT_CODE
+
+void UMassReplicationManager::DebugCheckArraysAreInSync()
+{
+#if UE_DEBUG_REPLICATION
+
+	checkf((ClientToViewerHandleArray.Num() == ClientsViewerHandles.Num()), TEXT("Client arrays out of sync with each other!"));
+
+	const int32 NumEntries = (BubbleInfoArray.Num() > 0) ? BubbleInfoArray[0].Bubbles.Num() : 0;
+
+	for (int32 IdxOuter = 0; IdxOuter < BubbleInfoArray.Num(); ++IdxOuter)
+	{
+		FMassClientBubbleInfoData& InfoDataOuter = BubbleInfoArray[IdxOuter];
+
+		checkf((InfoDataOuter.Bubbles.Num() == ClientsViewerHandles.Num()), TEXT("BubbleInfoArray arrays out of sync with ClientsViewerHandles!"));
+		checkf(InfoDataOuter.Bubbles.Num() == NumEntries, TEXT("Bubbles have different numbers of items!"));
+
+		for (int32 IdxInner = IdxOuter + 1; IdxInner < BubbleInfoArray.Num(); ++IdxInner)
+		{
+			FMassClientBubbleInfoData& InfoDataInner = BubbleInfoArray[IdxInner];
+
+			for (int32 IdxBubble = 0; IdxBubble < InfoDataOuter.Bubbles.Num(); ++IdxBubble)
+			{
+				AMassClientBubbleInfoBase* InfoOuter = InfoDataOuter.Bubbles[IdxBubble];
+				AMassClientBubbleInfoBase* InfoInner = InfoDataInner.Bubbles[IdxBubble];
+
+				const TArray<FMassClientHandle>& Handles = ClientHandleManager.GetHandles();
+
+				if (ClientHandleManager.IsValidHandle(Handles[IdxBubble]))
+				{				
+					check(InfoOuter && InfoInner);
+
+					APlayerController* OuterController = Cast<APlayerController>(InfoOuter->GetOwner());
+					APlayerController* InnerController = Cast<APlayerController>(InfoInner->GetOwner());
+
+					checkf((OuterController != nullptr) && (InnerController != nullptr), TEXT("Controller owners must be valid in BubbleInfoArray"));
+					checkf(OuterController == InnerController, TEXT("Owner controllers at the same indices in different BubbleInfoArray items must be equal!"));
+				}
+				else
+				{
+					check(!InfoOuter && !InfoInner);
+				}
+			}
+		}
+	}
+
+ #endif // UE_DEBUG_REPLICATION
+};
+
+void UMassReplicationManager::AddClient(FMassViewerHandle ViewerHandle, APlayerController& InController)
+{
+	check(World);
+	check(MassLODManager);
+	checkf(MassLODManager->IsValidViewer(ViewerHandle), TEXT("ViewerHandle must be valid"));
+
+#if !UE_ALLOW_DEBUG_REPLICATION_BUBBLES_STANDALONE
+	checkf(UE::Mass::Replication::HasParentNetConnection(&InController), TEXT("InController must have a parent net connection or replication will not occur!"));
+#endif
+
+	if (!ensureMsgf(BubbleInfoArray.Num() > 0, TEXT("BubbleInfoClass has not been set Client will not be added to UMassReplicationManager!")))
+	{
+		return;
+	}
+
+	if (ViewerHandle.GetIndex() >= ViewerToClientHandleArray.Num())
+	{
+		// making no assumptions about the order that ViewerHandles are added, make sure we add enough space to accomodate the new ViewerHandle index
+		ViewerToClientHandleArray.AddDefaulted(ViewerHandle.GetIndex() - ViewerToClientHandleArray.Num() + 1);
+	}
+	else 
+	{
+		checkf(ViewerToClientHandleArray[ViewerHandle.GetIndex()].ViewerHandle.IsValid() == false, TEXT("Adding a Client without removing the previous with the same Index!"));
+	}
+
+	FActorSpawnParameters SpawnParams;
+	SpawnParams.Owner = &InController;
+
+	const FMassClientHandle ClientHandle = ClientHandleManager.GetNextHandle();
+
+	const FViewerClientPair ViewerClientPair(ViewerHandle, ClientHandle);
+
+	FViewerClientPair& ViewerToClientHandleItem = ViewerToClientHandleArray[ViewerHandle.GetIndex()];
+
+	checkf(ViewerToClientHandleItem.IsValid() == false, TEXT("Handles should have been reset when previous Viewer in this slot was removed"));
+
+	ViewerToClientHandleItem = ViewerClientPair;
+
+	checkf(ClientHandle.GetIndex() <= ClientToViewerHandleArray.Num(), TEXT("ClientHandle out of sync with ClientToViewerHandleArray"));
+
+	// check if the handle is a new entry in the free list arrays or uses an existing entry
+	if (ClientHandle.GetIndex() == ClientToViewerHandleArray.Num())
+	{
+		ClientsViewerHandles.AddDefaulted();
+		ClientToViewerHandleArray.Emplace(ViewerHandle, ClientHandle);
+	}
+	else
+	{
+		checkf(ClientsViewerHandles[ClientHandle.GetIndex()].Handles.Num() == 0, TEXT("ClientsViewerHandles being replaced must have zero entries, they should have been removed first!"));
+
+		FViewerClientPair& ClientToViewerHandleItem = ClientToViewerHandleArray[ClientHandle.GetIndex()];
+
+		checkf(ClientToViewerHandleItem.IsValid() == false, TEXT("Handles should have been reset when previous Client in this slot was removed"));
+
+		ClientToViewerHandleItem = ViewerClientPair;
+	}
+
+	FClientViewerHandles& ClientViewers = ClientsViewerHandles[ClientHandle.GetIndex()];
+	ClientViewers.Handles.Add(ViewerHandle);
+
+	for (FMassClientBubbleInfoData& InfoData : BubbleInfoArray)
+	{
+		AMassClientBubbleInfoBase* ClientBubbleInfo = World->SpawnActor<AMassClientBubbleInfoBase>(InfoData.BubbleClass, SpawnParams);
+		ClientBubbleInfo->SetClientHandle(ClientHandle);
+
+		checkf(ClientHandle.GetIndex() <= InfoData.Bubbles.Num(), TEXT("ClientHandle out of sync with Bubbles"));
+
+		// check if the handle is a new entry in the free list arrays or uses an existing entry
+		if (ClientHandle.GetIndex() == InfoData.Bubbles.Num())
+		{
+			InfoData.Bubbles.Push(ClientBubbleInfo);
+		}
+		else
+		{
+			AMassClientBubbleInfoBase*& BubbleUpdate = InfoData.Bubbles[ClientHandle.GetIndex()];
+			checkf(BubbleUpdate == nullptr, TEXT("ClientBubble being replaced must be nullptr it should have been removed first!"));
+
+			BubbleUpdate = ClientBubbleInfo;
+		}
+	}
+	DebugCheckArraysAreInSync();
+}
+
+void UMassReplicationManager::RemoveClient(FMassClientHandle ClientHandle)
+{
+	check(World);
+
+	checkf(ClientHandleManager.IsValidHandle(ClientHandle), TEXT("ClientHandle must be a valid non stale handle"));
+
+	checkf(ClientToViewerHandleArray.IsValidIndex(ClientHandle.GetIndex()), TEXT("ClientHandle is out of sync with ClientToViewerHandleArray!"));
+	FViewerClientPair& ClientToViewerHandleItem = ClientToViewerHandleArray[ClientHandle.GetIndex()];
+
+	checkf(ClientToViewerHandleItem.ViewerHandle.IsValid(), TEXT("Invalid ViewerHandle! ClientHandle is out of sync with ClientToViewerHandleArray!"));
+	checkf(ClientToViewerHandleItem.ClientHandle == ClientHandle, TEXT("ClientHandle is out of sync with ClientToViewerHandleArray!"));
+
+	FClientViewerHandles& ClientViewerHandles = ClientsViewerHandles[ClientHandle.GetIndex()];
+
+	checkf(ClientViewerHandles.Handles.Num() > 0, TEXT("There should always be atleast one client viewer handle (the parent NetConnection)"));
+	ClientViewerHandles.Handles.Reset();
+
+	checkf(ViewerToClientHandleArray.IsValidIndex(ClientToViewerHandleItem.ViewerHandle.GetIndex()), TEXT("ViewerHandle is out of sync with ViewerToClientHandleArray!"));
+	FViewerClientPair& ViewerToClientHandleItem = ViewerToClientHandleArray[ClientToViewerHandleItem.ViewerHandle.GetIndex()];
+
+	checkf(ViewerToClientHandleItem == ClientToViewerHandleItem, TEXT("ClientToViewerHandleArray and ViewerToClientHandleArray out of sync!"));
+
+	ViewerToClientHandleItem.Invalidate();
+	ClientToViewerHandleItem.Invalidate();
+
+	for (FMassClientBubbleInfoData& InfoData : BubbleInfoArray)
+	{
+		checkf(InfoData.Bubbles.IsValidIndex(ClientHandle.GetIndex()), TEXT("ClientHandle is out of sync with Bubbles!"));
+		AMassClientBubbleInfoBase*& BubbleInfoItem = InfoData.Bubbles[ClientHandle.GetIndex()];
+
+		if ((World != nullptr) && (BubbleInfoItem != nullptr))
+		{
+			World->DestroyActor(BubbleInfoItem);
+		}
+
+		BubbleInfoItem = nullptr;
+	}
+
+	ClientHandleManager.RemoveHandle(ClientHandle);
+
+	DebugCheckArraysAreInSync();
+}
diff --git a/Engine/Plugins/Runtime/MassGameplay/Source/MassReplication/Private/MassReplicationProcessorBase.cpp b/Engine/Plugins/Runtime/MassGameplay/Source/MassReplication/Private/MassReplicationProcessorBase.cpp
new file mode 100644
--- /dev/null
+++ b/Engine/Plugins/Runtime/MassGameplay/Source/MassReplication/Private/MassReplicationProcessorBase.cpp
@@ -0,0 +1,156 @@
+// Copyright Epic Games, Inc. All Rights Reserved.UMassSimulationSettings
+
+#include "MassReplicationProcessorBase.h"
+#include "MassClientBubbleHandler.h"
+#include "MassLODManager.h"
+
+//----------------------------------------------------------------------//
+//  UMassProcessor_Replication
+//----------------------------------------------------------------------//
+UMassReplicationProcessorBase::UMassReplicationProcessorBase()
+{
+	LODDistance[EMassLOD::High] = 0.f;
+	LODDistance[EMassLOD::Medium] = 1000.f;
+	LODDistance[EMassLOD::Low] = 2500.f;
+	LODDistance[EMassLOD::Off] = 5000.f;
+	
+	LODMaxCount[EMassLOD::High] = 1600;
+	LODMaxCount[EMassLOD::Medium] = 3200;
+	LODMaxCount[EMassLOD::Low] = 48000;
+	LODMaxCount[EMassLOD::Off] = 0;
+
+	LODMaxCountPerViewer[EMassLOD::High] = 100;
+	LODMaxCountPerViewer[EMassLOD::Medium] = 200;
+	LODMaxCountPerViewer[EMassLOD::Low] = 300;
+	LODMaxCountPerViewer[EMassLOD::Off] = 0;
+	
+	UpdateInterval[EMassLOD::High] = 0.1f;
+	UpdateInterval[EMassLOD::Medium] = 0.2f;
+	UpdateInterval[EMassLOD::Low] = 0.3f;
+	UpdateInterval[EMassLOD::Off] = 0.5f;
+
+	ExecutionFlags = int32(EProcessorExecutionFlags::None);
+	ProcessingPhase = EPipeProcessingPhase::PostPhysics;
+}
+
+void UMassReplicationProcessorBase::ConfigureQueries()
+{
+	EntityQuery.AddRequirement<FDataFragment_MassSimulationLODInfo>(ELWComponentAccess::ReadOnly);
+	EntityQuery.AddRequirement<FMassNetworkIDFragment>(ELWComponentAccess::ReadOnly);
+	EntityQuery.AddRequirement<FDataFragment_ReplicationTemplateID>(ELWComponentAccess::ReadOnly);
+	EntityQuery.AddRequirement<FMassReplicationViewerLODFragment>(ELWComponentAccess::ReadWrite);
+	EntityQuery.AddRequirement<FMassReplicatedAgentFragment>(ELWComponentAccess::ReadWrite);
+}
+
+void UMassReplicationProcessorBase::Initialize(UObject& Owner)
+{
+	Super::Initialize(Owner);
+
+#if UE_REPLICATION_COMPILE_SERVER_CODE
+	LODCalculator.Initialize(LODDistance, BufferHysteresisOnDistancePercentage / 100.0f, LODMaxCount, LODMaxCountPerViewer);
+
+	ReplicationManager = UWorld::GetSubsystem<UMassReplicationManager>(World);
+
+	check(ReplicationManager);
+
+	//BubbleInfoClassHandle = ReplicationManager->GetBubbleInfoClassHandle(AMassCrowdClientBubbleInfo::StaticClass());
+#endif //UE_REPLICATION_COMPILE_SERVER_CODE
+}
+
+void UMassReplicationProcessorBase::PrepareExecution()
+{
+#if UE_REPLICATION_COMPILE_SERVER_CODE
+
+	check(ReplicationManager);
+
+	//first synchronize clients and viewers
+	ReplicationManager->SynchronizeClientsAndViewers();
+
+	//then call base class
+	check(LODManager);
+	const TArray<FViewerInfo>& Viewers = LODManager->GetViewers();
+	LODCalculator.PrepareExecution(Viewers);
+
+	const TArray<FMassClientHandle>& CurrentClientHandles = ReplicationManager->GetClientReplicationHandles();
+	const int32 MinNumHandles = FMath::Min(CachedClientHandles.Num(), CurrentClientHandles.Num());
+
+	//check to see if we don't have enough cached client handles
+	if (CachedClientHandles.Num() < CurrentClientHandles.Num())
+	{
+		CachedClientHandles.Reserve(CurrentClientHandles.Num());
+		bBubbleChanged.Reserve(CurrentClientHandles.Num());
+		BubbleInfos.Reserve(CurrentClientHandles.Num());
+
+		for (int32 Idx = CachedClientHandles.Num(); Idx < CurrentClientHandles.Num(); ++Idx)
+		{
+			const FMassClientHandle& CurrentClientHandle = CurrentClientHandles[Idx];
+
+			CachedClientHandles.Add(CurrentClientHandle);
+			bBubbleChanged.Add(true);
+			AMassClientBubbleInfoBase* Info = CurrentClientHandle.IsValid() ?
+				ReplicationManager->GetClientBubbleChecked(BubbleInfoClassHandle, CurrentClientHandle) :
+				nullptr;
+
+			check(Info);
+
+			BubbleInfos.Add(Info);
+		}
+	}
+	//check to see if we have too many cached client handles
+	else if (CachedClientHandles.Num() > CurrentClientHandles.Num())
+	{
+		const int32 NumRemove = CachedClientHandles.Num() - CurrentClientHandles.Num();
+
+		CachedClientHandles.RemoveAt(CurrentClientHandles.Num(), NumRemove, /* bAllowShrinking */ false);
+		bBubbleChanged.RemoveAt(CurrentClientHandles.Num(), NumRemove, /* bAllowShrinking */ false);
+		BubbleInfos.RemoveAt(CurrentClientHandles.Num(), NumRemove, /* bAllowShrinking */ false);
+	}
+
+	//check to see if any cached client handles have changed, if they have set the BubbleInfo[] appropriately
+	for (int32 Idx = 0; Idx < MinNumHandles; ++Idx)
+	{
+		const FMassClientHandle& CurrentClientHandle = CurrentClientHandles[Idx];
+		FMassClientHandle& CachedClientHandle = CachedClientHandles[Idx];
+
+		const bool bChanged = (bBubbleChanged[Idx] = (CurrentClientHandle != CachedClientHandle));
+		if (bChanged)
+		{
+			AMassClientBubbleInfoBase* Info = CurrentClientHandle.IsValid() ?
+				ReplicationManager->GetClientBubbleChecked(BubbleInfoClassHandle, CurrentClientHandle) :
+				nullptr;
+
+			check(Info);
+
+			BubbleInfos[Idx] = Info;
+			CachedClientHandle = CurrentClientHandle;
+		}
+	}
+
+#endif //UE_REPLICATION_COMPILE_SERVER_CODE
+}
+
+void UMassReplicationProcessorBase::Execute(UEntitySubsystem& EntitySubsystem, FLWComponentSystemExecutionContext& Context)
+{
+#if UE_REPLICATION_COMPILE_SERVER_CODE
+	check(World);
+
+	PrepareExecution();
+
+	EntityQuery.ForEachEntityChunk(EntitySubsystem, Context, [this](FLWComponentSystemExecutionContext& Context)
+		{
+			const TConstArrayView<FDataFragment_MassSimulationLODInfo> ViewersInfoList = Context.GetComponentView<FDataFragment_MassSimulationLODInfo>();
+			const TArrayView<FMassReplicationViewerLODFragment> ViewerLODList = Context.GetMutableComponentView<FMassReplicationViewerLODFragment>();
+			LODCalculator.CalculateLOD(Context, ViewerLODList, ViewersInfoList);
+		});
+
+	if (LODCalculator.AdjustDistancesFromCount())
+	{
+		EntityQuery.ForEachEntityChunk(EntitySubsystem, Context, [this](FLWComponentSystemExecutionContext& Context)
+			{
+				const TConstArrayView<FDataFragment_MassSimulationLODInfo> ViewersInfoList = Context.GetComponentView<FDataFragment_MassSimulationLODInfo>();
+				const TArrayView<FMassReplicationViewerLODFragment> ViewerLODList = Context.GetMutableComponentView<FMassReplicationViewerLODFragment>();
+				LODCalculator.AdjustLODFromCount(Context, ViewerLODList, ViewersInfoList);
+			});
+	}
+#endif //UE_REPLICATION_COMPILE_SERVER_CODE
+}
diff --git a/Engine/Plugins/Runtime/MassGameplay/Source/MassReplication/Private/MassReplicationTrait.cpp b/Engine/Plugins/Runtime/MassGameplay/Source/MassReplication/Private/MassReplicationTrait.cpp
new file mode 100644
--- /dev/null
+++ b/Engine/Plugins/Runtime/MassGameplay/Source/MassReplication/Private/MassReplicationTrait.cpp
@@ -0,0 +1,25 @@
+﻿// Copyright Epic Games, Inc. All Rights Reserved.
+
+#include "MassReplicationTrait.h"
+#include "MassEntityTemplateRegistry.h"
+#include "Engine/World.h"
+#include "MassSpawnerTypes.h"
+#include "MassSimulationLOD.h"
+#include "MassReplicationTypes.h"
+
+
+void UMassReplicationTrait::BuildTemplate(FMassEntityTemplateBuildContext& BuildContext, UWorld& World) const
+{
+	if (World.IsNetMode(NM_Standalone))
+	{
+		return;
+	}
+
+	FDataFragment_ReplicationTemplateID& TemplateIDFragment = BuildContext.AddFragmentWithDefaultInitializer_GetRef<FDataFragment_ReplicationTemplateID>();
+	TemplateIDFragment.ID = BuildContext.GetTemplateID();
+
+	BuildContext.AddFragmentWithDefaultInitializer<FMassNetworkIDFragment>();
+	BuildContext.AddFragment<FMassReplicatedAgentFragment>();
+	BuildContext.AddFragment<FMassReplicationViewerLODFragment>();
+	BuildContext.AddFragment<FDataFragment_MassSimulationLODInfo>();
+}
\ No newline at end of file
diff --git a/Engine/Plugins/Runtime/MassGameplay/Source/MassReplication/Private/MassReplicationTransformHandlers.cpp b/Engine/Plugins/Runtime/MassGameplay/Source/MassReplication/Private/MassReplicationTransformHandlers.cpp
new file mode 100644
--- /dev/null
+++ b/Engine/Plugins/Runtime/MassGameplay/Source/MassReplication/Private/MassReplicationTransformHandlers.cpp
@@ -0,0 +1,27 @@
+// Copyright Epic Games, Inc. All Rights Reserved.
+
+#include "MassReplicationTransformHandlers.h"
+#include "EntityQuery.h"
+#include "MassCommonFragments.h"
+#include "AIHelpers.h"
+
+void FMassReplicationProcessorTransformHandlerBase::AddRequirements(FLWComponentQuery& InQuery)
+{
+	InQuery.AddRequirement<FDataFragment_Transform>(ELWComponentAccess::ReadWrite);
+}
+
+void FMassReplicationProcessorTransformHandlerBase::CacheComponentViews(FLWComponentSystemExecutionContext& ExecContext)
+{
+	TransformList = ExecContext.GetMutableComponentView<FDataFragment_Transform>();
+}
+
+void FMassReplicationProcessorPositionYawHandler::AddEntity(const int32 EntityIdx, FReplicatedAgentPositionYawData& InOutReplicatedPositionYawData) const
+{
+	const FDataFragment_Transform& TransformFragment = TransformList[EntityIdx];
+	InOutReplicatedPositionYawData.SetPosition(TransformFragment.GetTransform().GetLocation());
+
+	if (const TOptional<float> Yaw = UE::AI::GetYawFromQuaternion(TransformFragment.GetTransform().GetRotation()))
+	{
+		InOutReplicatedPositionYawData.SetYaw(Yaw.GetValue());
+	}
+}
diff --git a/Engine/Plugins/Runtime/MassGameplay/Source/MassReplication/Private/MassReplicationTypes.cpp b/Engine/Plugins/Runtime/MassGameplay/Source/MassReplication/Private/MassReplicationTypes.cpp
new file mode 100644
--- /dev/null
+++ b/Engine/Plugins/Runtime/MassGameplay/Source/MassReplication/Private/MassReplicationTypes.cpp
@@ -0,0 +1,56 @@
+// Copyright Epic Games, Inc. All Rights Reserved.
+
+#include "MassReplicationTypes.h"
+#include "MassReplicationManager.h"
+#include "MassSpawnerTypes.h"
+#include "Engine/World.h"
+
+DEFINE_LOG_CATEGORY(LogMassReplication);
+
+//----------------------------------------------------------------------//
+// UMassNetworkIDFragmentInitializer 
+//----------------------------------------------------------------------//
+UMassNetworkIDFragmentInitializer::UMassNetworkIDFragmentInitializer()
+{
+	ExecutionFlags = (int32)(EProcessorExecutionFlags::Standalone | EProcessorExecutionFlags::Server);
+	FragmentType = FMassNetworkIDFragment::StaticStruct();
+}
+
+void UMassNetworkIDFragmentInitializer::ConfigureQueries()
+{
+	EntityQuery.AddRequirement<FMassNetworkIDFragment>(ELWComponentAccess::ReadWrite);
+}
+
+void UMassNetworkIDFragmentInitializer::Initialize(UObject& Owner)
+{
+	Super::Initialize(Owner);
+#if UE_REPLICATION_COMPILE_SERVER_CODE
+	ReplicationManager = UWorld::GetSubsystem<UMassReplicationManager>(Owner.GetWorld());
+#endif //UE_REPLICATION_COMPILE_SERVER_CODE
+}
+
+void UMassNetworkIDFragmentInitializer::Execute(UEntitySubsystem& EntitySubsystem, FLWComponentSystemExecutionContext& Context)
+{
+	QUICK_SCOPE_CYCLE_COUNTER(MassProcessor_InitNetworkID_Run);
+
+	const UWorld* World = EntitySubsystem.GetWorld();
+	const ENetMode NetMode = World->GetNetMode();
+
+	if (NetMode != NM_Client)
+	{
+#if UE_REPLICATION_COMPILE_SERVER_CODE
+		check(ReplicationManager);
+
+		EntityQuery.ForEachEntityChunk(EntitySubsystem, Context, [this](FLWComponentSystemExecutionContext& Context)
+			{
+				const TArrayView<FMassNetworkIDFragment> NetworkIDList = Context.GetMutableComponentView<FMassNetworkIDFragment>();
+				const int32 NumEntities = Context.GetEntitiesNum();
+
+				for (int32 Idx = 0; Idx < NumEntities; ++Idx)
+				{
+					NetworkIDList[Idx].NetID = ReplicationManager->GetNextAvailableMassNetID();
+				}
+			});
+#endif //UE_REPLICATION_COMPILE_SERVER_CODE
+	}
+}
diff --git a/Engine/Plugins/Runtime/MassGameplay/Source/MassReplication/Public/MassClientBubbleInfoBase.h b/Engine/Plugins/Runtime/MassGameplay/Source/MassReplication/Public/MassClientBubbleInfoBase.h
new file mode 100644
--- /dev/null
+++ b/Engine/Plugins/Runtime/MassGameplay/Source/MassReplication/Public/MassClientBubbleInfoBase.h
@@ -0,0 +1,40 @@
+// Copyright Epic Games, Inc. All Rights Reserved.
+
+#pragma once
+
+#include "GameFramework/Info.h"
+#include "MassReplicationTypes.h"
+#include "Engine/World.h"
+
+#include "MassClientBubbleInfoBase.generated.h"
+
+struct FMassClientBubbleSerializerBase;
+
+/** The info actor base class that provides the actual replication */
+UCLASS()
+class MASSREPLICATION_API AMassClientBubbleInfoBase : public AInfo
+{
+	GENERATED_BODY()
+
+public:
+	AMassClientBubbleInfoBase(const FObjectInitializer& ObjectInitializer);
+
+	void SetClientHandle(FMassClientHandle InClientHandle);
+
+protected:
+	virtual void PostInitProperties() override;
+
+	// Called either on PostWorldInit() or PostInitProperties()
+	virtual void InitializeForWorld(UWorld& World);
+
+	virtual void BeginPlay() override;
+	virtual void EndPlay(const EEndPlayReason::Type EndPlayReason) override;
+	virtual void Tick(float DeltaTime) override;
+
+private:
+	void OnPostWorldInit(UWorld* World, const UWorld::InitializationValues);
+
+protected:
+	FDelegateHandle OnPostWorldInitDelegateHandle;
+	TArray<FMassClientBubbleSerializerBase*> Serializers;
+};
diff --git a/Engine/Plugins/Runtime/MassGameplay/Source/MassReplication/Public/MassClientBubbleSerializerBase.h b/Engine/Plugins/Runtime/MassGameplay/Source/MassReplication/Public/MassClientBubbleSerializerBase.h
new file mode 100644
--- /dev/null
+++ b/Engine/Plugins/Runtime/MassGameplay/Source/MassReplication/Public/MassClientBubbleSerializerBase.h
@@ -0,0 +1,58 @@
+// Copyright Epic Games, Inc. All Rights Reserved.
+
+#pragma once
+
+#include "Engine/NetSerialization.h"
+#include "MassReplicationTypes.h"
+
+#include "MassClientBubbleSerializerBase.generated.h"
+
+
+class UWorld;
+class UMassSpawnerSubsystem;
+class UPipeEntitySubsystem;
+class IClientBubbleHandlerInterface;
+
+/**
+ * Classes derived from this will contain the IClientBubbleHandlerInterface derived class as well as the actual Fast Array.
+ * This class mainly provides the base automation with the IClientBubbleHandlerInterface
+ */
+USTRUCT()
+struct MASSREPLICATION_API FMassClientBubbleSerializerBase : public FFastArraySerializer
+{
+	GENERATED_BODY()
+
+public:
+
+#if UE_REPLICATION_COMPILE_CLIENT_CODE
+	void PreReplicatedRemove(const TArrayView<int32> RemovedIndices, int32 FinalSize) const;
+	void PostReplicatedAdd(const TArrayView<int32> AddedIndices, int32 FinalSize) const;
+	void PostReplicatedChange(const TArrayView<int32> ChangedIndices, int32 FinalSize) const;
+#endif //UE_REPLICATION_COMPILE_CLIENT_CODE
+
+	void InitializeForWorld(UWorld& InWorld);
+
+	UWorld* GetWorld() const { return World; }
+	UMassSpawnerSubsystem* GetSpawnerSubsystem() const { return SpawnerSubsystem; }
+	UMassReplicationManager* GetReplicationManager() const { return ReplicationManager; }
+	UPipeEntitySubsystem* GetEntitySystem() const { return EntitySystem; }
+
+	void SetClientHandler(IClientBubbleHandlerInterface& InClientHandler) { ClientHandler = &InClientHandler; }
+	IClientBubbleHandlerInterface* GetClientHandler() const { return ClientHandler; }
+
+private:
+	UPROPERTY(Transient)
+	UWorld* World = nullptr;
+
+	UPROPERTY(Transient)
+	UMassSpawnerSubsystem* SpawnerSubsystem = nullptr;
+
+	UPROPERTY(Transient)
+	UPipeEntitySubsystem* EntitySystem = nullptr;
+
+	UPROPERTY(Transient)
+	UMassReplicationManager* ReplicationManager = nullptr;
+
+	/** Pointer to the IClientBubbleHandlerInterface derived class in the class derived from this one */
+	IClientBubbleHandlerInterface* ClientHandler = nullptr;
+};
diff --git a/Engine/Plugins/Runtime/MassGameplay/Source/MassReplication/Public/MassReplicationManager.h b/Engine/Plugins/Runtime/MassGameplay/Source/MassReplication/Public/MassReplicationManager.h
new file mode 100644
--- /dev/null
+++ b/Engine/Plugins/Runtime/MassGameplay/Source/MassReplication/Public/MassReplicationManager.h
@@ -0,0 +1,291 @@
+// Copyright Epic Games, Inc. All Rights Reserved.
+
+#pragma once
+
+#include "MassCommonTypes.h"
+#include "LWComponentTypes.h"
+#include "MassReplicationTypes.h"
+#include "Subsystems/WorldSubsystem.h"
+
+#include "MassReplicationManager.generated.h"
+
+class UPipeEntitySubsystem;
+class AMassCrowdClientBubbleInfo;
+class UMassLODManager;
+class AMassClientBubbleInfoBase;
+
+struct FClientViewerHandles
+{
+	TArray<FMassViewerHandle> Handles;
+};
+
+struct FViewerClientPair
+{
+	FViewerClientPair() = default;
+
+	FViewerClientPair(FMassViewerHandle InViewerHandle, FMassClientHandle InClientHandle)
+		: ViewerHandle(InViewerHandle)
+		, ClientHandle(InClientHandle)
+	{}
+
+	bool operator== (FViewerClientPair Other) const { return ViewerHandle == Other.ViewerHandle && ClientHandle == Other.ClientHandle; }
+
+	bool IsValid() const { return ViewerHandle.IsValid() && ClientHandle.IsValid(); }
+	void Invalidate() { ViewerHandle.Invalidate(); ClientHandle.Invalidate(); }
+
+	FMassViewerHandle ViewerHandle;
+	FMassClientHandle ClientHandle;
+};
+
+USTRUCT()
+struct FMassClientBubbleInfoData
+{
+	GENERATED_BODY()
+
+	FMassClientBubbleInfoData() = default;
+
+	FMassClientBubbleInfoData(TSubclassOf<AMassClientBubbleInfoBase> InBubbleClass)
+		: BubbleClass(InBubbleClass)
+	{}
+
+	/** A free list array of AMassClientBubbleInfos. This is organised so the index is that of the client FMassClientHandle */
+	UPROPERTY(Transient)
+	TArray<AMassClientBubbleInfoBase*> Bubbles;
+
+	UPROPERTY(Transient)
+	TSubclassOf<AMassClientBubbleInfoBase> BubbleClass;
+};
+
+namespace UE::MassReplication
+{
+	DECLARE_MULTICAST_DELEGATE_TwoParams(FMassAgentDelegate, FMassNetworkID, FLWEntity);
+
+} // UE::MassReplication
+
+
+/**
+ *  Manages the creation of NetworkIDs, ClientBubbles and ClientViewers.
+ *  NetworkIDs are per replicated Agent Entity and are unique and replicated between server and clients.
+ *  ClientBubbles relate to the player controller that owns the parent UNetConnection to a Client machine.
+ *  ClientViewers relate to all the player controllers that have a parent or child UNetConnection to a single Client machine (split screen etc).
+ */
+UCLASS()
+class MASSREPLICATION_API UMassReplicationManager : public UWorldSubsystem
+{
+	GENERATED_BODY()
+
+public:
+
+	enum class EFindOrAddMassEntityInfo : uint8
+	{
+		FoundOlderReplicationID,
+		FoundNewerReplicationID,
+		Added,
+	};
+
+	bool IsBubbleClassHandleValid(FMassBubbleInfoClassHandle Handle) const
+	{
+		return BubbleInfoArray.IsValidIndex(Handle.GetIndex());
+	}
+
+	/** Retrieve network id from Mass handle. */
+	FMassNetworkID GetNetIDFromHandle(const FMassHandle& Handle) const;
+
+	static inline FMassNetworkID GetNextAvailableMassNetID() { return FMassNetworkID(++CurrentNetMassCounter); }
+
+	const TArray<FMassClientHandle>& GetClientReplicationHandles() const { return ClientHandleManager.GetHandles(); }
+
+	/** Gets the client bubble safely */
+	AMassClientBubbleInfoBase* GetClientBubble(FMassBubbleInfoClassHandle BubbleClassHandle, FMassClientHandle ClientHandle) const
+	{
+		AMassClientBubbleInfoBase* ClientBubble = nullptr;
+
+		if (IsBubbleClassHandleValid(BubbleClassHandle) && ClientHandleManager.IsValidHandle(ClientHandle))
+		{
+			ClientBubble = BubbleInfoArray[BubbleClassHandle.GetIndex()].Bubbles[ClientHandle.GetIndex()];
+			check(ClientBubble);
+		}
+		return ClientBubble;
+	}
+
+	/** Get the client bubble. Faster version using check()s */
+	AMassClientBubbleInfoBase* GetClientBubbleChecked(FMassBubbleInfoClassHandle BubbleClassHandle, FMassClientHandle ClientHandle) const
+	{
+		check(IsBubbleClassHandleValid(BubbleClassHandle));
+		check(ClientHandleManager.IsValidHandle(ClientHandle));
+		AMassClientBubbleInfoBase* ClientBubble = BubbleInfoArray[BubbleClassHandle.GetIndex()].Bubbles[ClientHandle.GetIndex()];
+
+		check(ClientBubble);
+		return ClientBubble;
+	}
+
+	/** Gets the client bubble safely returning the template type */
+	template<typename TType>
+	TType* GetTypedClientBubble(FMassBubbleInfoClassHandle BubbleClassHandle, FMassClientHandle ClientHandle) const
+	{
+		return Cast<TType>(GetClientBubble(BubbleClassHandle, ClientHandle));
+	}
+
+	/** Gets the client bubble returning the template type. Faster version using check()s */
+	template<typename TType>
+	TType* GetTypedClientBubbleChecked(FMassBubbleInfoClassHandle BubbleClassHandle, FMassClientHandle ClientHandle) const
+	{
+		AMassClientBubbleInfoBase* ClientBubble = GetClientBubbleChecked(BubbleClassHandle, ClientHandle);
+		checkSlow(Cast<TType>(ClientBubble) != nullptr);
+
+		return static_cast<TType *>(ClientBubble);
+	}
+
+	/** Gets the client viewers safely */
+	const FClientViewerHandles* GetClientViewers(FMassClientHandle Handle) const
+	{
+		return ClientHandleManager.IsValidHandle(Handle) ? &(ClientsViewerHandles[Handle.GetIndex()]) : nullptr;
+	}
+
+	/** Gets the client viewers. Faster version using check()s */
+	const FClientViewerHandles& GetClientViewersChecked(FMassClientHandle Handle) const
+	{
+		check(ClientHandleManager.IsValidHandle(Handle));
+
+		return ClientsViewerHandles[Handle.GetIndex()];
+	}
+
+	void SynchronizeClientsAndViewers();
+
+	/** 
+	 * Registers BubbleInfoClass with the system. These are created one per client as clients join. This function can be called multiple times with
+	 * different BubbleInfoClass types (but not the same type), in that case each client will get multiple BubbleInfoClasses. This can be useful for 
+	 * replicating different Entity types in different bubbles, although its also possible to have multiple TClientBubbleHandlerBase derived class
+	 * instances per BubbleInfoClass. This must not be called after AddClient or SynchronizeClientsAndViewers is called
+	 * @return FMassBubbleInfoClassHandle Handle to the BubbleInfoClass, this will be an invalid handle if its not been created
+	 */
+	FMassBubbleInfoClassHandle RegisterBubbleInfoClass(const TSubclassOf<AMassClientBubbleInfoBase>& BubbleInfoClass);
+
+	/** @return FMassBubbleInfoClassHandle Handle to the BubbleInfoClass, this will be an invalid handle if BubbleInfoClass can not be found */
+	FMassBubbleInfoClassHandle GetBubbleInfoClassHandle(const TSubclassOf<AMassClientBubbleInfoBase>& BubbleInfoClass) const;
+
+#if UE_REPLICATION_COMPILE_CLIENT_CODE
+
+	/** 
+	 * Must be called immediately after an Entity is added to the simulation and after FindAndUpdateOrAddMassEntityInfo has been called
+	 */
+	void SetEntity(const FMassNetworkID NetworkID, const FLWEntity Entity);
+
+	/** 
+	 * Resets the item with NetworkID in the  EntityInfoMap if its currently Set / Valid and its the most recent ReplicationID. In this case the FLWEntity that was reset is returned.
+	 * Otherwise FLWEntity::IsSet() will be false on the returned FLWEntity
+	 * Must be called after SetEntity and just before the Entity is removed on the client. This sets the Entity member of the AddToEntityInfoMap item to FLWEntity so it will be not IsSet() 
+	 */
+	FLWEntity ResetEntityIfValid(const FMassNetworkID NetworkID, int32 ReplicationID);
+
+	void RemoveFromEntityInfoMap(const FMassNetworkID NetworkID);
+
+	/** 
+	 * Finds or adds an FMassReplicationEntityInfo to the EntityInfoMap for an Entity with NetWorkID. This will update the ReplicationID if an older one is found but it does not change the 
+	 * FMassReplicationEntityInfo::Entity member. Returns whether we added an entry in EntityInfoMap, or if the found item has an older or newer ReplicationID than that that is passed in. 
+	 * We shouldn't actually get an equal ReplicationID, but in this case we'll treat the existing EntityInfoMap entry as newer. SetEntity() must be called soon after this if the entity handle is !IsSet()
+	 * @param NetworkID ID of the entity we are looking up OutMassEntityInfo for
+	 * @param ReplicationID this is the ReplicationID of the FastArray item
+	 * @param OutMassEntityInfo retrieved or added FMassReplicationEntityInfo
+	 * @return EFindOrAddMassEntityInfo, result of the function call
+	 */
+	EFindOrAddMassEntityInfo FindAndUpdateOrAddMassEntityInfo(const FMassNetworkID NetworkID, int32 ReplicationID, const FMassReplicationEntityInfo*& OutMassEntityInfo);
+
+	/** 
+	 * Finds the FMassReplicationEntityInfo for the Entity with NetworkID. NetworkID must be valid and nullptr is returned if we can't find it. Note FMassReplicationEntityInfo::IsSet() may be false,
+	 * in which case this Entity has been removed from the client simulation (but we are still storing the ReplicationID between Removes and Adds!
+	 */
+	const FMassReplicationEntityInfo* FindMassEntityInfo(const FMassNetworkID NetworkID) const;
+
+	FLWEntity FindEntity(const FMassNetworkID NetworkID) const;
+
+	const TMap<FMassNetworkID, FMassReplicationEntityInfo>& GetEntityInfoMap() const { return EntityInfoMap; }
+
+	UE::MassReplication::FMassAgentDelegate& GetOnMassAgentAdded()
+	{
+		return OnMassAgentAdded;
+	}
+
+	UE::MassReplication::FMassAgentDelegate& GetOnRemovingMassAgent()
+	{
+		return OnRemovingMassAgent;
+	}
+
+#endif // UE_REPLICATION_COMPILE_CLIENT_CODE
+
+protected:
+	// USubsystem BEGIN
+	virtual void Initialize(FSubsystemCollectionBase& Collection) override;
+	virtual void Deinitialize() override;
+	// USubsystem END
+	
+#if UE_REPLICATION_COMPILE_CLIENT_CODE
+	FMassReplicationEntityInfo* FindMassEntityInfoMutable(const FMassNetworkID NetworkID);
+#endif // UE_REPLICATION_COMPILE_CLIENT_CODE
+
+	/** Adds a Client and sets up all the data relevant to the bubble. Where there are multiple viewers for
+	 *  one NetConnection only the parent net connection should be added here
+	 *  @param ViewerHandle handle to viewer to add as a client. This must be UMassLODManager->IsValidViewer()
+	 *  @param InController associated with the viewer
+	 */
+	void AddClient(FMassViewerHandle ViewerHandle, APlayerController& InController);
+
+	/** Removes a Client and removes up all the data relevant to the bubble
+	 *  @param ClientHandle handle to client to remove. This must be ClientHandleManager.IsValidHandle()
+	 */
+	void RemoveClient(FMassClientHandle ClientHandle);
+
+	/** @return true if we should shrink the number of handles. */
+	bool SynchronizeClients(const TArray<FViewerInfo>& Viewers);
+	void SynchronizeClientViewers(const TArray<FViewerInfo>& Viewers);
+
+	static void DebugCheckArraysAreInSync();
+
+protected:
+	static uint32 CurrentNetMassCounter;
+
+	UPROPERTY()
+	UPipeEntitySubsystem* EntitySystem;
+
+	/** Clients free list FMassClientHandle manager, handles will to the indices of FMassClientReplicationData::ClientBubbles */
+	FMassClientHandleManager ClientHandleManager;
+
+	UPROPERTY(Transient)
+	TArray<FMassClientBubbleInfoData> BubbleInfoArray;
+
+	/** An Array of each Clients viewer handles (split screen players sharing the same client connections). 
+	 *  This will include both the parent and child NetConnections per client.
+	 *	The array is organised so the array index is the same as the index of the client FMassClientHandle.
+	 */
+	TArray<FClientViewerHandles> ClientsViewerHandles;
+
+	/** A free list array of FViewerClientPairs. This is organised so the index is that of the FMassViewerHandle for fast lookup of the related FMassClientHandle. 
+	 *  This only contains viewers that are also clients (ie only parent NetConnections not child ones).
+	 */
+	TArray<FViewerClientPair> ViewerToClientHandleArray;
+
+	/** A free list array of FViewerClientPairs. This is organised so the index is that of the FMassClientHandle.
+	 *  For fast lookup of the related FMassViewerHandle. This only contains viewers that are also clients (ie only parent NetConnections not child ones).
+	 */
+	TArray<FViewerClientPair> ClientToViewerHandleArray;
+
+	UPROPERTY()
+	UWorld* World;
+
+	UPROPERTY()
+	UMassLODManager* MassLODManager;
+
+#if UE_REPLICATION_COMPILE_CLIENT_CODE
+	TMap<FMassNetworkID, FMassReplicationEntityInfo> EntityInfoMap;
+
+	// @todo MassReplication consider batching these
+	/** Broadcast just after a mass agent is added to the client simulation */
+	UE::MassReplication::FMassAgentDelegate OnMassAgentAdded;
+
+	/** Broadcast just before a mass agent is removed from the client simulation */
+	UE::MassReplication::FMassAgentDelegate OnRemovingMassAgent;
+#endif // UE_REPLICATION_COMPILE_CLIENT_CODE
+
+	/** Used to make sure the FMassHandle_ClientReplications are synchronized immediately before they are needed*/
+	uint64 LastSynchronizedFrame = 0;
+};
\ No newline at end of file
diff --git a/Engine/Plugins/Runtime/MassGameplay/Source/MassReplication/Public/MassReplicationProcessorBase.h b/Engine/Plugins/Runtime/MassGameplay/Source/MassReplication/Public/MassReplicationProcessorBase.h
new file mode 100644
--- /dev/null
+++ b/Engine/Plugins/Runtime/MassGameplay/Source/MassReplication/Public/MassReplicationProcessorBase.h
@@ -0,0 +1,210 @@
+// Copyright Epic Games, Inc. All Rights Reserved.
+
+#pragma once
+
+#include "Engine/World.h"
+#include "MassProcessor.h"
+#include "MassCommonTypes.h"
+#include "MassLODCalculator.h"
+#include "MassSimulationLOD.h"
+#include "MassReplicationTypes.h"
+#include "MassSpawnerTypes.h"
+#include "MassLODManager.h"
+#include "MassReplicationManager.h"
+#include "MassReplicationProcessorBase.generated.h"
+
+class UMassReplicationManager;
+class AMassClientBubbleInfoBase;
+class UWorld;
+
+/** 
+ *  Base processor that handles replication and only runs on the server. You should derive from this per entity type (that require different replication processing). It and its derived classes 
+ *  query Mass entity fragments and set those values for replication when appropriate, using the MassClientBubbleHandler.
+ */
+UCLASS()
+class MASSREPLICATION_API UMassReplicationProcessorBase : public UMassProcessor_LODBase
+{
+	GENERATED_BODY()
+
+public:
+	UMassReplicationProcessorBase();
+
+protected:
+	virtual void ConfigureQueries() override;
+	virtual void Initialize(UObject& Owner) override;
+	virtual void Execute(UEntitySubsystem& EntitySubsystem, FLWComponentSystemExecutionContext& Context) override;
+
+	void PrepareExecution();
+
+	/** 
+	 *  Implemented as straight template callbacks as when profiled this was faster than TFunctionRef. Its probably easier to pass Lamdas in to these
+	 *  but Functors can also be used as well as TFunctionRefs etc. Its also fairly straight forward to call member functions via some Lamda glue code
+	 */
+	template<typename AgentArrayItem, typename CacheViewsCallback, typename AddEntityCallback, typename ModifyEntityCallback, typename RemoveEntityCallback>
+	void CalculateClientReplication(FLWComponentSystemExecutionContext& Context, CacheViewsCallback&& CacheViews, AddEntityCallback&& AddEntity, ModifyEntityCallback&& ModifyEntity, RemoveEntityCallback&& RemoveEntity);
+
+	template<typename T>
+	T& GetTypedClientBubbleInfoChecked(FMassClientHandle Handle);
+
+protected:
+
+	UPROPERTY()
+	UMassReplicationManager* ReplicationManager;
+
+	FMassBubbleInfoClassHandle BubbleInfoClassHandle;
+
+	/** Distance where each LOD becomes relevant */
+	UPROPERTY(EditAnywhere, Category = "Mass|LOD", config)
+	float LODDistance[EMassLOD::Max];
+
+	/** Hysteresis percentage on delta between the LOD distances */
+	UPROPERTY(EditAnywhere, Category = "Mass|LOD", meta = (ClampMin = "0.0", UIMin = "0.0"), config)
+	float BufferHysteresisOnDistancePercentage = 10.0f;
+
+	/** Maximum limit of entity per LOD */
+	UPROPERTY(EditAnywhere, Category = "Mass|LOD", config)
+	int32 LODMaxCount[EMassLOD::Max];
+
+	/** Maximum limit of entity per LOD per viewer */
+	UPROPERTY(EditAnywhere, Category = "Mass|LOD", config)
+	int32 LODMaxCountPerViewer[EMassLOD::Max];
+
+	/** Distance where each LOD becomes relevant */
+	UPROPERTY(EditAnywhere, Category = "Mass|LOD", config)
+	float UpdateInterval[EMassLOD::Max];
+
+	TArray<FMassClientHandle, TInlineAllocator<UE::Mass::Replication::MaxNumOfClients>> CachedClientHandles;
+	TArray<bool, TInlineAllocator<UE::Mass::Replication::MaxNumOfClients>> bBubbleChanged;
+
+	//TODO review if we need to have this as a UPROPERTY at all and also if we can make this use a TInlineAllocator
+	//Can not use TInlineAllocator with UPROPERTY()
+	UPROPERTY()
+	TArray<AMassClientBubbleInfoBase*> BubbleInfos;
+
+	TMassLODCalculator<FReplicationLODLogic> LODCalculator;
+
+	FLWComponentQuery EntityQuery;
+};
+
+template<typename AgentArrayItem, typename CacheViewsCallback, typename AddEntityCallback, typename ModifyEntityCallback, typename RemoveEntityCallback>
+void UMassReplicationProcessorBase::CalculateClientReplication(FLWComponentSystemExecutionContext& Context, CacheViewsCallback&& CacheViews, AddEntityCallback&& AddEntity, ModifyEntityCallback&& ModifyEntity, RemoveEntityCallback&& RemoveEntity)
+{
+#if UE_REPLICATION_COMPILE_SERVER_CODE
+
+	check(ReplicationManager);
+	check(LODManager);
+	check(World);
+
+	const int32 NumEntities = Context.GetEntitiesNum();
+
+	TConstArrayView<FMassNetworkIDFragment> NetworkIDList = Context.GetComponentView<FMassNetworkIDFragment>();
+	TArrayView<FMassReplicationViewerLODFragment> ViewerLODList = Context.GetMutableComponentView<FMassReplicationViewerLODFragment>();
+	TArrayView<FMassReplicatedAgentFragment> ReplicatedAgentList = Context.GetMutableComponentView<FMassReplicatedAgentFragment>();
+	TConstArrayView<FDataFragment_ReplicationTemplateID> TemplateIDList = Context.GetComponentView<FDataFragment_ReplicationTemplateID>();
+
+	CacheViews();
+
+	const float Time = World->GetRealTimeSeconds();
+
+	for (int EntityIdx = 0; EntityIdx < NumEntities; EntityIdx++)
+	{
+		FMassReplicatedAgentFragment& AgentFragment = ReplicatedAgentList[EntityIdx];
+
+		if (AgentFragment.AgentsData.Num() < CachedClientHandles.Num())
+		{
+			AgentFragment.AgentsData.AddDefaulted(CachedClientHandles.Num() - AgentFragment.AgentsData.Num());
+		}
+		else if (AgentFragment.AgentsData.Num() > CachedClientHandles.Num())
+		{
+			AgentFragment.AgentsData.RemoveAt(CachedClientHandles.Num(), AgentFragment.AgentsData.Num() - CachedClientHandles.Num(), /* bAllowShrinking */ false);
+		}
+
+		for (int32 ClientIdx = 0; ClientIdx < CachedClientHandles.Num(); ++ClientIdx)
+		{
+			const FMassClientHandle& ClientHandle = CachedClientHandles[ClientIdx];
+
+			if (ClientHandle.IsValid())
+			{
+				checkSlow(BubbleInfos[ClientHandle.GetIndex()] != nullptr);
+
+				FMassReplicatedAgentArrayData& AgentData = AgentFragment.AgentsData[ClientIdx];
+
+				//if the bubble has changed, set the handle Invalid. We will set this to something valid if the agent is going to replicate to the bubble
+				//When a bubble has changed the existing AMassCrowdClientBubbleInfo will reset all the data associated with it.
+				if (bBubbleChanged[ClientIdx])
+				{
+					AgentData.Invalidate();
+				}
+
+				//now we need to see what our highest viewer LOD is on this client (split screen etc), we can use the unsafe version as we have already checked all the
+				//CachedClientHandles for validity
+				const FClientViewerHandles& ClientViewers = ReplicationManager->GetClientViewersChecked(ClientHandle);
+
+				EMassLOD::Type HighestLOD = EMassLOD::Off;
+
+				for (const FMassViewerHandle& ViewerHandle : ClientViewers.Handles)
+				{
+					//this should always we valid as we synchronized the viewers just previously
+					check(LODManager->IsValidViewer(ViewerHandle));
+
+					const EMassLOD::Type MassLOD = ViewerLODList[EntityIdx].LODPerViewer[ViewerHandle.GetIndex()];
+					check(MassLOD <= EMassLOD::Off);
+
+					if (HighestLOD > MassLOD)
+					{
+						HighestLOD = MassLOD;
+					}
+				}
+
+				if (HighestLOD < EMassLOD::Off)
+				{
+#if UE_ALLOW_DEBUG_REPLICATION
+					AgentData.LOD = HighestLOD;
+#endif // UE_ALLOW_DEBUG_REPLICATION
+
+					//if the handle isn't valid we need to add the agent
+					if (!AgentData.Handle.IsValid())
+					{
+						typename AgentArrayItem::FReplicatedAgentType ReplicatedAgent;
+
+						const FMassNetworkIDFragment& NetIDFragment = NetworkIDList[EntityIdx];
+						const FDataFragment_ReplicationTemplateID& TemplateIDFragment = TemplateIDList[EntityIdx];
+
+						ReplicatedAgent.SetNetID(NetIDFragment.NetID);
+						ReplicatedAgent.SetTemplateID(TemplateIDFragment.ID);
+
+						AgentData.Handle = AddEntity(EntityIdx, ReplicatedAgent, ClientHandle);
+
+						AgentData.LastUpdateTime = Time;
+					}
+					else
+					{
+						ModifyEntity(EntityIdx, HighestLOD, Time, AgentData.Handle, ClientHandle);
+					}
+				}
+				else
+				{
+					// as this is a fresh handle, if its valid then we can use the unsafe remove function
+					if (AgentData.Handle.IsValid())
+					{
+						RemoveEntity(AgentData.Handle, ClientHandle);
+						AgentData.Invalidate();
+					}
+				}
+			}
+		}
+	}
+#endif //UE_REPLICATION_COMPILE_SERVER_CODE
+}
+
+template<typename T>
+T& UMassReplicationProcessorBase::GetTypedClientBubbleInfoChecked(FMassClientHandle Handle)
+{
+	checkSlow(BubbleInfos.IsValidIndex(Handle.GetIndex()));
+
+	AMassClientBubbleInfoBase* BubbleInfo = BubbleInfos[Handle.GetIndex()];
+
+	checkSlow(BubbleInfo && Cast<T>(BubbleInfo) != nullptr);
+
+	return *static_cast<T*>(BubbleInfo);
+}
diff --git a/Engine/Plugins/Runtime/MassGameplay/Source/MassReplication/Public/MassReplicationTrait.h b/Engine/Plugins/Runtime/MassGameplay/Source/MassReplication/Public/MassReplicationTrait.h
new file mode 100644
--- /dev/null
+++ b/Engine/Plugins/Runtime/MassGameplay/Source/MassReplication/Public/MassReplicationTrait.h
@@ -0,0 +1,17 @@
+// Copyright Epic Games, Inc. All Rights Reserved.
+
+#pragma once 
+
+#include "MassEntityTraitBase.h"
+#include "MassReplicationTrait.generated.h"
+
+
+UCLASS(BlueprintType, EditInlineNew, CollapseCategories, meta=(DisplayName="Replication"))
+class MASSREPLICATION_API UMassReplicationTrait : public UMassEntityTraitBase
+{
+	GENERATED_BODY()
+
+protected:
+
+	virtual void BuildTemplate(FMassEntityTemplateBuildContext& BuildContext, UWorld& World) const override;
+};
diff --git a/Engine/Plugins/Runtime/MassGameplay/Source/MassReplication/Public/MassReplicationTransformHandlers.h b/Engine/Plugins/Runtime/MassGameplay/Source/MassReplication/Public/MassReplicationTransformHandlers.h
new file mode 100644
--- /dev/null
+++ b/Engine/Plugins/Runtime/MassGameplay/Source/MassReplication/Public/MassReplicationTransformHandlers.h
@@ -0,0 +1,229 @@
+// Copyright Epic Games, Inc. All Rights Reserved.
+
+#pragma once
+
+#include "MassMovementTypes.h"
+#include "MassReplicationTypes.h"
+#include "EntityView.h"
+#include "MassClientBubbleHandler.h"
+#include "MassCommonFragments.h"
+#include "AIHelpers.h"
+
+#include "MassReplicationTransformHandlers.generated.h"
+
+struct FLWComponentQuery;
+
+namespace UE::Mass::Replication
+{
+	constexpr float PositionReplicateTolerance = 1.f;
+	constexpr float YawReplicateTolerance = 0.004363323f;
+}; //namespace UE::Mass::Replication
+
+
+//////////////////////////////////////////////////////////////////////////// FReplicatedAgentPositionYawData ////////////////////////////////////////////////////////////////////////////
+/**
+ * To replicate position and yaw make a member variable of this class in your FReplicatedAgentBase derived class. In the FReplicatedAgentBase derived class you must also provide an accessor function
+ * FReplicatedAgentPathData& GetReplicatedPositionYawDataMutable().
+ */
+USTRUCT()
+struct MASSREPLICATION_API FReplicatedAgentPositionYawData
+{
+	GENERATED_BODY()
+public:
+	
+	FReplicatedAgentPositionYawData()
+		: Position(ForceInitToZero)
+	{}
+
+	void SetPosition(const FVector& InPosition) { Position = InPosition; }
+	const FVector& GetPosition() const { return Position; }
+
+	void SetYaw(const float InYaw) { Yaw = InYaw; }
+	float GetYaw() const { return Yaw; }
+
+private:
+	
+	UPROPERTY(Transient)
+	FVector Position;
+
+	/** Yaw in radians */
+	UPROPERTY(Transient)
+	float Yaw = 0;
+};
+
+//////////////////////////////////////////////////////////////////////////// TMassClientBubbleTransformHandler ////////////////////////////////////////////////////////////////////////////
+/**
+ * To replicate Transforms make a member variable of this class in your TClientBubbleHandlerBase derived class. This class is a friend of TMassClientBubblePathHandler.
+ * It is meant to have access to the protected data declared there.
+ */
+template<typename AgentArrayItem>
+class TMassClientBubbleTransformHandler
+{
+public:
+	TMassClientBubbleTransformHandler(TClientBubbleHandlerBase<AgentArrayItem>& InOwnerHandler)
+		: OwnerHandler(InOwnerHandler)
+	{}
+
+#if UE_REPLICATION_COMPILE_SERVER_CODE
+	/** Sets the position and yaw data in the client bubble on the server */
+	void SetBubblePositionYawFromTransform(const FMassReplicatedAgentHandle Handle, const FTransform& Transform);
+
+	// Another function  SetBubbleTransform() could be added here if required
+#endif // UE_REPLICATION_COMPILE_SERVER_CODE
+
+#if UE_REPLICATION_COMPILE_CLIENT_CODE
+	/**
+	 * When entities are spawned in Mass by the replication system on the client, a spawn query is used to set the data on the spawned entities.
+	 * The following functions are used to configure the query and then set the position and yaw data.
+	 */
+	static void AddRequirementsForSpawnQuery(FLWComponentQuery& InQuery);
+	void CacheComponentViewsForSpawnQuery(FLWComponentSystemExecutionContext& InExecContext);
+	void ClearComponentViewsForSpawnQuery();
+
+	void SetSpawnedEntityData(const int32 EntityIdx, const FReplicatedAgentPositionYawData& ReplicatedPathData) const;
+
+	/** Call this when an Entity that has already been spawned is modified on the client */
+	static void SetModifiedEntityData(const FEntityView& EntityView, const FReplicatedAgentPositionYawData& ReplicatedPathData);
+
+	// We could easily add support replicating FReplicatedAgentTransformData here if required
+#endif // UE_REPLICATION_COMPILE_CLIENT_CODE
+
+protected:
+#if UE_REPLICATION_COMPILE_CLIENT_CODE
+	static void SetEntityData(FDataFragment_Transform& TransformFragment, const FReplicatedAgentPositionYawData& ReplicatedPositionYawData);
+#endif // UE_REPLICATION_COMPILE_CLIENT_CODE
+
+protected:
+	TArrayView<FDataFragment_Transform> TransformList;
+
+	TClientBubbleHandlerBase<AgentArrayItem>& OwnerHandler;
+};
+
+#if UE_REPLICATION_COMPILE_SERVER_CODE
+template<typename AgentArrayItem>
+void TMassClientBubbleTransformHandler<AgentArrayItem>::SetBubblePositionYawFromTransform(const FMassReplicatedAgentHandle Handle, const FTransform& Transform)
+{
+	check(OwnerHandler.AgentHandleManager.IsValidHandle(Handle));
+
+	const int32 AgentsIdx = OwnerHandler.AgentLookupArray[Handle.GetIndex()].AgentsIdx;
+	bool bMarkDirty = false;
+
+	AgentArrayItem& Item = (*OwnerHandler.Agents)[AgentsIdx];
+
+	checkf(Item.Agent.GetNetID().IsValid(), TEXT("Pos should not be updated on FCrowdFastArrayItem's that have an Invalid ID! First Add the Agent!"));
+
+	// GetReplicatedPositionYawDataMutable() must be defined in your FReplicatedAgentBase derived class
+	FReplicatedAgentPositionYawData& ReplicatedPositionYaw = Item.Agent.GetReplicatedPositionYawDataMutable();
+
+	// Only update the Pos and mark the item as dirty if it has changed more than the tolerance
+	const FVector Pos = Transform.GetLocation();
+	if (!Pos.Equals(ReplicatedPositionYaw.GetPosition(), UE::Mass::Replication::PositionReplicateTolerance))
+	{
+		ReplicatedPositionYaw.SetPosition(Pos);
+		bMarkDirty = true;
+	}
+
+	if (const TOptional<float> Yaw = UE::AI::GetYawFromQuaternion(Transform.GetRotation()))
+	{
+		// Only update the Yaw and mark the item as dirty if it has changed more than the tolerance
+		if (FMath::Abs(FMath::FindDeltaAngleRadians(Yaw.GetValue(), ReplicatedPositionYaw.GetYaw())) > UE::Mass::Replication::YawReplicateTolerance)
+		{
+			ReplicatedPositionYaw.SetYaw(Yaw.GetValue());
+			bMarkDirty = true;
+		}
+	}
+
+	if (bMarkDirty)
+	{
+		OwnerHandler.Serializer->MarkItemDirty(Item);
+	}
+
+}
+#endif //UE_REPLICATION_COMPILE_SERVER_CODE
+
+#if UE_REPLICATION_COMPILE_CLIENT_CODE
+template<typename AgentArrayItem>
+void TMassClientBubbleTransformHandler<AgentArrayItem>::AddRequirementsForSpawnQuery(FLWComponentQuery& InQuery)
+{
+	InQuery.AddRequirement<FDataFragment_Transform>(ELWComponentAccess::ReadWrite);
+}
+#endif // UE_REPLICATION_COMPILE_CLIENT_CODE
+
+#if UE_REPLICATION_COMPILE_CLIENT_CODE
+template<typename AgentArrayItem>
+void TMassClientBubbleTransformHandler<AgentArrayItem>::CacheComponentViewsForSpawnQuery(FLWComponentSystemExecutionContext& InExecContext)
+{
+	TransformList = InExecContext.GetMutableComponentView<FDataFragment_Transform>();
+}
+#endif // UE_REPLICATION_COMPILE_CLIENT_CODE
+
+#if UE_REPLICATION_COMPILE_CLIENT_CODE
+template<typename AgentArrayItem>
+void TMassClientBubbleTransformHandler<AgentArrayItem>::ClearComponentViewsForSpawnQuery()
+{
+	TransformList = TArrayView<FDataFragment_Transform>();
+}
+#endif // UE_REPLICATION_COMPILE_CLIENT_CODE
+
+#if UE_REPLICATION_COMPILE_CLIENT_CODE
+template<typename AgentArrayItem>
+void TMassClientBubbleTransformHandler<AgentArrayItem>::SetSpawnedEntityData(const int32 EntityIdx, const FReplicatedAgentPositionYawData& ReplicatedPositionYawData) const
+{
+	FDataFragment_Transform& TransformFragment = TransformList[EntityIdx];
+
+	SetEntityData(TransformFragment, ReplicatedPositionYawData);
+}
+#endif // UE_REPLICATION_COMPILE_CLIENT_CODE
+
+#if UE_REPLICATION_COMPILE_CLIENT_CODE
+template<typename AgentArrayItem>
+void TMassClientBubbleTransformHandler<AgentArrayItem>::SetModifiedEntityData(const FEntityView& EntityView, const FReplicatedAgentPositionYawData& ReplicatedPositionYawData)
+{
+	FDataFragment_Transform& TransformFragment = EntityView.GetComponentData<FDataFragment_Transform>();
+
+	SetEntityData(TransformFragment, ReplicatedPositionYawData);
+}
+#endif // UE_REPLICATION_COMPILE_CLIENT_CODE
+
+#if UE_REPLICATION_COMPILE_CLIENT_CODE
+template<typename AgentArrayItem>
+void TMassClientBubbleTransformHandler<AgentArrayItem>::SetEntityData(FDataFragment_Transform& TransformFragment, const FReplicatedAgentPositionYawData& ReplicatedPositionYawData)
+{
+	TransformFragment.GetMutableTransform().SetLocation(ReplicatedPositionYawData.GetPosition());
+	TransformFragment.GetMutableTransform().SetRotation(FQuat(FVector::UpVector, ReplicatedPositionYawData.GetYaw()));
+}
+#endif // UE_REPLICATION_COMPILE_CLIENT_CODE
+
+//////////////////////////////////////////////////////////////////////////// FMassReplicationProcessorTransformHandlerBase ////////////////////////////////////////////////////////////////////////////
+
+class MASSREPLICATION_API FMassReplicationProcessorTransformHandlerBase
+{
+public:
+	static void AddRequirements(FLWComponentQuery& InQuery);
+	void CacheComponentViews(FLWComponentSystemExecutionContext& ExecContext);
+
+protected:
+	TArrayView<FDataFragment_Transform> TransformList;
+};
+
+//////////////////////////////////////////////////////////////////////////// FMassReplicationProcessorPositionYawHandler ////////////////////////////////////////////////////////////////////////////
+/**
+ * Used to replicate position and yaw by your UMassReplicationProcessorBase derived class. This class should only get used on the server.
+ * @todo add #if UE_REPLICATION_COMPILE_SERVER_CODE
+ */
+class MASSREPLICATION_API FMassReplicationProcessorPositionYawHandler : public FMassReplicationProcessorTransformHandlerBase
+{
+public:
+	void AddEntity(const int32 EntityIdx, FReplicatedAgentPositionYawData& InOUtReplicatedPathData) const;
+
+	template<typename AgentArrayItem>
+	void ModifyEntity(const FMassReplicatedAgentHandle Handle, const int32 EntityIdx, TMassClientBubbleTransformHandler<AgentArrayItem>& BubblePathHandler);
+};
+
+template<typename AgentArrayItem>
+void FMassReplicationProcessorPositionYawHandler::ModifyEntity(const FMassReplicatedAgentHandle Handle, const int32 EntityIdx, TMassClientBubbleTransformHandler<AgentArrayItem>& BubbleTransformHandler)
+{
+	const FDataFragment_Transform& TransformFragment = TransformList[EntityIdx];
+
+	BubbleTransformHandler.SetBubblePositionYawFromTransform(Handle, TransformFragment.GetTransform());
+}
\ No newline at end of file
diff --git a/Engine/Plugins/Runtime/MassGameplay/Source/MassReplication/Public/MassReplicationTypes.h b/Engine/Plugins/Runtime/MassGameplay/Source/MassReplication/Public/MassReplicationTypes.h
new file mode 100644
--- /dev/null
+++ b/Engine/Plugins/Runtime/MassGameplay/Source/MassReplication/Public/MassReplicationTypes.h
@@ -0,0 +1,248 @@
+// Copyright Epic Games, Inc. All Rights Reserved.
+
+#pragma once
+
+#include "MassCommonTypes.h"
+#include "MassLODManager.h"
+#include "IndexedHandle.h"
+#include "MassEntityTemplate.h"
+#include "MassTranslator.h"
+#include "MassLODLogic.h"
+
+#include "MassReplicationTypes.generated.h"
+
+
+/** Debug option to have bubbles on all player controllers in a stand alone game. This is especially useful for testing and profiling
+ *  these systems. This will also disable the ClientsViewers functionality. If UE_ALLOW_DEBUG_REPLICATION_BUBBLES_STANDALONE is set 
+ *  UE_REPLICATION_COMPILE_SERVER_CODE will be defined as 1 so server only code is compiled in standalone builds. 
+ */ 
+#ifndef UE_DEBUG_REPLICATION_BUBBLES_STANDALONE
+#define UE_DEBUG_REPLICATION_BUBBLES_STANDALONE 0
+#endif
+
+#if UE_DEBUG_REPLICATION_BUBBLES_STANDALONE && !UE_BUILD_SHIPPING && (UE_GAME || UE_EDITOR)
+#define UE_ALLOW_DEBUG_REPLICATION_BUBBLES_STANDALONE 1
+#else
+#define UE_ALLOW_DEBUG_REPLICATION_BUBBLES_STANDALONE 0
+#endif
+
+// WITH_SERVER_CODE is all builds except client only.
+#if WITH_SERVER_CODE
+#define UE_REPLICATION_COMPILE_SERVER_CODE 1
+#else
+#define UE_REPLICATION_COMPILE_SERVER_CODE 0
+#endif
+
+// If not a dedicated server build then we may need client code
+#if !UE_SERVER
+#define UE_REPLICATION_COMPILE_CLIENT_CODE 1
+#else
+#define UE_REPLICATION_COMPILE_CLIENT_CODE 0
+#endif
+
+#ifndef UE_DEBUG_SLOW_REPLICATION
+#define UE_DEBUG_SLOW_REPLICATION 0
+#endif
+
+#if UE_DEBUG_SLOW_REPLICATION && !UE_BUILD_SHIPPING
+#define UE_ALLOW_DEBUG_SLOW_REPLICATION 1
+#else
+#define UE_ALLOW_DEBUG_SLOW_REPLICATION 0
+#endif
+
+#ifndef UE_DEBUG_REPLICATION
+#define UE_DEBUG_REPLICATION 0
+#endif
+
+#if (UE_DEBUG_REPLICATION && !UE_BUILD_SHIPPING) || UE_ALLOW_DEBUG_SLOW_REPLICATION
+#define UE_ALLOW_DEBUG_REPLICATION 1
+#else
+#define UE_ALLOW_DEBUG_REPLICATION 0
+#endif
+
+class UMassReplicationManager;
+
+namespace UE { namespace Mass { namespace Replication
+{
+	const int32 MaxNumOfClients = UE::MassLOD::MaxNumOfViewers;
+}}}
+
+DECLARE_LOG_CATEGORY_EXTERN(LogMassReplication, Log, All);
+
+/** Handle of each Agent per Client (bubble), these handles are not consistent between different client bubbles.
+ *  Using the compact handle here will make the memory footprint smaller and help the ECS system be more cache friendly.
+ *	TODO: in the long run we can probably use a non serial number handle here, I'd like to stick with the serial number for now though
+ */
+USTRUCT()
+struct MASSREPLICATION_API FMassReplicatedAgentHandle : public FCompactIndexedHandleBase
+{
+	GENERATED_BODY()
+
+	FMassReplicatedAgentHandle() = default;
+
+	/** @param InIndex - passing INDEX_NONE will make this handle Invalid */
+	FMassReplicatedAgentHandle(int32 InIndex, uint32 InSerialNumber) : FCompactIndexedHandleBase(InIndex, InSerialNumber)
+	{
+	};
+};
+
+typedef FCompactIndexedHandleManager<FMassReplicatedAgentHandle> FMassReplicatedAgentHandleManager;
+
+//Handle of each Client on the server. A Client is defined as a player controller with a non nullptr parent NetConnection
+USTRUCT()
+struct MASSREPLICATION_API FMassClientHandle : public FIndexedHandleBase
+{
+	GENERATED_BODY()
+
+	FMassClientHandle() = default;
+
+	/** @param InIndex - passing INDEX_NONE will make this handle Invalid */
+	FMassClientHandle(int32 InIndex, uint32 InSerialNumber) : FIndexedHandleBase(InIndex, InSerialNumber)
+	{
+	};
+};
+
+typedef FIndexedHandleManager<FMassClientHandle> FMassClientHandleManager;
+
+USTRUCT()
+struct MASSREPLICATION_API FMassBubbleInfoClassHandle : public FSimpleIndexedHandleBase
+{
+	GENERATED_BODY()
+
+	FMassBubbleInfoClassHandle() = default;
+
+	/** @param InIndex - passing INDEX_NONE will make this handle Invalid */
+	FMassBubbleInfoClassHandle(int32 InIndex) : FSimpleIndexedHandleBase(InIndex)
+	{
+	};
+};
+
+/**
+ *  Light weight component data for the mass network id of a light weight entity
+ */
+USTRUCT()
+struct MASSREPLICATION_API FMassNetworkIDFragment : public FLWComponentData
+{
+	GENERATED_BODY()
+
+	FMassNetworkID NetID;
+};
+
+struct FMassReplicatedAgentArrayData
+{
+	void Invalidate()
+	{
+		Handle.Invalidate();
+		LastUpdateTime = 0.f;
+#if UE_ALLOW_DEBUG_REPLICATION
+		LOD = EMassLOD::Off;
+#endif // UE_DO_DEBUG_REPLICATION
+	}
+
+	FMassReplicatedAgentHandle Handle;
+	float LastUpdateTime = 0.f;
+
+#if UE_ALLOW_DEBUG_REPLICATION
+	EMassLOD::Type LOD = EMassLOD::Off;
+#endif // UE_DO_DEBUG_REPLICATION
+};
+
+/**
+ * Agent handle per client, these will be at TArray indices of the Client handles indicies (used as a free list array)
+ */
+USTRUCT()
+struct MASSREPLICATION_API FMassReplicatedAgentFragment : public FLWComponentData
+{
+	GENERATED_BODY()
+
+	//these will be maintained so that the TArray entries are at the same indices as Client Handles
+	TArray<FMassReplicatedAgentArrayData, TInlineAllocator<UE::Mass::Replication::MaxNumOfClients>> AgentsData;
+};
+
+USTRUCT()
+struct FReplicatedAgentBase
+{
+	GENERATED_BODY()
+
+	FReplicatedAgentBase() = default;
+
+	FReplicatedAgentBase(FMassNetworkID InNetID, FMassEntityTemplateID InTemplateID)
+		: NetID(InNetID)
+		, TemplateID(InTemplateID)
+	{}
+
+	void SetNetID(FMassNetworkID InNetID) { NetID = InNetID; }
+	FMassNetworkID GetNetID() const { return NetID; }
+
+	void SetTemplateID(FMassEntityTemplateID InTemplateID) { TemplateID = InTemplateID; }
+	FMassEntityTemplateID GetTemplateID() const { return TemplateID; }
+
+private:
+	UPROPERTY()
+	FMassNetworkID NetID;
+
+	UPROPERTY()
+	FMassEntityTemplateID TemplateID;
+};
+
+struct FReplicationLODLogic : public FLODDefaultLogic
+{
+	enum
+	{
+		bStoreLODPerViewer = true,
+		bMaximizeCountPerViewer = true,
+	};
+};
+
+USTRUCT()
+struct MASSREPLICATION_API FMassReplicationViewerLODFragment : public FLWComponentData
+{
+	GENERATED_BODY()
+
+	/** Saved closest ViewerDistance */
+	float ClosestViewerDistanceSq = FLT_MAX;
+
+	/**LOD information */
+	TStaticArray<TEnumAsByte<EMassLOD::Type>, UE::MassLOD::MaxNumOfViewers> LODPerViewer;
+	TEnumAsByte<EMassLOD::Type> LOD = EMassLOD::Max;
+
+	/** Previous LOD information*/
+	TStaticArray<TEnumAsByte<EMassLOD::Type>, UE::MassLOD::MaxNumOfViewers> PrevLODPerViewer;
+	TEnumAsByte<EMassLOD::Type> PrevLOD = EMassLOD::Max;
+};
+
+UCLASS()
+class MASSREPLICATION_API UMassNetworkIDFragmentInitializer : public UMassFragmentInitializer
+{
+	GENERATED_BODY()
+public:
+	UMassNetworkIDFragmentInitializer();
+
+protected:
+	virtual void ConfigureQueries() override;
+	virtual void Initialize(UObject& Owner) override;
+	virtual void Execute(UEntitySubsystem& EntitySubsystem, FLWComponentSystemExecutionContext& Context) override;
+
+	UPROPERTY()
+	UMassReplicationManager* ReplicationManager = nullptr;
+
+	FLWComponentQuery EntityQuery;
+};
+
+#if UE_REPLICATION_COMPILE_CLIENT_CODE
+/** Data that can be accessed from a FMassNetworkID on a client */
+struct FMassReplicationEntityInfo
+{
+	FMassReplicationEntityInfo() = default;
+	FMassReplicationEntityInfo(FLWEntity InEntity, int32 InReplicationID)
+		: Entity(InEntity)
+		, ReplicationID(InReplicationID)
+	{}
+
+	/** If this is not IsSet() then the entity has been removed from the client simulation */
+	FLWEntity Entity; 
+
+	/** This is stored between removes and adds, however this item in the UMassReplicationManager::EntityInfoMap will eventually get cleaned up if Entity.IsSet() == false for a fairly substantial length of time. */
+	int32 ReplicationID = INDEX_NONE;
+};
+#endif //UE_REPLICATION_COMPILE_CLIENT_CODE
\ No newline at end of file
diff --git a/Engine/Plugins/Runtime/MassGameplay/Source/MassRepresentation/Private/MassRepresentationModule.cpp b/Engine/Plugins/Runtime/MassGameplay/Source/MassRepresentation/Private/MassRepresentationModule.cpp
new file mode 100644
--- /dev/null
+++ b/Engine/Plugins/Runtime/MassGameplay/Source/MassRepresentation/Private/MassRepresentationModule.cpp
@@ -0,0 +1,41 @@
+// Copyright Epic Games, Inc. All Rights Reserved.
+
+#include "CoreMinimal.h"
+#include "IMassRepresentationModule.h"
+#include "UObject/CoreRedirects.h"
+
+class FMassRepresentationModule : public IMassRepresentationModule
+{
+	/** IModuleInterface implementation */
+	virtual void StartupModule() override;
+	virtual void ShutdownModule() override;
+};
+
+IMPLEMENT_MODULE(FMassRepresentationModule, MassRepresentation)
+
+
+
+void FMassRepresentationModule::StartupModule()
+{
+	TArray<FCoreRedirect> Redirects;
+
+	Redirects.Emplace(ECoreRedirectFlags::Type_Struct, TEXT("/Script/MassVisualization.DataFragment_Visualization"), TEXT("/Script/MassRepresentation.MassRepresentationFragment"));
+	Redirects.Emplace(ECoreRedirectFlags::Type_Class, TEXT("/Script/MassVisualization.MassProcessor_Visualization"), TEXT("/Script/MassRepresentation.MassRepresentationProcessor"));
+	Redirects.Emplace(ECoreRedirectFlags::Type_Class, TEXT("/Script/MassVisualization.MassFragmentDestructor_Visualization"), TEXT("/Script/MassRepresentation.MassRepresentationFragmentDestructor"));
+	Redirects.Emplace(ECoreRedirectFlags::Type_Class, TEXT("/Script/MassVisualization.MassVisualizationLODProcessor"), TEXT("/Script/MassRepresentation.MassVisualizationLODProcessor"));
+	Redirects.Emplace(ECoreRedirectFlags::Type_Class, TEXT("/Script/MassRepresentation.MassRepresentationLODProcessor"), TEXT("/Script/MassRepresentation.MassVisualizationLODProcessor"));
+	Redirects.Emplace(ECoreRedirectFlags::Type_Class, TEXT("/Script/MassVisualization.MassVisualizationFeature"), TEXT("/Script/MassRepresentation.MassVisualizationTrait"));
+	Redirects.Emplace(ECoreRedirectFlags::Type_Property, TEXT("/Script/MassRepresentation.MassVisualizationTrait.TemplateActor"), TEXT("/Script/MassRepresentation.MassVisualizationTrait.HighResTemplateActor"));
+
+	FCoreRedirects::AddRedirectList(Redirects, TEXT("MassRepresentation"));
+}
+
+
+void FMassRepresentationModule::ShutdownModule()
+{
+	// This function may be called during shutdown to clean up your module.  For modules that support dynamic reloading,
+	// we call this function before unloading the module.
+}
+
+
+
