AggregateTicking and Pipe merged into Engine/Plugins/Runtime/MassEntity

diff --git a/Engine/Plugins/Runtime/MassGameplay/Source/MassRepresentation/Private/MassRepresentationProcessor.cpp b/Engine/Plugins/Runtime/MassGameplay/Source/MassRepresentation/Private/MassRepresentationProcessor.cpp
new file mode 100644
--- /dev/null
+++ b/Engine/Plugins/Runtime/MassGameplay/Source/MassRepresentation/Private/MassRepresentationProcessor.cpp
@@ -0,0 +1,493 @@
+// Copyright Epic Games, Inc. All Rights Reserved.
+
+#include "MassRepresentationProcessor.h"
+#include "MassRepresentationSubsystem.h"
+#include "MassRepresentationUtils.h"
+#include "MassActorSubsystem.h"
+#include "MassCommonFragments.h"
+#include "LWCCommandBuffer.h"
+#include "MassEntitySystem.h"
+#include "EntityView.h"
+#include "Engine/World.h"
+
+//----------------------------------------------------------------------//
+// UMassRepresentationProcessor(); 
+//----------------------------------------------------------------------//
+UMassRepresentationProcessor::UMassRepresentationProcessor()
+{
+	bAutoRegisterWithProcessingPhases = false;
+
+	ExecutionOrder.ExecuteInGroup = UE::Mass::ProcessorGroupNames::Representation;
+
+	LODRepresentation[EMassLOD::High] = ERepresentationType::HighResSpawnedActor;
+	LODRepresentation[EMassLOD::Medium] = ERepresentationType::LowResSpawnedActor;
+	LODRepresentation[EMassLOD::Low] = ERepresentationType::StaticMeshInstance;
+	LODRepresentation[EMassLOD::Off] = ERepresentationType::None;
+}
+
+void UMassRepresentationProcessor::ConfigureQueries()
+{
+	EntityQuery.AddRequirement<FDataFragment_Transform>(ELWComponentAccess::ReadOnly);
+	EntityQuery.AddRequirement<FMassRepresentationFragment>(ELWComponentAccess::ReadWrite);
+	EntityQuery.AddRequirement<FMassRepresentationLODFragment>(ELWComponentAccess::ReadOnly);
+	EntityQuery.AddRequirement<FDataFragment_Actor>(ELWComponentAccess::ReadWrite);
+	EntityQuery.AddChunkRequirement<FMassVisualizationChunkFragment>(ELWComponentAccess::ReadWrite);
+}
+
+void UMassRepresentationProcessor::Initialize(UObject& Owner)
+{
+	Super::Initialize(Owner);
+
+	World = Owner.GetWorld();
+	PipeEntitySubsystem = UPipeEntitySubsystem::GetCurrent(World);
+	RepresentationSubsystem = UWorld::GetSubsystem<UMassRepresentationSubsystem>(World);
+
+	// Calculate the default representation when actor isn't spawned yet.
+	for (int32 LOD = EMassLOD::High; LOD < EMassLOD::Max; LOD++)
+	{
+		// Find the first representation type after any actors
+		if (LODRepresentation[LOD] == ERepresentationType::HighResSpawnedActor ||
+			LODRepresentation[LOD] == ERepresentationType::LowResSpawnedActor)
+		{
+			continue;
+		}
+
+		DefaultRepresentationType = LODRepresentation[LOD];
+		break;
+	}
+}
+
+void UMassRepresentationProcessor::UpdateRepresentation(FLWComponentSystemExecutionContext& Context)
+{
+	check(RepresentationSubsystem);
+
+	const TConstArrayView<FDataFragment_Transform> TransformList = Context.GetComponentView<FDataFragment_Transform>();
+	const TArrayView<FMassRepresentationFragment> RepresentationList = Context.GetMutableComponentView<FMassRepresentationFragment>();
+	const TConstArrayView<FMassRepresentationLODFragment> RepresentationLODList = Context.GetComponentView<FMassRepresentationLODFragment>();
+	const TArrayView<FDataFragment_Actor> ActorList = Context.GetMutableComponentView<FDataFragment_Actor>();
+
+	const int32 NumEntities = Context.GetEntitiesNum();
+	for (int32 EntityIdx = 0; EntityIdx < NumEntities; EntityIdx++)
+	{
+		const FLWEntity Entity = Context.GetEntity(EntityIdx);
+		const FMassHandle MassAgent(Entity);
+		const FDataFragment_Transform& TransformFragment = TransformList[EntityIdx];
+		const FMassRepresentationLODFragment& RepresentationLOD = RepresentationLODList[EntityIdx];
+		FMassRepresentationFragment& Representation = RepresentationList[EntityIdx];
+		FDataFragment_Actor& ActorInfo = ActorList[EntityIdx];
+
+		Representation.PrevRepresentation = Representation.CurrentRepresentation;
+		
+		ERepresentationType WantedRepresentationType = LODRepresentation[FMath::Min((int32)RepresentationLOD.LOD, (int32)EMassLOD::Off)];
+
+		// Make sure we do not have actor spawned in areas not fully loaded
+		if ((WantedRepresentationType == ERepresentationType::HighResSpawnedActor || WantedRepresentationType == ERepresentationType::LowResSpawnedActor) &&
+			!RepresentationSubsystem->IsCollisionLoaded(WorldPartitionGridNameContainingCollision, TransformFragment.GetTransform()))
+		{
+			WantedRepresentationType = DefaultRepresentationType;
+		}
+
+		// Process switch between representation if there is a change in the representation or there is a pending spawning request
+		if (WantedRepresentationType != Representation.CurrentRepresentation || Representation.ActorSpawnRequestHandle.IsValid())
+		{
+			if (Representation.CurrentRepresentation == ERepresentationType::None)
+			{
+				Representation.PrevTransform = TransformFragment.GetTransform();
+			}
+
+			AActor* Actor = ActorInfo.GetMutable();
+			switch (WantedRepresentationType)
+			{
+				case ERepresentationType::HighResSpawnedActor:
+				case ERepresentationType::LowResSpawnedActor:
+				{
+					const bool bHighResActor = WantedRepresentationType == ERepresentationType::HighResSpawnedActor;
+
+					// Reuse actor, if it is valid and not owned by mass or same representation as low res without a valid spawning request
+					AActor* NewActor = nullptr;
+					if (!Actor || ActorInfo.IsOwnedByMass())
+					{
+						const int32 WantedTemplateActorIndex = bHighResActor ? Representation.HighResTemplateActorIndex : Representation.LowResTemplateActorIndex;
+
+						// If the low res is different than the high res, cancel any pending spawn request that is the opposite of what is needed.
+						if (Representation.LowResTemplateActorIndex != Representation.HighResTemplateActorIndex)
+						{
+							ReleaseActorOrCancelSpawning(MassAgent, ActorInfo, bHighResActor ? Representation.LowResTemplateActorIndex : Representation.HighResTemplateActorIndex, Representation.ActorSpawnRequestHandle, Context, /*bCancelSpawningOnly*/true);
+							Actor = ActorInfo.GetOwnedByMassMutable();
+						}
+
+						// If there isn't any actor yet or
+						// If the actor isn't matching the one needed or
+						// If there is still a pending spawn request
+						// Then try to retrieve/spawn the new actor
+						if (!Actor ||
+							!RepresentationSubsystem->DoesActorMatchTemplate(*Actor, WantedTemplateActorIndex) ||
+							Representation.ActorSpawnRequestHandle.IsValid())
+						{
+							NewActor = GetOrSpawnActor(MassAgent, ActorInfo, TransformFragment.GetTransform(), WantedTemplateActorIndex, Representation.ActorSpawnRequestHandle, GetSpawnPriority(RepresentationLOD));
+						}
+						else
+						{
+							NewActor = Actor;
+						}
+					}
+					else
+					{
+						NewActor = Actor;
+					}
+
+					if (NewActor)
+					{
+						// Make sure our (re)activated actor is at the simulated position
+						// Needs to be done before enabling the actor so the animation initialization can use the new values
+						if (Representation.CurrentRepresentation == ERepresentationType::StaticMeshInstance)
+						{
+							TeleportActor(Representation.PrevTransform, *NewActor, Context);
+						}
+
+						SetActorEnabled(bHighResActor ? EActorEnabledType::HighRes : EActorEnabledType::LowRes, *NewActor, EntityIdx, Context);
+						Representation.CurrentRepresentation = WantedRepresentationType;
+					}
+					else if (!Actor)
+					{
+						Representation.CurrentRepresentation = DefaultRepresentationType;
+					}
+					break;
+				}
+				case ERepresentationType::StaticMeshInstance:
+					if (!Actor || ActorInfo.IsOwnedByMass())
+					{
+						// If high res actor if it is different than the low res, 
+						if (Representation.HighResTemplateActorIndex != Representation.LowResTemplateActorIndex|| !bKeepLowResActors)
+						{
+							// Try releasing the high actor or any hig res spawning request
+							if (ReleaseActorOrCancelSpawning(MassAgent, ActorInfo, Representation.HighResTemplateActorIndex, Representation.ActorSpawnRequestHandle, Context))
+							{
+								Actor = ActorInfo.GetOwnedByMassMutable();
+							}
+
+							// Do not do the same with low res if indicated so
+							if (!bKeepLowResActors && ReleaseActorOrCancelSpawning(MassAgent, ActorInfo, Representation.LowResTemplateActorIndex, Representation.ActorSpawnRequestHandle, Context))
+							{
+								Actor = ActorInfo.GetOwnedByMassMutable();
+							}
+						}
+
+						// If we were in the process of spawning a low res actor, continue and keep it around disabled
+						if (Representation.ActorSpawnRequestHandle.IsValid())
+						{
+							Actor = GetOrSpawnActor(MassAgent, ActorInfo, TransformFragment.GetTransform(), Representation.LowResTemplateActorIndex, Representation.ActorSpawnRequestHandle, GetSpawnPriority(RepresentationLOD));
+						}
+					}
+
+					if (Actor != nullptr)
+					{
+						SetActorEnabled(EActorEnabledType::Disabled, *Actor, EntityIdx, Context);
+					}
+					Representation.CurrentRepresentation = ERepresentationType::StaticMeshInstance;
+					break;
+				case ERepresentationType::None:
+					if (!Actor || ActorInfo.IsOwnedByMass())
+					{
+						// Try releasing both, could have an high res spawned actor and a spawning request for a low res one
+						ReleaseActorOrCancelSpawning(MassAgent, ActorInfo, Representation.LowResTemplateActorIndex, Representation.ActorSpawnRequestHandle, Context);
+						ReleaseActorOrCancelSpawning(MassAgent, ActorInfo, Representation.HighResTemplateActorIndex, Representation.ActorSpawnRequestHandle, Context);
+					}
+					else
+					{
+						SetActorEnabled(EActorEnabledType::Disabled, *Actor, EntityIdx, Context);
+					}
+					Representation.CurrentRepresentation = ERepresentationType::None;
+					break;
+				default:
+					checkf(false, TEXT("Unsupported LOD type"));
+					break;
+			}
+		}
+	}
+}
+
+void UMassRepresentationProcessor::Execute(UEntitySubsystem& EntitySubsystem, FLWComponentSystemExecutionContext& Context)
+{
+	check(RepresentationSubsystem);
+
+	// Visualize entities
+	EntityQuery.ForEachEntityChunk(EntitySubsystem, Context, [this](FLWComponentSystemExecutionContext& Context)
+	{
+		UpdateRepresentation(Context);
+	});
+}
+
+AActor* UMassRepresentationProcessor::GetOrSpawnActor(const FMassHandle MassAgent, FDataFragment_Actor& ActorInfo, const FTransform& Transform, const int16 TemplateActorIndex, FMassHandle_ActorSpawnRequest& SpawnRequestHandle, const float Priority)
+{
+	check(RepresentationSubsystem);
+
+	return RepresentationSubsystem->GetOrSpawnActorFromTemplate(MassAgent, Transform, TemplateActorIndex, SpawnRequestHandle, Priority, 
+		FMassActorPreSpawnDelegate::CreateUObject(this, &UMassRepresentationProcessor::OnActorPreSpawn), 
+		FMassActorPostSpawnDelegate::CreateUObject(this, &UMassRepresentationProcessor::OnActorPostSpawn));
+}
+
+bool UMassRepresentationProcessor::ReleaseActorOrCancelSpawning(const FMassHandle MassAgent, FDataFragment_Actor& ActorInfo, const int16 TemplateActorIndex, FMassHandle_ActorSpawnRequest& SpawnRequestHandle, FLWComponentSystemExecutionContext& Context, bool bCancelSpawningOnly /*= false*/)
+{
+	check(RepresentationSubsystem);
+	check(!ActorInfo.IsValid() || ActorInfo.IsOwnedByMass());
+
+	AActor* Actor = ActorInfo.GetOwnedByMassMutable();
+	// note that it's fine for Actor to be null. That means the RepresentationSubsystem will try to stop 
+	// the spawning of whatever SpawnRequestHandle reference to
+	const bool bSuccess = bCancelSpawningOnly ? RepresentationSubsystem->CancelSpawning(MassAgent, TemplateActorIndex, SpawnRequestHandle) :
+			RepresentationSubsystem->ReleaseTemplateActorOrCancelSpawning(MassAgent, TemplateActorIndex, Actor, SpawnRequestHandle);
+	if(bSuccess)
+	{
+		Actor = ActorInfo.GetOwnedByMassMutable();
+		if(Actor && RepresentationSubsystem->DoesActorMatchTemplate(*Actor, TemplateActorIndex))
+		{
+			ActorInfo.ResetNoHandleMapUpdate();
+			
+			TObjectKey<const AActor> ActorKey(Actor); 
+			if (UMassActorSubsystem* MassActorSubsystem = UWorld::GetSubsystem<UMassActorSubsystem>(World))
+			{
+				Context.Defer().EmplaceCommand<FDeferredCommand>([MassActorSubsystem, ActorKey](UEntitySubsystem&)
+				{
+					MassActorSubsystem->RemoveHandleForActor(ActorKey);
+				});
+			}
+		}
+		return true;
+	}
+	return false;
+}
+
+void UMassRepresentationProcessor::SetActorEnabled(const EActorEnabledType EnabledType, AActor& Actor, const int32 EntityIdx, FLWComponentSystemExecutionContext& Context)
+{
+	const bool bEnabled = EnabledType != EActorEnabledType::Disabled;
+	if (Actor.IsActorTickEnabled() != bEnabled)
+	{
+		Actor.SetActorTickEnabled(bEnabled);
+	}
+	if (Actor.GetActorEnableCollision() != bEnabled)
+	{
+		// Deferring this as there is a callback internally that could end up doing things outside of the game thread and will fire checks(Chaos mostly)
+		Context.Defer().EmplaceCommand<FDeferredCommand>( [&Actor,bEnabled](UEntitySubsystem&)
+		{
+			Actor.SetActorEnableCollision(bEnabled);
+		});
+	}
+}
+
+void UMassRepresentationProcessor::TeleportActor(const FTransform& Transform, AActor& Actor, FLWComponentSystemExecutionContext& Context)
+{
+	if (!Actor.GetTransform().Equals(Transform))
+	{
+		Context.Defer().EmplaceCommand<FDeferredCommand>([&Actor, Transform](UEntitySubsystem&)
+		{
+			Actor.SetActorTransform(Transform, /*bSweep*/false, /*OutSweepHitResult*/nullptr, ETeleportType::TeleportPhysics);
+		});
+	}
+}
+
+void UMassRepresentationProcessor::ReleaseAnyActorOrCancelAnySpawning(UMassRepresentationSubsystem& RepresentationSubsystem, const FMassHandle MassAgent)
+{
+	if (const UPipeEntitySubsystem* EntitySubsystem = UPipeEntitySubsystem::GetCurrent(RepresentationSubsystem.GetWorld()))
+	{
+		if (FDataFragment_Actor* ActorInfo = EntitySubsystem->GetComponentDataPtr<FDataFragment_Actor>(MassAgent.GetLWEntity()))
+		{
+			if (FMassRepresentationFragment* Representation = EntitySubsystem->GetComponentDataPtr<FMassRepresentationFragment>(MassAgent.GetLWEntity()))
+			{
+				ReleaseAnyActorOrCancelAnySpawning(RepresentationSubsystem, MassAgent, *ActorInfo, *Representation);
+			}
+		}
+	}
+}
+
+void UMassRepresentationProcessor::ReleaseAnyActorOrCancelAnySpawning(UMassRepresentationSubsystem& RepresentationSubsystem, const FMassHandle MassAgent, FDataFragment_Actor& ActorInfo, FMassRepresentationFragment& Representation)
+{
+	// This method can only release owned by mass actors
+	AActor* Actor = ActorInfo.GetOwnedByMassMutable();
+	if (Actor)
+	{
+		// WARNING!
+		// Need to reset before ReleaseTemplateActorOrCancelSpawning as this action might move the entity to a new archetype and
+		// so the Fragment passed in parameters would not be valid anymore.
+		ActorInfo.ResetAndUpdateHandleMap();
+	}
+	// Try releasing both as we can have a low res actor and a high res spawning request
+	RepresentationSubsystem.ReleaseTemplateActorOrCancelSpawning(MassAgent, Representation.HighResTemplateActorIndex, Actor, Representation.ActorSpawnRequestHandle);
+	if (Representation.LowResTemplateActorIndex != Representation.HighResTemplateActorIndex)
+	{
+		RepresentationSubsystem.ReleaseTemplateActorOrCancelSpawning(MassAgent, Representation.LowResTemplateActorIndex, Actor, Representation.ActorSpawnRequestHandle);
+	}
+}
+
+void UMassRepresentationProcessor::UpdateVisualization(FLWComponentSystemExecutionContext& Context)
+{
+	FMassVisualizationChunkFragment& ChunkData = UpdateChunkVisibility(Context);
+	if (!ChunkData.ShouldUpdateVisualization())
+	{
+		return;
+	}
+
+	UpdateRepresentation(Context);
+
+	// Update entity visibility
+	const TArrayView<FMassRepresentationFragment> RepresentationList = Context.GetMutableComponentView<FMassRepresentationFragment>();
+	const TConstArrayView<FMassRepresentationLODFragment> RepresentationLODList = Context.GetComponentView<FMassRepresentationLODFragment>();
+
+	const int32 NumEntities = Context.GetEntitiesNum();
+	for (int32 EntityIdx = 0; EntityIdx < NumEntities; EntityIdx++)
+	{
+		const FLWEntity Entity = Context.GetEntity(EntityIdx);
+		FMassRepresentationFragment& Representation = RepresentationList[EntityIdx];
+		const FMassRepresentationLODFragment& RepresentationLOD = RepresentationLODList[EntityIdx];
+		UpdateEntityVisibility(Entity, Representation, RepresentationLOD, ChunkData, Context);
+	}
+}
+
+FMassVisualizationChunkFragment& UMassRepresentationProcessor::UpdateChunkVisibility(FLWComponentSystemExecutionContext& Context) const
+{
+	// Setup chunk component data about visibility
+	FMassVisualizationChunkFragment& ChunkData = Context.GetMutableChunkComponent<FMassVisualizationChunkFragment>();
+	EMassVisibility ChunkVisibility = ChunkData.GetVisibility();
+	if (ChunkVisibility == EMassVisibility::Max)
+	{
+		// The visibility on the chunk component data isn't set yet, let see if the Archetype has an visibility tag and set it on the ChunkData
+		ChunkVisibility = UE::MassRepresentation::GetVisibilityFromArchetype(Context);
+		ChunkData.SetVisibility(ChunkVisibility);
+	}
+	else
+	{
+		checkfSlow(UE::MassRepresentation::IsVisibilityTagSet(Context, ChunkVisibility), TEXT("Expecting the same Visibility as what we saved in the chunk data, maybe external code is modifying the tags"))
+	}
+
+	if (ChunkVisibility == EMassVisibility::CulledByDistance)
+	{
+		float DeltaTime = ChunkData.GetDeltaTime();
+		if (DeltaTime < 0.0f)
+		{
+			DeltaTime += NotVisibleUpdateRate * (1.0f + FMath::RandRange(-0.1f, 0.1f));
+		}
+		DeltaTime -= Context.GetDeltaTimeSeconds();
+
+		ChunkData.Update(DeltaTime);
+	}
+
+	return ChunkData;
+}
+
+void UMassRepresentationProcessor::UpdateEntityVisibility(const FLWEntity Entity, const FMassRepresentationFragment& Representation, const FMassRepresentationLODFragment& RepresentationLOD, FMassVisualizationChunkFragment& ChunkData, FLWComponentSystemExecutionContext& Context)
+{
+	// Move the visible entities together into same chunks so we can skip entire chunk when not visible as an optimization
+	const EMassVisibility Visibility = Representation.CurrentRepresentation != ERepresentationType::None ? EMassVisibility::CanBeSeen : 
+		(RepresentationLOD.bIsInVisibleRange ? EMassVisibility::CulledByFrustum : EMassVisibility::CulledByDistance);
+	const EMassVisibility ChunkVisibility = ChunkData.GetVisibility();
+	if (ChunkVisibility != Visibility)
+	{
+		Context.Defer().PushCommand(FCommandSwapTags(Entity, UE::MassRepresentation::GetTagFromVisibility(ChunkVisibility), UE::MassRepresentation::GetTagFromVisibility(Visibility)));
+		ChunkData.SetContainsNewlyVisibleEntity(Visibility == EMassVisibility::CanBeSeen);
+	}
+}
+
+void UMassRepresentationProcessor::OnActorPreSpawn(const FMassHandle_ActorSpawnRequest& SpawnRequestHandle, const FStructView& SpawnRequest)
+{
+	check(PipeEntitySubsystem);
+	check(RepresentationSubsystem);
+
+	const FMassActorSpawnRequest& MassActorSpawnRequest = SpawnRequest.Get<FMassActorSpawnRequest>();
+	const FEntityView View(*PipeEntitySubsystem, MassActorSpawnRequest.MassAgent.GetLWEntity());
+
+	if (FDataFragment_Actor* ActorInfo = View.GetComponentDataPtr<FDataFragment_Actor>())
+	{
+		// Release any existing actor
+		if (AActor* Actor = ActorInfo->GetMutable())
+		{
+			checkf(ActorInfo->IsOwnedByMass(), TEXT("If we reach here, we expect the actor to be owned by mass, otherwise we should not be spawning a new one one top of this one."));
+
+			// WARNING!
+			// Need to reset before ReleaseTemplateActor as this action might move the entity to a new archetype and
+			// so the Fragment passed in parameters would not be valid anymore.
+			ActorInfo->ResetAndUpdateHandleMap();
+
+			if (const FMassRepresentationFragment* Representation = View.GetComponentDataPtr<FMassRepresentationFragment>())
+			{
+				if (!RepresentationSubsystem->ReleaseTemplateActor(MassActorSpawnRequest.MassAgent, Representation->HighResTemplateActorIndex, Actor, /*bImmediate*/ true))
+				{
+					if (!RepresentationSubsystem->ReleaseTemplateActor(MassActorSpawnRequest.MassAgent, Representation->LowResTemplateActorIndex, Actor, /*bImmediate*/ true))
+					{
+						checkf(false, TEXT("Expecting to be able to release spawned actor either the high res or low res one"));
+					}
+				}
+			}
+		}
+	}
+}
+
+EMassActorSpawnRequestAction UMassRepresentationProcessor::OnActorPostSpawn(const FMassHandle_ActorSpawnRequest& SpawnRequestHandle, const FStructView& SpawnRequest)
+{
+	check(PipeEntitySubsystem);
+
+	const FMassActorSpawnRequest& MassActorSpawnRequest = SpawnRequest.Get<FMassActorSpawnRequest>();
+	checkf(MassActorSpawnRequest.SpawnedActor, TEXT("Expecting valid spawned actor"));
+
+	// Might be already done if the actor has a MassAgentComponent via the callback OnMassAgentComponentEntityAssociated on the MassRepresentationSubsystem
+	FDataFragment_Actor& ActorInfo = PipeEntitySubsystem->GetComponentDataChecked<FDataFragment_Actor>(MassActorSpawnRequest.MassAgent.GetLWEntity());
+	if (ActorInfo.IsValid())
+	{
+		// If already set, make sure it is pointing to the same actor.
+		checkf(ActorInfo.Get() == MassActorSpawnRequest.SpawnedActor, TEXT("Expecting the pointer to the spawned actor in the actor fragment"));
+	}
+	else
+	{
+		ActorInfo.SetAndUpdateHandleMap(MassActorSpawnRequest.MassAgent, MassActorSpawnRequest.SpawnedActor, true/*bIsOwnedByMass*/);
+	}
+
+	PostActorSpawned(MassActorSpawnRequest.MassAgent, *MassActorSpawnRequest.SpawnedActor);
+
+	return EMassActorSpawnRequestAction::Keep;
+}
+
+//----------------------------------------------------------------------//
+// UMassRepresentationFragmentDestructor 
+//----------------------------------------------------------------------//
+UMassRepresentationFragmentDestructor::UMassRepresentationFragmentDestructor()
+{
+	// By putting this to null, this Deinitializer needs to be explicitly added via a Entity traits and will no longer be automatically called
+	// This FragmentType should be deprecated and we should always use the traits.
+	FragmentType = nullptr; // FMassRepresentationFragment::StaticStruct();
+	ExecutionFlags = (int32)(EProcessorExecutionFlags::All);
+}
+
+void UMassRepresentationFragmentDestructor::Initialize(UObject& Owner)
+{
+	Super::Initialize(Owner);
+
+	RepresentationSubsystem = UWorld::GetSubsystem<UMassRepresentationSubsystem>(Owner.GetWorld());
+}
+
+void UMassRepresentationFragmentDestructor::ConfigureQueries()
+{
+	EntityQuery.AddRequirement<FMassRepresentationFragment>(ELWComponentAccess::ReadOnly);
+	EntityQuery.AddRequirement<FDataFragment_Actor>(ELWComponentAccess::ReadWrite);
+}
+
+void UMassRepresentationFragmentDestructor::Execute(UEntitySubsystem& EntitySubsystem, FLWComponentSystemExecutionContext& Context)
+{
+	check(RepresentationSubsystem);
+
+	EntityQuery.ForEachEntityChunk(EntitySubsystem, Context, [this](FLWComponentSystemExecutionContext& Context)
+	{
+		const TArrayView<FMassRepresentationFragment> RepresentationList = Context.GetMutableComponentView<FMassRepresentationFragment>();
+		const TArrayView<FDataFragment_Actor> ActorList = Context.GetMutableComponentView<FDataFragment_Actor>();
+
+		const int32 NumEntities = Context.GetEntitiesNum();
+		for (int32 i = 0; i < NumEntities; ++i)
+		{
+			FMassRepresentationFragment& Representation = RepresentationList[i];
+			FDataFragment_Actor& ActorInfo = ActorList[i];
+
+			const FLWEntity Entity = Context.GetEntity(i);
+			const FMassHandle MassAgent(Entity);
+
+			UMassRepresentationProcessor::ReleaseAnyActorOrCancelAnySpawning(*RepresentationSubsystem, MassAgent, ActorInfo, Representation);
+		}
+	});
+}
\ No newline at end of file
diff --git a/Engine/Plugins/Runtime/MassGameplay/Source/MassRepresentation/Private/MassRepresentationSubsystem.cpp b/Engine/Plugins/Runtime/MassGameplay/Source/MassRepresentation/Private/MassRepresentationSubsystem.cpp
new file mode 100644
--- /dev/null
+++ b/Engine/Plugins/Runtime/MassGameplay/Source/MassRepresentation/Private/MassRepresentationSubsystem.cpp
@@ -0,0 +1,405 @@
+// Copyright Epic Games, Inc. All Rights Reserved.
+
+#include "MassRepresentationSubsystem.h"
+#include "Engine/World.h"
+#include "MassActorSpawnerSubsystem.h"
+#include "MassActorSubsystem.h"
+#include "MassSimulationSubsystem.h"
+#include "MassVisualizationComponent.h"
+#include "MassVisualizer.h"
+#include "MassRepresentationTypes.h"
+#include "MassSimulationSettings.h"
+#include "MassAgentComponent.h"
+#include "MassAgentSubsystem.h"
+#include "MassEntitySystem.h"
+#include "MassRepresentationFragments.h"
+#include "MassRepresentationProcessor.h"
+#include "EntityView.h"
+#include "WorldPartition/WorldPartitionSubsystem.h"
+#include "WorldPartition/WorldPartitionRuntimeCell.h"
+
+int16 UMassRepresentationSubsystem::FindOrAddStaticMeshDesc(const FStaticMeshInstanceVisualizationDesc& Desc) const
+{
+	check(VisualizationComponent);
+	return VisualizationComponent->FindOrAddVisualDesc(Desc);
+}
+
+FMassInstancedStaticMeshInfoArrayView UMassRepresentationSubsystem::GetMutableInstancedStaticMeshInfos() const
+{
+	check(VisualizationComponent);
+	FMassInstancedStaticMeshInfoArrayView View = VisualizationComponent->GetMutableVisualInfos();
+	return MoveTemp(View);
+}
+
+void UMassRepresentationSubsystem::DirtyStaticMeshInstances() const
+{
+	check(VisualizationComponent);
+	return VisualizationComponent->DirtyVisuals();
+}
+
+int16 UMassRepresentationSubsystem::FindOrAddTemplateActor(const TSubclassOf<AActor>& ActorClass)
+{
+	UE_MT_SCOPED_WRITE_ACCESS(TemplateActorsMTAccessDetector);
+	int32 VisualIndex = TemplateActors.Find(ActorClass);
+	if (VisualIndex == INDEX_NONE)
+	{
+		VisualIndex = TemplateActors.Add(ActorClass);
+	}
+	check(VisualIndex < INT16_MAX);
+	return (int16)VisualIndex;
+}
+
+AActor* UMassRepresentationSubsystem::GetOrSpawnActorFromTemplate(const FMassHandle MassAgent, const FTransform& Transform, const int16 TemplateActorIndex, FMassHandle_ActorSpawnRequest& SpawnRequestHandle, float Priority /*= MAX_FLT*/, 
+	FMassActorPreSpawnDelegate ActorPreSpawnDelegate /*= FMassActorPreSpawnDelegate()*/, FMassActorPostSpawnDelegate ActorPostSpawnDelegate /*= FMassActorPostSpawnDelegate()*/)
+{
+	UE_MT_SCOPED_READ_ACCESS(TemplateActorsMTAccessDetector);
+	if (!TemplateActors.IsValidIndex(TemplateActorIndex))
+	{
+		UE_LOG(LogMassRepresentation, Error, TEXT("Template actor type %i is not referring to a valid type"), TemplateActorIndex);
+		return nullptr;
+	}
+
+	//@todo: this would be a good place to do pooling of actors instead of spawning them every time
+	check(ActorSpawnerSubsystem);
+	const TSubclassOf<AActor> TemplateToSpawn = TemplateActors[TemplateActorIndex];
+
+	if (SpawnRequestHandle.IsValid())
+	{
+		FMassActorSpawnRequest& SpawnRequest = ActorSpawnerSubsystem->GetMutableSpawnRequest<FMassActorSpawnRequest>(SpawnRequestHandle);
+		// Check if this existing spawn request is matching the template actor
+		if (SpawnRequest.Template != TemplateToSpawn)
+		{
+			return nullptr;
+		}
+		switch (SpawnRequest.SpawnStatus)
+		{
+			case ESpawnRequestStatus::RetryPending:
+			case ESpawnRequestStatus::Pending:
+				// Update spawn request with latest information
+				SpawnRequest.Transform = Transform;
+				SpawnRequest.Priority = Priority;
+				return nullptr;
+			case ESpawnRequestStatus::Failed:
+			{
+				check(!SpawnRequest.SpawnedActor);
+				// The most common case for failure is that collision was preventing the spawning, the agent might have moved or the thing preventing it from spawning might have moved too.
+				UWorld* World = GetWorld();
+				check(World);
+
+				// Limits the retry as they are bad for rendering performance
+				if ((World->GetTimeSeconds() - SpawnRequest.RequestedTime) > RetryTimeInterval ||
+					(SpawnRequest.Transform.GetLocation() - Transform.GetLocation()).SizeSquared() > RetryMovedDistanceSq)
+				{
+					// Update spawn request with latest information and retry
+					SpawnRequest.Transform = Transform;
+					SpawnRequest.Priority = Priority;
+					ActorSpawnerSubsystem->RetryActorSpawnRequest(SpawnRequestHandle);
+				}
+				return nullptr;
+			}
+			case ESpawnRequestStatus::Succeeded:
+			{
+				AActor* SpawnedActor = SpawnRequest.SpawnedActor;
+				check(SpawnedActor);
+				ensureMsgf(ActorSpawnerSubsystem->RemoveActorSpawnRequest(SpawnRequestHandle), TEXT("Unable to remove a valid spawn request"));
+				return SpawnedActor;
+			}
+			default:
+				checkf(false, TEXT("Unexpected spawn request status!"));
+				SpawnRequestHandle.Invalidate();
+				return nullptr;
+		}
+	}
+
+	// If we reach here, means we need to create a spawn request
+	FMassActorSpawnRequest SpawnRequest;
+	SpawnRequest.MassAgent = MassAgent;
+	SpawnRequest.Template = TemplateToSpawn;
+	SpawnRequest.Transform = Transform;
+	SpawnRequest.Priority = Priority;
+	SpawnRequest.ActorPreSpawnDelegate = ActorPreSpawnDelegate;
+	SpawnRequest.ActorPostSpawnDelegate = ActorPostSpawnDelegate;
+	SpawnRequestHandle = ActorSpawnerSubsystem->RequestActorSpawn(SpawnRequest);
+
+	++(HandledMassAgents.FindOrAdd(MassAgent));
+
+	return nullptr;
+}
+
+TSubclassOf<AActor> UMassRepresentationSubsystem::GetTemplateActorClass(const int16 TemplateActorIndex)
+{
+	UE_MT_SCOPED_READ_ACCESS(TemplateActorsMTAccessDetector);
+	if (!TemplateActors.IsValidIndex(TemplateActorIndex))
+	{
+		UE_LOG(LogMassRepresentation, Error, TEXT("Template actor type %i is not referring to a valid type"), TemplateActorIndex);
+		return nullptr;
+	}
+
+	return TemplateActors[TemplateActorIndex];
+}
+
+bool UMassRepresentationSubsystem::IsCollisionLoaded(const FName TargetGrid, const FTransform& Transform) const
+{
+	if (!WorldPartitionSubsystem)
+	{
+		// Assuming that all collisions are loaded if not using WorldPartition.
+		return true;
+	}
+
+	// @todo optimize by doing one query per cell
+	// Build a query source
+	TArray<FWorldPartitionStreamingQuerySource> QuerySources;
+	FWorldPartitionStreamingQuerySource& QuerySource = QuerySources.Emplace_GetRef();
+	QuerySource.bSpatialQuery = true;
+	QuerySource.Location = Transform.GetLocation();
+	QuerySource.Rotation = Transform.Rotator();
+	QuerySource.TargetGrid = TargetGrid;
+	QuerySource.bUseGridLoadingRange = false;
+	QuerySource.Radius = 1.f; // 1cm should be enough to know if grid is loaded at specific area
+	QuerySource.bDataLayersOnly = false;
+
+	// Execute query
+	return WorldPartitionSubsystem->IsStreamingCompleted(EWorldPartitionRuntimeCellState::Activated, QuerySources, /*bExactState*/ false);
+}
+
+void UMassRepresentationSubsystem::ReleaseAllResources()
+{
+	ensureMsgf(HandledMassAgents.Num() == 0, TEXT("All entities must be released before releasing resources"));
+
+	{
+		// Release all template actors
+		UE_MT_SCOPED_WRITE_ACCESS(TemplateActorsMTAccessDetector);
+		TemplateActors.Reset();
+	}
+
+	// Release all static meshes resources
+	VisualizationComponent->ClearAllVisualInstances();
+}
+
+bool UMassRepresentationSubsystem::ReleaseTemplateActorOrCancelSpawning(const FMassHandle MassAgent, const int16 TemplateActorIndex, AActor* ActorToRelease, FMassHandle_ActorSpawnRequest& SpawnRequestHandle)
+{
+	// First try to cancel the spawning request, then try to release the actor
+	if( CancelSpawningInternal(TemplateActorIndex, SpawnRequestHandle) ||
+		ReleaseTemplateActorInternal(TemplateActorIndex, ActorToRelease, false))
+	{ 
+		const int32 RefCount = --(HandledMassAgents.FindChecked(MassAgent));
+		checkf(RefCount >= 0, TEXT("RefCount are expected to be greater than or equal to 0"));
+		if (RefCount == 0)
+		{
+			HandledMassAgents.Remove(MassAgent);
+		}
+		return true;
+	}
+
+	return false;
+}
+
+bool UMassRepresentationSubsystem::ReleaseTemplateActor(const FMassHandle MassAgent, const int16 TemplateActorIndex, AActor* ActorToRelease, bool bImmediate)
+{
+	if (ReleaseTemplateActorInternal(TemplateActorIndex, ActorToRelease, bImmediate))
+	{
+		const int32 RefCount = --HandledMassAgents.FindChecked(MassAgent);
+		checkf(RefCount >= 0, TEXT("RefCount are expected to be greater than or equal to 0"));
+		if (RefCount == 0)
+		{
+			HandledMassAgents.Remove(MassAgent);
+		}
+		return true;
+	}
+
+	return false;
+}
+
+bool UMassRepresentationSubsystem::CancelSpawning(const FMassHandle MassAgent, const int16 TemplateActorIndex, FMassHandle_ActorSpawnRequest& SpawnRequestHandle)
+{
+	if (CancelSpawningInternal(TemplateActorIndex, SpawnRequestHandle))
+	{
+		const int32 RefCount = --HandledMassAgents.FindChecked(MassAgent);
+		checkf(RefCount >= 0, TEXT("RefCount are expected to be greater than or equal to 0"));
+		if (RefCount == 0)
+		{
+			HandledMassAgents.Remove(MassAgent);
+		}
+		return true;
+	}
+	return false;
+}
+
+bool UMassRepresentationSubsystem::ReleaseTemplateActorInternal(const int16 TemplateActorIndex, AActor* ActorToRelease, bool bImmediate)
+{
+	UE_MT_SCOPED_READ_ACCESS(TemplateActorsMTAccessDetector);
+	if (!TemplateActors.IsValidIndex(TemplateActorIndex))
+	{
+		UE_LOG(LogMassRepresentation, Error, TEXT("Template actor type %i is not referring a valid type"), TemplateActorIndex);
+		return false;
+	}
+	const TSubclassOf<AActor> TemplateToRelease = TemplateActors[TemplateActorIndex];
+
+	// We can only release existing and matching template actors
+	if (!ActorToRelease || ActorToRelease->GetClass() != TemplateToRelease)
+	{
+		return false;
+	}
+	ActorSpawnerSubsystem->DestroyActor(ActorToRelease, bImmediate);
+	return true;
+}
+
+bool UMassRepresentationSubsystem::CancelSpawningInternal(const int16 TemplateActorIndex, FMassHandle_ActorSpawnRequest& SpawnRequestHandle)
+{
+	UE_MT_SCOPED_READ_ACCESS(TemplateActorsMTAccessDetector);
+	if (!TemplateActors.IsValidIndex(TemplateActorIndex))
+	{
+		UE_LOG(LogMassRepresentation, Error, TEXT("Template actor type %i is not referring a valid type"), TemplateActorIndex);
+		return false;
+	}
+	const TSubclassOf<AActor> TemplateToRelease = TemplateActors[TemplateActorIndex];
+
+	check(ActorSpawnerSubsystem);
+	// Check if there is something to cancel
+	if (!SpawnRequestHandle.IsValid())
+	{
+		return false;
+	}
+
+	const FMassActorSpawnRequest& SpawnRequest = ActorSpawnerSubsystem->GetMutableSpawnRequest<FMassActorSpawnRequest>(SpawnRequestHandle);
+	// Check if the spawning request matches the template actor
+	if (SpawnRequest.Template != TemplateToRelease)
+	{
+		return false;
+	}
+
+	if (SpawnRequest.SpawnStatus == ESpawnRequestStatus::Succeeded)
+	{
+		check(SpawnRequest.SpawnedActor);
+		ReleaseTemplateActorInternal(TemplateActorIndex, SpawnRequest.SpawnedActor, false);
+	}
+
+	// Remove the spawn request
+	ensureMsgf(ActorSpawnerSubsystem->RemoveActorSpawnRequest(SpawnRequestHandle), TEXT("Unable to remove a valid spawn request"));
+
+	return true;
+}
+
+bool UMassRepresentationSubsystem::DoesActorMatchTemplate(const AActor& Actor, const int16 TemplateActorIndex)
+{
+	UE_MT_SCOPED_READ_ACCESS(TemplateActorsMTAccessDetector);
+	if (!TemplateActors.IsValidIndex(TemplateActorIndex))
+	{
+		UE_LOG(LogMassRepresentation, Error, TEXT("Template actor type %i is not referring to a valid type"), TemplateActorIndex);
+		return false;
+	}
+
+	const TSubclassOf<AActor> Template = TemplateActors[TemplateActorIndex];
+	return Actor.GetClass() == Template;
+}
+
+void UMassRepresentationSubsystem::Initialize(FSubsystemCollectionBase& Collection)
+{
+	Collection.InitializeDependency(UMassSimulationSubsystem::StaticClass());
+	Collection.InitializeDependency(UMassActorSpawnerSubsystem::StaticClass());
+	Collection.InitializeDependency(UMassAgentSubsystem::StaticClass());
+	Super::Initialize(Collection);
+
+	if (UWorld* World = GetWorld())
+	{
+		EntitySubsystem = UPipeEntitySubsystem::GetCurrent(World);
+		ActorSpawnerSubsystem = World->GetSubsystem<UMassActorSpawnerSubsystem>();
+		WorldPartitionSubsystem = World->GetSubsystem<UWorldPartitionSubsystem>();
+
+		if (Visualizer == nullptr)
+		{
+			FActorSpawnParameters SpawnInfo;
+			SpawnInfo.SpawnCollisionHandlingOverride = ESpawnActorCollisionHandlingMethod::AlwaysSpawn;
+			// The helper actor is only once per world so we can allow it to spawn during construction script.
+			SpawnInfo.bAllowDuringConstructionScript = true;
+			Visualizer = World->SpawnActor<AMassVisualizer>(SpawnInfo);
+			check(Visualizer);
+			VisualizationComponent = &Visualizer->GetVisualizationComponent();
+		}
+
+		UMassSimulationSubsystem* SimSystem = World->GetSubsystem<UMassSimulationSubsystem>();
+		check(SimSystem);
+		SimSystem->GetOnProcessingPhaseStarted(EPipeProcessingPhase::PrePhysics).AddUObject(this, &UMassRepresentationSubsystem::OnProcessingPhaseStarted, EPipeProcessingPhase::PrePhysics);
+		SimSystem->GetOnProcessingPhaseFinished(EPipeProcessingPhase::PostPhysics).AddUObject(this, &UMassRepresentationSubsystem::OnProcessingPhaseStarted, EPipeProcessingPhase::PostPhysics);
+
+		UMassAgentSubsystem* AgentSystem = World->GetSubsystem<UMassAgentSubsystem>();
+		check(AgentSystem);
+		AgentSystem->GetOnMassAgentComponentEntityAssociated().AddUObject(this, &UMassRepresentationSubsystem::OnMassAgentComponentEntityAssociated);
+		AgentSystem->GetOnMassAgentComponentEntityDetaching().AddUObject(this, &UMassRepresentationSubsystem::OnMassAgentComponentEntityDetaching);
+	}
+
+	RetryMovedDistanceSq = FMath::Square(GET_MASS_CONFIG_VALUE(DesiredActorFailedSpawningRetryMoveDistance));
+	RetryTimeInterval = GET_MASS_CONFIG_VALUE(DesiredActorFailedSpawningRetryTimeInterval);
+}
+
+void UMassRepresentationSubsystem::Deinitialize()
+{
+	if (const UWorld* World = GetWorld())
+	{
+		if (const UMassSimulationSubsystem* SimSystem = World->GetSubsystem<UMassSimulationSubsystem>())
+		{
+			SimSystem->GetOnProcessingPhaseStarted(EPipeProcessingPhase::PrePhysics).RemoveAll(this);
+			SimSystem->GetOnProcessingPhaseFinished(EPipeProcessingPhase::PostPhysics).RemoveAll(this);
+		}
+
+		if (UMassAgentSubsystem* AgentSystem = World->GetSubsystem<UMassAgentSubsystem>())
+		{
+			AgentSystem->GetOnMassAgentComponentEntityAssociated().RemoveAll(this);
+			AgentSystem->GetOnMassAgentComponentEntityDetaching().RemoveAll(this);
+		}
+	}
+
+}
+
+void UMassRepresentationSubsystem::OnProcessingPhaseStarted(const float DeltaSeconds, const EPipeProcessingPhase Phase) const
+{
+	check(VisualizationComponent);
+	switch (Phase)
+	{
+		case EPipeProcessingPhase::PrePhysics:
+			VisualizationComponent->BeginVisualChanges();
+			break;
+		case EPipeProcessingPhase::PostPhysics:/* Currently this is the end of phases signal */
+			VisualizationComponent->EndVisualChanges();
+			break;
+		default:
+			check(false); // Need to handle this case
+			break;
+	}
+}
+
+void UMassRepresentationSubsystem::OnMassAgentComponentEntityAssociated(const UMassAgentComponent& AgentComponent)
+{
+	check(EntitySubsystem);
+
+	const FMassHandle MassAgent = AgentComponent.GetEntityHandle();
+	checkf(EntitySubsystem->IsEntityValid(MassAgent.GetLWEntity()), TEXT("Expecting a valid mass entity"));
+	if (EntitySubsystem->IsEntityValid(MassAgent.GetLWEntity()) && AgentComponent.IsNetSimulating())
+	{
+		// Check if this mass agent already handled by this sub system, if yes than release any local spawned actor or cancel any spawn requests
+		if (HandledMassAgents.Find(MassAgent))
+		{
+			UMassRepresentationProcessor::ReleaseAnyActorOrCancelAnySpawning(*this, MassAgent);
+		}
+	}
+}
+
+void UMassRepresentationSubsystem::OnMassAgentComponentEntityDetaching(const UMassAgentComponent& AgentComponent)
+{
+	check(EntitySubsystem);
+
+	AActor* ComponentOwner = AgentComponent.GetOwner();
+	check(ComponentOwner);
+
+	const FMassHandle MassAgent = AgentComponent.GetEntityHandle();
+	checkf(EntitySubsystem->IsEntityValid(MassAgent.GetLWEntity()), TEXT("Expecting a valid mass entity"));
+	if (EntitySubsystem->IsEntityValid(MassAgent.GetLWEntity()) && AgentComponent.IsNetSimulating())
+	{
+		const FEntityView EntityView(*EntitySubsystem, MassAgent.GetLWEntity());
+		if (FMassRepresentationFragment* Representation = EntityView.GetComponentDataPtr<FMassRepresentationFragment>())
+		{
+			// Force a reevaluate of the current representation
+			Representation->CurrentRepresentation = ERepresentationType::None;
+		}
+	}
+}
diff --git a/Engine/Plugins/Runtime/MassGameplay/Source/MassRepresentation/Private/MassUpdateISMProcessor.cpp b/Engine/Plugins/Runtime/MassGameplay/Source/MassRepresentation/Private/MassUpdateISMProcessor.cpp
new file mode 100644
--- /dev/null
+++ b/Engine/Plugins/Runtime/MassGameplay/Source/MassRepresentation/Private/MassUpdateISMProcessor.cpp
@@ -0,0 +1,87 @@
+// Copyright Epic Games, Inc. All Rights Reserved.
+
+#include "MassUpdateISMProcessor.h"
+#include "MassVisualizationComponent.h"
+#include "MassRepresentationSubsystem.h"
+#include "MassEntitySystem.h"
+#include "MassRepresentationFragments.h"
+#include "MassCommonFragments.h"
+#include "MassLODTypes.h"
+#include "Engine/World.h"
+
+UMassUpdateISMProcessor::UMassUpdateISMProcessor()
+{
+	ExecutionFlags = (int32)(EProcessorExecutionFlags::Client | EProcessorExecutionFlags::Standalone);
+
+	ExecutionOrder.ExecuteAfter.Add(UE::Mass::ProcessorGroupNames::Representation);
+}
+
+void UMassUpdateISMProcessor::Initialize(UObject& Owner)
+{
+	Super::Initialize(Owner);
+
+	RepresentationSubsystem = UWorld::GetSubsystem<UMassRepresentationSubsystem>(Owner.GetWorld());
+}
+
+void UMassUpdateISMProcessor::ConfigureQueries()
+{
+	for (const FInstancedStruct& TagFilter : TagFilters)
+	{
+		FLWComponentQuery& EntityQuery = EntityQueries.AddDefaulted_GetRef();
+		if (TagFilter.GetScriptStruct())
+		{
+			EntityQuery.AddTagRequirement(*TagFilter.GetScriptStruct(), ELWComponentPresence::All);
+		}
+		EntityQuery.AddRequirement<FDataFragment_Transform>(ELWComponentAccess::ReadOnly);
+		EntityQuery.AddRequirement<FMassRepresentationFragment>(ELWComponentAccess::ReadWrite);
+		EntityQuery.AddRequirement<FMassRepresentationLODFragment>(ELWComponentAccess::ReadWrite);
+		EntityQuery.AddChunkRequirement<FMassVisualizationChunkFragment>(ELWComponentAccess::ReadWrite);
+		EntityQuery.SetChunkFilter(&FMassVisualizationChunkFragment::AreAnyEntitiesVisibleInChunk);
+	}
+}
+
+void UMassUpdateISMProcessor::Execute(UEntitySubsystem& EntitySubsystem, FLWComponentSystemExecutionContext& Context)
+{
+	check(RepresentationSubsystem);
+	FMassInstancedStaticMeshInfoArrayView ISMInfo = RepresentationSubsystem->GetMutableInstancedStaticMeshInfos();
+
+	for (FLWComponentQuery& EntityQuery : EntityQueries)
+	{
+		EntityQuery.ForEachEntityChunk(EntitySubsystem, Context, [&ISMInfo](FLWComponentSystemExecutionContext& Context)
+		{
+			const TConstArrayView<FDataFragment_Transform> TransformList = Context.GetComponentView<FDataFragment_Transform>();
+			const TArrayView<FMassRepresentationFragment> RepresentationList = Context.GetMutableComponentView<FMassRepresentationFragment>();
+			const TConstArrayView<FMassRepresentationLODFragment> RepresentationLODList = Context.GetComponentView<FMassRepresentationLODFragment>();
+
+			const int32 NumEntities = Context.GetEntitiesNum();
+			for (int32 EntityIdx = 0; EntityIdx < NumEntities; EntityIdx++)
+			{
+				const FDataFragment_Transform& TransformFragment = TransformList[EntityIdx];
+				const FMassRepresentationLODFragment& RepresentationLOD = RepresentationLODList[EntityIdx];
+				FMassRepresentationFragment& Representation = RepresentationList[EntityIdx];
+
+				if (Representation.CurrentRepresentation == ERepresentationType::StaticMeshInstance)
+				{
+					UpdateISMTransform(GetTypeHash(Context.GetEntity(EntityIdx)), ISMInfo[Representation.StaticMeshDescIndex], TransformFragment.GetTransform(), Representation.PrevTransform, RepresentationLOD.LODSignificance);
+				}
+				Representation.PrevTransform = TransformFragment.GetTransform();
+			}
+		});
+	}
+}
+
+void UMassUpdateISMProcessor::UpdateISMTransform(int32 EntityId, FMassInstancedStaticMeshInfo& ISMInfo, const FTransform& Transform, const FTransform& PrevTransform, const float LODSignificance)
+{
+	if (ISMInfo.ShouldUseTransformOffset())
+	{
+		const FTransform& TransformOffset = ISMInfo.GetTransformOffset();
+		const FTransform SMTransform = TransformOffset * Transform;
+		const FTransform SMPrevTransform = TransformOffset * PrevTransform;
+
+		ISMInfo.AddBatchedTransform(EntityId, SMTransform, SMPrevTransform, LODSignificance);
+	}
+	else
+	{
+		ISMInfo.AddBatchedTransform(EntityId, Transform, PrevTransform, LODSignificance);
+	}
+}
diff --git a/Engine/Plugins/Runtime/MassGameplay/Source/MassRepresentation/Private/MassVisualizationComponent.cpp b/Engine/Plugins/Runtime/MassGameplay/Source/MassRepresentation/Private/MassVisualizationComponent.cpp
new file mode 100644
--- /dev/null
+++ b/Engine/Plugins/Runtime/MassGameplay/Source/MassRepresentation/Private/MassVisualizationComponent.cpp
@@ -0,0 +1,338 @@
+// Copyright Epic Games, Inc. All Rights Reserved.
+
+#include "MassVisualizationComponent.h"
+#include "Logging/LogMacros.h"
+#include "MassVisualizer.h"
+#include "MassRepresentationTypes.h"
+#include "Components/InstancedStaticMeshComponent.h"
+#include "Engine/InstancedStaticMesh.h"
+#include "Engine/CollisionProfile.h"
+#include "RenderUtils.h"
+
+
+//---------------------------------------------------------------
+// UMassVisualizationComponent
+//---------------------------------------------------------------
+
+namespace UE::MassRepresentation
+{
+	int32 GCallUpdateInstances = 0;
+	FAutoConsoleVariableRef  CVarCallUpdateInstances(TEXT("Mass.CallUpdateInstances"), GCallUpdateInstances, TEXT("Toggle between UpdateInstances and BatchUpdateTransform."));
+}  // UE::MassRepresentation
+
+void UMassVisualizationComponent::PostInitProperties()
+{
+	Super::PostInitProperties();
+	if (HasAnyFlags(RF_ClassDefaultObject) == false && GetOuter())
+	{
+		ensureMsgf(GetOuter()->GetClass()->IsChildOf(AMassVisualizer::StaticClass()), TEXT("UMassVisualizationComponent should only be added to AMassVisualizer-like instances"));
+	}
+}
+
+int16 UMassVisualizationComponent::FindOrAddVisualDesc(const FStaticMeshInstanceVisualizationDesc& Desc)
+{
+	UE_MT_SCOPED_WRITE_ACCESS(InstancedStaticMeshInfosDetector);
+	int32 VisualIndex = InstancedStaticMeshInfos.IndexOfByPredicate([&Desc](const FMassInstancedStaticMeshInfo& Info) { return Info == Desc; });
+	if (VisualIndex == INDEX_NONE)
+	{
+		VisualIndex = InstancedStaticMeshInfos.Emplace(Desc);
+		bNeedStaticMeshComponentConstruction = true;
+	}
+	check(VisualIndex < INT16_MAX);
+	return (int16)VisualIndex;
+}
+
+void UMassVisualizationComponent::ConstructStaticMeshComponents()
+{
+	AActor* ActorOwner = GetOwner();
+	check(ActorOwner);
+	
+	UE_MT_SCOPED_WRITE_ACCESS(InstancedStaticMeshInfosDetector);
+	for (FMassInstancedStaticMeshInfo& Info : InstancedStaticMeshInfos)
+	{
+		// Check if it is already created
+		if (!Info.InstancedStaticMeshComponents.IsEmpty())
+		{
+			continue;
+		}
+
+		// Check if there are any specified meshes for this visual type
+		if(Info.Desc.Meshes.Num() == 0)
+		{
+			UE_LOG(LogMassRepresentation, Error, TEXT("No associated meshes for this intanced static mesh type"));
+			continue;
+		}
+		for (const FStaticMeshInstanceVisualizationMeshDesc& MeshDesc : Info.Desc.Meshes)
+		{			
+			FISMCSharedData* SharedData = ISMCSharedData.Find(MeshDesc.Mesh);
+			UInstancedStaticMeshComponent* ISMC = SharedData ? SharedData->ISMC : nullptr;
+			if (SharedData)
+			{
+				SharedData->RefCount += 1;
+			}
+			else
+			{
+				ISMC = NewObject<UInstancedStaticMeshComponent>(ActorOwner);
+				ISMC->SetStaticMesh(MeshDesc.Mesh);
+				for (int32 ElementIndex = 0; ElementIndex < MeshDesc.MaterialOverrides.Num(); ++ElementIndex)
+				{
+					if (UMaterialInterface* MaterialOverride = MeshDesc.MaterialOverrides[ElementIndex])
+					{
+						ISMC->SetMaterial(ElementIndex, MaterialOverride);
+					}
+				}
+				ISMC->SetCullDistances(0.0f, 1000000.0); // @todo: Need to figure out what to do here, either LOD or cull distances.
+				ISMC->SetupAttachment(ActorOwner->GetRootComponent());
+				ISMC->SetCanEverAffectNavigation(false);
+				ISMC->SetCollisionProfileName(UCollisionProfile::NoCollision_ProfileName);
+				ISMC->SetCastShadow(MeshDesc.bCastShadows);
+				ISMC->SetReceivesDecals(false);
+				ISMC->RegisterComponent();
+
+				ISMCSharedData.Emplace(MeshDesc.Mesh, FISMCSharedData(ISMC));
+			}
+
+			Info.InstancedStaticMeshComponents.Add(ISMC);
+		}
+
+		// Build the LOD significance ranges
+		TArray<float> AllLODSignificances;
+		auto UniqueInsertOrdered = [&AllLODSignificances](const float Significance)
+		{
+			int i = 0;
+			for (; i < AllLODSignificances.Num(); ++i)
+			{
+				// I did not use epsilon check here on purpose, because it will make it hard later meshes inside.
+				if (Significance == AllLODSignificances[i])
+				{
+					return;
+				}
+				if (AllLODSignificances[i] > Significance)
+				{
+					break;
+				}
+			}
+			AllLODSignificances.Insert(Significance, i);
+		};
+		for (const FStaticMeshInstanceVisualizationMeshDesc& MeshDesc : Info.Desc.Meshes)
+		{
+			UniqueInsertOrdered(MeshDesc.MinLODSignificance);
+			UniqueInsertOrdered(MeshDesc.MaxLODSignificance);
+		}
+		Info.LODSignificanceRanges.SetNum(AllLODSignificances.Num() - 1);
+		for (int i = 0; i < Info.LODSignificanceRanges.Num(); ++i)
+		{
+			FMassLODSignificanceRange& Range = Info.LODSignificanceRanges[i];
+			Range.MinSignificance = AllLODSignificances[i];
+			Range.MaxSignificance = AllLODSignificances[i+1];
+			Range.ISMCSharedDataPtr = &ISMCSharedData;
+
+			for (int j = 0; j < Info.Desc.Meshes.Num(); ++j)
+			{
+				const FStaticMeshInstanceVisualizationMeshDesc& MeshDesc = Info.Desc.Meshes[j];
+				const bool bAddMeshInRange = (Range.MinSignificance >= MeshDesc.MinLODSignificance && Range.MinSignificance < MeshDesc.MaxLODSignificance);
+				if (bAddMeshInRange)
+				{
+					Range.StaticMeshRefs.Add(MeshDesc.Mesh);
+				}
+			}
+		}
+
+	}
+}
+
+void UMassVisualizationComponent::ClearAllVisualInstances()
+{
+	UE_MT_SCOPED_WRITE_ACCESS(InstancedStaticMeshInfosDetector);
+	for (FMassInstancedStaticMeshInfo& Info : InstancedStaticMeshInfos)
+	{
+		Info.ClearVisualInstance(ISMCSharedData);
+	}
+	InstancedStaticMeshInfos.Reset();
+	
+	// Pool should already be empty, got a problem if it's not
+	for (auto It = ISMCSharedData.CreateIterator(); It; ++It)
+	{
+		if (UInstancedStaticMeshComponent* InstancedStaticMeshComponent = It.Value().ISMC)
+		{
+			InstancedStaticMeshComponent->ClearInstances();
+			InstancedStaticMeshComponent->DestroyComponent();
+		}
+	}
+
+	ISMCSharedData.Reset();
+}
+
+void UMassVisualizationComponent::DirtyVisuals()
+{
+	UE_MT_SCOPED_WRITE_ACCESS(InstancedStaticMeshInfosDetector);
+	for (FMassInstancedStaticMeshInfo& Info : InstancedStaticMeshInfos)
+	{
+		for (UInstancedStaticMeshComponent* InstancedStaticMeshComponent : Info.InstancedStaticMeshComponents)
+		{
+			InstancedStaticMeshComponent->MarkRenderStateDirty();
+		}
+	}
+}
+
+void UMassVisualizationComponent::BeginVisualChanges()
+{
+	TRACE_CPUPROFILER_EVENT_SCOPE_STR("MassVisualizationComponent BeginVisualChanges")
+
+	// Conditionally construct static mesh components
+	if (bNeedStaticMeshComponentConstruction)
+	{
+		ConstructStaticMeshComponents();
+		bNeedStaticMeshComponentConstruction = false;
+	}
+
+	// Reset instance transform scratch buffers
+	for (auto It = ISMCSharedData.CreateIterator(); It; ++It)
+	{
+		FISMCSharedData& SharedData = It.Value();
+		SharedData.UpdateInstanceIds.Reset();
+		SharedData.StaticMeshInstanceCustomFloats.Reset();
+		SharedData.StaticMeshInstanceTransforms.Reset();
+		SharedData.StaticMeshInstancePrevTransforms.Reset();
+		SharedData.WriteIterator = 0;
+	}
+}
+
+void UMassVisualizationComponent::EndVisualChanges()
+{
+	TRACE_CPUPROFILER_EVENT_SCOPE_STR("MassVisualizationComponent EndVisualChanges")
+
+	// Batch update gathered instance transforms
+	for (auto It = ISMCSharedData.CreateIterator(); It; ++It)
+	{
+		FISMCSharedData& SharedData = It.Value();
+		if (UInstancedStaticMeshComponent* InstancedStaticMeshComponent = SharedData.ISMC)
+		{
+			const int32 NumCustomDataFloats = SharedData.StaticMeshInstanceCustomFloats.Num() / (FMath::Max(1, SharedData.UpdateInstanceIds.Num()));
+
+			// Ensure InstanceCustomData is passed if NumCustomDataFloats > 0. If it is, also make sure
+			// its length is NumCustomDataFloats * InstanceTransforms.Num()
+			ensure(NumCustomDataFloats == 0 || (SharedData.StaticMeshInstanceCustomFloats.Num() == NumCustomDataFloats * SharedData.UpdateInstanceIds.Num()));
+
+			// Resize PerInstanceSMData & PerInstanceSMCustomData to new (possibly culled or expanded) transform batch length
+			const int32 NewNumInstances = SharedData.StaticMeshInstanceTransforms.Num();
+
+			// Update PerInstanceSMData transforms
+			// @Todo reenable when new UpdateInstances code is in main branch
+// 			if ((bool)UE::MassRepresentation::GCallUpdateInstances)
+// 			{
+// 				InstancedStaticMeshComponent->UpdateInstances(SharedData.UpdateInstanceIds, SharedData.StaticMeshInstanceTransforms, SharedData.StaticMeshInstancePrevTransforms, NumCustomDataFloats, SharedData.StaticMeshInstanceCustomFloats);
+// 			}
+// 			else
+			{
+				// Update NumCustomDataFloats
+				InstancedStaticMeshComponent->NumCustomDataFloats = NumCustomDataFloats;
+				if (NumCustomDataFloats > 0)
+				{
+					InstancedStaticMeshComponent->PerInstanceSMCustomData = SharedData.StaticMeshInstanceCustomFloats;
+				}
+
+				InstancedStaticMeshComponent->PerInstanceSMData.SetNum(NewNumInstances, /*bAllowShrinking*/false);
+				InstancedStaticMeshComponent->PerInstancePrevTransform.SetNum(NewNumInstances, /*bAllowShrinking*/false);
+
+				// Update PerInstanceSMData transforms
+				InstancedStaticMeshComponent->BatchUpdateInstancesTransforms(/*StartInstanceIndex*/0, SharedData.StaticMeshInstanceTransforms, SharedData.StaticMeshInstancePrevTransforms, /*bWorldSpace*/false, /*bMarkRenderStateDirty*/false);
+
+				// Nanite ISMC? 
+				TObjectPtr<UStaticMesh> StaticMeshObjectPtr = InstancedStaticMeshComponent->GetStaticMesh();
+				FStaticMeshRenderData* StaticMeshRenderData = nullptr;
+				if (UStaticMesh* StaticMesh = StaticMeshObjectPtr.Get())
+				{
+					StaticMeshRenderData = StaticMesh->GetRenderData();
+				}
+				const bool bNaniteISMC = UseNanite(InstancedStaticMeshComponent->GetScene()->GetShaderPlatform()) && StaticMeshRenderData && StaticMeshRenderData->NaniteResources.PageStreamingStates.Num();
+				if (bNaniteISMC)
+				{
+
+					// ISMC currently rebuilds PerInstanceRenderData regardless of whether it's using a 
+					// Nanite::FSceneProxy which doesn't actually use this data. So to skip that we 
+					// reset the InstanceUpdateCmdBuffer here after BatchUpdateInstancesTransforms has
+					// marked it dirty but before CreateSceneProxy checks it
+					// @todo This should be in ISMC code
+					InstancedStaticMeshComponent->InstanceUpdateCmdBuffer.Cmds.Reset();
+					InstancedStaticMeshComponent->InstanceUpdateCmdBuffer.NumAdds = 0;
+					InstancedStaticMeshComponent->InstanceUpdateCmdBuffer.NumEdits = 0;
+				}
+
+				// Dirty render state
+				InstancedStaticMeshComponent->MarkRenderStateDirty();
+			}
+		}
+	}
+}
+
+//---------------------------------------------------------------
+// FMassInstancedStaticMeshInfo
+//---------------------------------------------------------------
+
+void FMassInstancedStaticMeshInfo::ClearVisualInstance(TMap<UStaticMesh*, FISMCSharedData>& ISMCSharedData)
+{
+	for (int i = 0; i < Desc.Meshes.Num(); i++)
+	{
+		FISMCSharedData* SharedData = ISMCSharedData.Find(Desc.Meshes[i].Mesh);
+		if (SharedData)
+		{
+			SharedData->RefCount -= 1;
+			if (SharedData->RefCount == 0)
+			{
+				SharedData->ISMC->ClearInstances();
+				SharedData->ISMC->DestroyComponent();
+				ISMCSharedData.Remove(Desc.Meshes[i].Mesh);
+			}
+		}
+	}
+
+	InstancedStaticMeshComponents.Reset();
+	LODSignificanceRanges.Reset();
+}
+
+//---------------------------------------------------------------
+// FMassLODSignificanceRange
+//---------------------------------------------------------------
+
+void FMassLODSignificanceRange::AddBatchedTransform(const int32 InstanceId, const FTransform& Transform, const FTransform& PrevTransform)
+{
+	check(ISMCSharedDataPtr);
+	for (int i = 0; i < StaticMeshRefs.Num(); i++)
+	{
+		FISMCSharedData& SharedData = (*ISMCSharedDataPtr)[StaticMeshRefs[i]];
+
+		SharedData.UpdateInstanceIds.Add(InstanceId);
+		SharedData.StaticMeshInstanceTransforms.Add(Transform);
+		SharedData.StaticMeshInstancePrevTransforms.Add(PrevTransform);
+	}
+}
+
+void FMassLODSignificanceRange::AddBatchedCustomDataFloats(const TArray<float>& CustomFloats)
+{
+	check(ISMCSharedDataPtr);
+	for (int i = 0; i < StaticMeshRefs.Num(); i++)
+	{
+		FISMCSharedData& SharedData = (*ISMCSharedDataPtr)[StaticMeshRefs[i]];
+		SharedData.StaticMeshInstanceCustomFloats.Append(CustomFloats);
+	}
+}
+
+void FMassLODSignificanceRange::WriteCustomDataFloatsAtStartIndex(int32 StaticMeshIndex, const TArrayView<float> CustomFloats, int32 FloatsPerInstance, int32 StartFloatIndex)
+{
+	check(ISMCSharedDataPtr);
+	if (StaticMeshRefs.IsValidIndex(StaticMeshIndex))
+	{
+		FISMCSharedData& SharedData = (*ISMCSharedDataPtr)[StaticMeshRefs[StaticMeshIndex]];
+
+		int32 StartIndex = FloatsPerInstance * SharedData.WriteIterator + StartFloatIndex;
+
+		ensure(SharedData.StaticMeshInstanceCustomFloats.Num() >= StartIndex + CustomFloats.Num());
+
+		for (int CustomFloatIdx = 0; CustomFloatIdx < CustomFloats.Num(); CustomFloatIdx++)
+		{
+			SharedData.StaticMeshInstanceCustomFloats[StartIndex + CustomFloatIdx] = CustomFloats[CustomFloatIdx];
+		}
+		SharedData.WriteIterator++;
+	}
+}
diff --git a/Engine/Plugins/Runtime/MassGameplay/Source/MassRepresentation/Private/MassVisualizationLODProcessor.cpp b/Engine/Plugins/Runtime/MassGameplay/Source/MassRepresentation/Private/MassVisualizationLODProcessor.cpp
new file mode 100644
--- /dev/null
+++ b/Engine/Plugins/Runtime/MassGameplay/Source/MassRepresentation/Private/MassVisualizationLODProcessor.cpp
@@ -0,0 +1,67 @@
+// Copyright Epic Games, Inc. All Rights Reserved.
+
+#include "MassVisualizationLODProcessor.h"
+
+namespace UE::MassRepresentation
+{
+	int32 bDebugRepresentationLOD = 0;
+	FAutoConsoleVariableRef CVarDebugRepresentationLOD(TEXT("ai.debug.RepresentationLOD"), bDebugRepresentationLOD, TEXT("Debug representation LOD"), ECVF_Cheat);
+} // UE::MassRepresentation
+
+
+UMassVisualizationLODProcessor::UMassVisualizationLODProcessor()
+{
+	bAutoRegisterWithProcessingPhases = false;
+
+	ExecutionOrder.ExecuteInGroup = UE::Mass::ProcessorGroupNames::LOD;
+
+	BaseLODDistance[EMassLOD::High] = 0.f;
+	BaseLODDistance[EMassLOD::Medium] = 1000.f;
+	BaseLODDistance[EMassLOD::Low] = 2500.f;
+	BaseLODDistance[EMassLOD::Off] = 10000.f;
+	
+	VisibleLODDistance[EMassLOD::High] = 0.f;
+	VisibleLODDistance[EMassLOD::Medium] = 2000.f;
+	VisibleLODDistance[EMassLOD::Low] = 4000.f;
+	VisibleLODDistance[EMassLOD::Off] = 10000.f;
+	
+	LODMaxCount[EMassLOD::High] = 50;
+	LODMaxCount[EMassLOD::Medium] = 100;
+	LODMaxCount[EMassLOD::Low] = 500;
+	LODMaxCount[EMassLOD::Off] = 0;
+}
+
+void UMassVisualizationLODProcessor::ConfigureQueries()
+{
+	CloseEntityQuery.AddRequirement<FMassLODInfoFragment>(ELWComponentAccess::ReadOnly);
+	CloseEntityQuery.AddRequirement<FMassRepresentationLODFragment>(ELWComponentAccess::ReadWrite);
+	CloseEntityQuery.AddRequirement<FDataFragment_Transform>(ELWComponentAccess::ReadOnly);
+
+	FarEntityQuery = CloseEntityQuery;
+
+	CloseEntityQuery.AddTagRequirement<FMassVisibilityCulledByDistanceTag>(ELWComponentPresence::None);
+	
+	FarEntityQuery.AddTagRequirement<FMassVisibilityCulledByDistanceTag>(ELWComponentPresence::All);
+	FarEntityQuery.AddChunkRequirement<FMassVisualizationChunkFragment>(ELWComponentAccess::ReadOnly);
+}
+
+void UMassVisualizationLODProcessor::Initialize(UObject& Owner)
+{
+	Super::Initialize(Owner);
+
+	LODCalculator.Initialize(BaseLODDistance, BufferHysteresisOnDistancePercentage / 100.f, LODMaxCount, nullptr, VisibleLODDistance);
+}
+
+void UMassVisualizationLODProcessor::PrepareExecution()
+{
+	check(World);
+
+	check(LODManager);
+	const TArray<FViewerInfo>& Viewers = LODManager->GetViewers();
+	LODCalculator.PrepareExecution(Viewers);
+}
+
+void UMassVisualizationLODProcessor::Execute(UEntitySubsystem& EntitySubsystem, FLWComponentSystemExecutionContext& Context)
+{
+	ExecuteInternal(EntitySubsystem, Context);
+}
\ No newline at end of file
diff --git a/Engine/Plugins/Runtime/MassGameplay/Source/MassRepresentation/Private/MassVisualizationTrait.cpp b/Engine/Plugins/Runtime/MassGameplay/Source/MassRepresentation/Private/MassVisualizationTrait.cpp
new file mode 100644
--- /dev/null
+++ b/Engine/Plugins/Runtime/MassGameplay/Source/MassRepresentation/Private/MassVisualizationTrait.cpp
@@ -0,0 +1,45 @@
+﻿// Copyright Epic Games, Inc. All Rights Reserved.
+
+#include "MassVisualizationTrait.h"
+#include "MassEntityTemplateRegistry.h"
+#include "MassRepresentationSubsystem.h"
+#include "MassCommonFragments.h"
+#include "MassRepresentationFragments.h"
+#include "MassRepresentationProcessor.h"
+#include "Engine/World.h"
+
+UMassVisualizationTrait::UMassVisualizationTrait()
+{
+	RepresentationSubsystemClass = UMassRepresentationSubsystem::StaticClass();
+	RepresentationFragmentDeinitializerClass = UMassRepresentationFragmentDestructor::StaticClass();
+}
+
+void UMassVisualizationTrait::BuildTemplate(FMassEntityTemplateBuildContext& BuildContext, UWorld& World) const
+{
+	// This should not be ran on NM_Server network mode
+	if (World.IsNetMode(NM_DedicatedServer))
+	{
+		return;
+	}
+
+	// the following needs to be always there for mesh vis to work. Adding following fragments after already 
+	// adding Config.AdditionalDataFragments to let user configure the fragments first. Calling BuildContext.Add() 
+	// won't override any fragments that are already there
+	BuildContext.AddFragment<FDataFragment_Transform>();
+
+	if (UMassRepresentationSubsystem* RepresentationSubsystem = Cast<UMassRepresentationSubsystem>(World.GetSubsystemBase(RepresentationSubsystemClass)))
+	{
+		FMassRepresentationFragment& RepresentationFragment = BuildContext.AddFragment_GetRef<FMassRepresentationFragment>();
+		RepresentationFragment.StaticMeshDescIndex = RepresentationSubsystem->FindOrAddStaticMeshDesc(StaticMeshInstanceDesc);
+		RepresentationFragment.HighResTemplateActorIndex = HighResTemplateActor.Get() ? RepresentationSubsystem->FindOrAddTemplateActor(HighResTemplateActor.Get()) : INDEX_NONE;
+		RepresentationFragment.LowResTemplateActorIndex = LowResTemplateActor.Get() ? RepresentationSubsystem->FindOrAddTemplateActor(LowResTemplateActor.Get()) : INDEX_NONE;
+
+		BuildContext.AddDeinitializer(*(RepresentationFragmentDeinitializerClass->GetDefaultObject<UPipeProcessor>()));
+
+		BuildContext.AddFragment<FMassRepresentationLODFragment>();
+
+		BuildContext.AddChunkFragment<FMassVisualizationChunkFragment>();
+	}
+}
+
+
diff --git a/Engine/Plugins/Runtime/MassGameplay/Source/MassRepresentation/Public/MassRepresentationFragments.h b/Engine/Plugins/Runtime/MassGameplay/Source/MassRepresentation/Public/MassRepresentationFragments.h
new file mode 100644
--- /dev/null
+++ b/Engine/Plugins/Runtime/MassGameplay/Source/MassRepresentation/Public/MassRepresentationFragments.h
@@ -0,0 +1,56 @@
+// Copyright Epic Games, Inc. All Rights Reserved.
+
+#pragma once
+
+#include "MassLODTypes.h"
+#include "LWComponentTypes.h"
+#include "EntitySubsystem.h"
+#include "MassActorSpawnerSubsystem.h"
+#include "MassRepresentationTypes.h"
+
+#include "MassRepresentationFragments.generated.h"
+
+USTRUCT()
+struct MASSREPRESENTATION_API FMassRepresentationLODFragment : public FLWComponentData
+{
+	GENERATED_BODY()
+
+	/** Saved closest ViewerDistance */
+	float ClosestViewerDistanceSq = FLT_MAX;
+
+	/** LOD information */
+	TEnumAsByte<EMassLOD::Type> LOD = EMassLOD::Max;
+	TEnumAsByte<EMassLOD::Type> PrevLOD = EMassLOD::Max;
+
+	/** Visibility Info */
+	bool bIsVisibleByAViewer = false;
+	bool bWasVisibleByAViewer = false;
+	bool bIsInVisibleRange = false;
+	bool bWasInVisibleRange = false;
+
+	/** Value scaling from 0 to 3, 0 highest LOD we support and 3 being completely off LOD */
+	float LODSignificance = 0.0f;
+};
+
+USTRUCT()
+struct MASSREPRESENTATION_API FMassRepresentationFragment : public FLWComponentData
+{
+	GENERATED_BODY()
+
+	ERepresentationType CurrentRepresentation = ERepresentationType::None;
+
+	ERepresentationType PrevRepresentation = ERepresentationType::None;
+
+	int16 HighResTemplateActorIndex = INDEX_NONE;
+
+	int16 LowResTemplateActorIndex = INDEX_NONE;
+
+	int16 StaticMeshDescIndex = INDEX_NONE;
+
+	FMassHandle_ActorSpawnRequest ActorSpawnRequestHandle;
+
+	FTransform PrevTransform;
+
+	// @todo: Remove this when we address character movement
+	bool bControlCharacterMovement = true;
+};
\ No newline at end of file
diff --git a/Engine/Plugins/Runtime/MassGameplay/Source/MassRepresentation/Public/MassRepresentationProcessor.h b/Engine/Plugins/Runtime/MassGameplay/Source/MassRepresentation/Public/MassRepresentationProcessor.h
new file mode 100644
--- /dev/null
+++ b/Engine/Plugins/Runtime/MassGameplay/Source/MassRepresentation/Public/MassRepresentationProcessor.h
@@ -0,0 +1,220 @@
+// Copyright Epic Games, Inc. All Rights Reserved.
+
+#pragma once
+
+#include "MassRepresentationFragments.h"
+#include "MassActorSpawnerSubsystem.h"
+#include "MassTranslator.h"
+
+#include "MassRepresentationProcessor.generated.h"
+
+class UMassRepresentationSubsystem;
+struct FDataFragment_Actor;
+
+enum EActorEnabledType
+{
+	Disabled,
+	LowRes,
+	HighRes,
+};
+
+UCLASS()
+class MASSREPRESENTATION_API UMassRepresentationProcessor : public UPipeProcessor
+{
+	GENERATED_BODY()
+
+public:
+	UMassRepresentationProcessor();
+
+protected:
+
+	/** Configure the owned FLWComponentQuery instances to express processor's requirements */
+	virtual void ConfigureQueries() override;
+
+	/**
+	 * Initialize the processor 
+	 * @param Owner of the Processor
+	 */
+	virtual void Initialize(UObject& Owner) override;
+
+	/** 
+	 * Execution method for this processor 
+	 * @param EntitySubsystem is the system to execute the lambdas on each entity chunk
+	 * @param Context is the execution context to be passed when executing the lambdas
+	 */
+	virtual void Execute(UEntitySubsystem& EntitySubsystem, FLWComponentSystemExecutionContext& Context) override;
+
+	/** 
+	 * Returns an actor of the template type and setup fragments values from it
+	 * @param MassAgent is the handle to the associated mass agent
+	 * @param ActorInfo is the fragment where we are going to store the actor pointer
+	 * @param Transform is the spatial information about where to spawn the actor 
+	 * @param TemplateActorIndex is the index of the type fetched with UMassRepresentationSubsystem::FindOrAddTemplateActor()
+	 * @param SpawnRequestHandle (in/out) In: previously requested spawn Out: newly requested spawn
+	 * @param Priority of this spawn request in comparison with the others, lower value means higher priority
+	 * @return the actor spawned
+	 */
+	virtual AActor* GetOrSpawnActor(const FMassHandle MassAgent, FDataFragment_Actor& ActorInfo, const FTransform& Transform, const int16 TemplateActorIndex, FMassHandle_ActorSpawnRequest& SpawnRequestHandle, const float Priority);
+
+	/**
+	 * Release the actor to the subsystem, will only release it the actor or spawn request matches the template actor
+	 * @param MassAgent is the handle to the associated mass agent
+	 * @param ActorInfo is the fragment where we are going to store the actor pointer
+	 * @param TemplateActorIndex is the index of the type to release
+	 * @param SpawnRequestHandle (in/out) In: previously requested spawn to cancel if any
+	 * @param Context of the execution from the entity sub system
+	 * @param bCancelSpawningOnly tell to only cancel the existing spawning request and to not release the associated actor it any.
+	 * @return if the actor was release or the spawning was canceled.
+	 */
+	virtual bool ReleaseActorOrCancelSpawning(const FMassHandle MassAgent, FDataFragment_Actor& ActorInfo, const int16 TemplateActorIndex, FMassHandle_ActorSpawnRequest& SpawnRequestHandle, FLWComponentSystemExecutionContext& Context, const bool bCancelSpawningOnly = false);
+
+	/** 
+	 * Enable/disable a spawned actor
+	 * @param EnabledType is the type of enabling to do on this actor
+	 * @param Actor is the actual actor to perform enabling type on
+	 * @param EntityIdx is the entity index currently processing
+	 * @param Context is the current light weight component execution context 
+	 */
+	virtual void SetActorEnabled(const EActorEnabledType EnabledType, AActor& Actor, const int32 EntityIdx, FLWComponentSystemExecutionContext& Context);
+
+	/**
+	 * Teleports the actor at the specified transform by preserving its velocity and without collision.
+	 * The destination will be adjusted to fit an existing capsule.
+	 * @param Transform is the new actor's transform 
+	 * @param Actor is the actual actor to teleport
+	 * @param Context is the current light weight component execution context
+	 */
+	virtual void TeleportActor(const FTransform& Transform, AActor& Actor, FLWComponentSystemExecutionContext& Context);
+
+	/**
+	 * Method that will be bound to a delegate called before the spawning of an actor to let the requester prepare it
+	 * @param SpawnRequestHandle the handle of the spawn request that is about to spawn
+	 * @param SpawnRequest of the actor that is about to spawn
+	 */
+	void OnActorPreSpawn(const FMassHandle_ActorSpawnRequest& SpawnRequestHandle, const FStructView& SpawnRequest);
+
+	/**
+	 * Method that will be bound to a delegate used post-spawn to notify and let the requester configure the actor
+	 * @param SpawnRequestHandle the handle of the spawn request that was just spawned
+	 * @param SpawnRequest of the actor that just spawned
+	 * @return The action to take on the spawn request, either keep it there or remove it.
+	 */
+	EMassActorSpawnRequestAction OnActorPostSpawn(const FMassHandle_ActorSpawnRequest& SpawnRequestHandle, const FStructView& SpawnRequest);
+
+	/**
+	 * Method called once the spawning of an actor just occurred, give the processors a chance to configure the actor during the spawn time slicing
+	 * @param MassAgent mass entity to the spawn actor
+	 * @param SpawnedActor that was just spawned
+	 */
+	virtual void PostActorSpawned(const FMassHandle MassAgent, AActor& SpawnedActor) {}
+
+	/**
+	 * Updates chunk visibility info for later chunk logic optimization
+	 * @param Context of the execution from the entity sub system
+	 * @return The visualization chunk fragment
+	 */
+	FMassVisualizationChunkFragment& UpdateChunkVisibility(FLWComponentSystemExecutionContext& Context) const;
+
+	/**
+	 * Updates entity visibility tag for later chunk logic optimization
+	 * @param Entity of the entity to update visibility on
+	 * @param Representation fragment containing the current and previous visual state
+	 * @param RepresentationLOD fragment containing the visibility information
+	 * @param ChunkData is the visualization chunk fragment
+     * @param Context of the execution from the entity sub system
+	 */
+	static void UpdateEntityVisibility(const FLWEntity Entity, const FMassRepresentationFragment& Representation, const FMassRepresentationLODFragment& RepresentationLOD, FMassVisualizationChunkFragment& ChunkData, FLWComponentSystemExecutionContext& Context);
+
+public:
+	/**
+	 * Release an actor or cancel its spawning (calls ReleaseAnyActorOrCancelAnySpawning)
+	 * WARNING: This method will destroy the associated actor in any and by the same fact might also move the entity into a new archetype.
+	 *          So any reference to fragment might become invalid.
+	 * @param RepresentationSubsystem to use to release any actors or cancel spawning requests
+	 * @param MassAgent is the handle to the associated mass agent
+	 * @return True if actor was release or spawning request was canceled
+	 */
+	static void ReleaseAnyActorOrCancelAnySpawning(UMassRepresentationSubsystem& RepresentationSubsystem, const FMassHandle MassAgent);
+
+	/**
+	 * Release an actor or cancel its spawning
+	 * WARNING: This method will destroy the associated actor in any and by the same fact might also move the entity into a new archetype.
+	 *          So any reference to fragment might become invalid if you are not within the pipe execution
+	 * @param RepresentationSubsystem to use to release any actors or cancel spawning requests
+	 * @param MassAgent is the handle to the associated mass agent
+	 * @param ActorInfo is the fragment where we are going to store the actor pointer
+	 * @param Representation fragment containing the current and previous visual state
+	 */
+	static void ReleaseAnyActorOrCancelAnySpawning(UMassRepresentationSubsystem& RepresentationSubsystem, const FMassHandle MassAgent, FDataFragment_Actor& ActorInfo, FMassRepresentationFragment& Representation);
+
+protected:
+
+	virtual float GetSpawnPriority(const FMassRepresentationLODFragment& Representation)
+	{
+		return Representation.LODSignificance;
+	}
+
+	/*
+	 * Update representation type for each entity, must be called within a ForEachEntityChunk
+	 * @param Context of the execution from the entity sub system
+	 */
+	void UpdateRepresentation(FLWComponentSystemExecutionContext& Context);
+
+	/** 
+	 * Update representation and visibility for each entity, must be called within a ForEachEntityChunk
+	 * @param Context of the execution from the entity sub system
+	 */
+	void UpdateVisualization(FLWComponentSystemExecutionContext& Context);
+
+	/** What should be the representation of this entity for each specificLOD */
+	UPROPERTY(EditAnywhere, Category = "Mass|Representation", config)
+	ERepresentationType LODRepresentation[EMassLOD::Max];
+
+	/** If true, LowRes actors will be kept around, disabled, whilst StaticMeshInstance representation is active */
+	UPROPERTY(EditAnywhere, Category = "Mass|Representation", config)
+	bool bKeepLowResActors = true;
+
+	/** World Partition grid name to test collision against, default None will be the main grid */
+	UPROPERTY(EditAnywhere, Category = "Mass|Representation", config)
+	FName WorldPartitionGridNameContainingCollision;
+
+	/** Caching ptr to our associated world */
+	UPROPERTY(Transient)
+	UWorld* World;
+
+	UPROPERTY(Transient)
+	UPipeEntitySubsystem* PipeEntitySubsystem;
+
+	/** A cache pointer to the representation subsystem */
+	UPROPERTY(Transient)
+	UMassRepresentationSubsystem* RepresentationSubsystem;
+
+	/** Default representation when unable to spawn an actor */
+	ERepresentationType DefaultRepresentationType = ERepresentationType::None;
+
+	FLWComponentQuery EntityQuery;
+
+	/** At what rate should the not visible entity be updated in seconds */
+	UPROPERTY(EditAnywhere, Category = "Mass|Visualization", config)
+	float NotVisibleUpdateRate = 0.5f;
+};
+
+UCLASS()
+class MASSREPRESENTATION_API UMassRepresentationFragmentDestructor : public UMassFragmentDestructor
+{
+	GENERATED_BODY()
+
+public:
+	UMassRepresentationFragmentDestructor();
+
+protected:
+	virtual void Initialize(UObject& Owner) override;
+	virtual void ConfigureQueries() override;
+	virtual void Execute(UEntitySubsystem& EntitySubsystem, FLWComponentSystemExecutionContext& Context) override;
+
+	/** A cache pointer to the representation subsystem */
+	UPROPERTY(Transient)
+	UMassRepresentationSubsystem* RepresentationSubsystem;
+
+	FLWComponentQuery EntityQuery;
+};
\ No newline at end of file
diff --git a/Engine/Plugins/Runtime/MassGameplay/Source/MassRepresentation/Public/MassRepresentationSubsystem.h b/Engine/Plugins/Runtime/MassGameplay/Source/MassRepresentation/Public/MassRepresentationSubsystem.h
new file mode 100644
--- /dev/null
+++ b/Engine/Plugins/Runtime/MassGameplay/Source/MassRepresentation/Public/MassRepresentationSubsystem.h
@@ -0,0 +1,167 @@
+// Copyright Epic Games, Inc. All Rights Reserved.
+
+#pragma once
+
+#include "GameFramework/Actor.h"
+#include "MassCommonTypes.h"
+#include "Misc/MTAccessDetector.h"
+#include "MassRepresentationTypes.h"
+#include "MassActorSpawnerSubsystem.h"
+#include "Subsystems/WorldSubsystem.h"
+#include "MassRepresentationSubsystem.generated.h"
+
+class UMassVisualizationComponent;
+class AMassVisualizer;
+class UHierarchicalInstancedStaticMeshComponent;
+struct FStaticMeshInstanceVisualizationDesc;
+struct FMassInstancedStaticMeshInfo;
+struct FMassHandle_ActorSpawnRequest;
+class UMassActorSpawnerSubsystem;
+class UMassAgentComponent;
+class UPipeEntitySubsystem;
+enum class EPipeProcessingPhase : uint8;
+class UWorldPartitionSubsystem;
+
+/**
+ * Subsystem responsible for all visual of mass agents, will handle actors spawning and static mesh instances
+ */
+UCLASS()
+class MASSREPRESENTATION_API UMassRepresentationSubsystem : public UWorldSubsystem
+{
+	GENERATED_BODY()
+
+public:
+	/** 
+	 * Get the index of the static mesh visual type, will add a new one if does not exist  
+	 * @param Desc is the information for the static mesh that will be instantiated later via AddStaticMeshInstance()
+	 * @return The index of the static mesh type 
+	 */
+	int16 FindOrAddStaticMeshDesc(const FStaticMeshInstanceVisualizationDesc& Desc) const;
+
+	/** @todo: need to add a release API at some point for static mesh types */
+
+	/** 
+	 * @return the array of all the static mesh instance component information
+	 */
+	FMassInstancedStaticMeshInfoArrayView GetMutableInstancedStaticMeshInfos() const;
+
+	/** Mark render state of the static mesh instances dirty */
+	void DirtyStaticMeshInstances() const;
+
+	/** 
+	 * Store the template actor uniquely and return an index to it 
+	 * @param ActorClass is a template actor class we will need to spawn for an agent 
+	 * @return The index of the template actor type
+	 */
+	int16 FindOrAddTemplateActor(const TSubclassOf<AActor>& ActorClass);
+
+	/** 
+	 * Get or spawn an actor from the TemplateActorIndex
+	 * @param MassAgent is the handle to the associated mass agent
+	 * @param Transform where to create this actor
+	 * @param TemplateActorIndex is the index of the type fetched with FindOrAddTemplateActor()
+	 * @param SpawnRequestHandle [IN/OUT] IN: previously requested spawn OUT: newly requested spawn
+	 * @param Priority of this spawn request in comparison with the others, lower value means higher priority (optional)
+	 * @param ActorPreSpawnDelegate is an optional delegate called before the spawning of an actor
+	 * @param ActorPostSpawnDelegate is an optional delegate called once the actor is spawned
+	 * @return The spawned actor from the template actor type if ready
+	 */
+	AActor* GetOrSpawnActorFromTemplate(const FMassHandle MassAgent, const FTransform& Transform, const int16 TemplateActorIndex, FMassHandle_ActorSpawnRequest& SpawnRequestHandle, float Priority = MAX_FLT,
+		FMassActorPreSpawnDelegate ActorPreSpawnDelegate = FMassActorPreSpawnDelegate(), FMassActorPostSpawnDelegate ActorPostSpawnDelegate = FMassActorPostSpawnDelegate());
+
+	/**
+	 * Cancel spawning request that is matching the TemplateActorIndex
+	 * @param MassAgent is the handle to the associated mass agent
+	 * @param TemplateActorIndex is the template type of the actor to release in case it was successfully spawned
+	 * @param SpawnRequestHandle [IN/OUT] previously requested spawn, gets invalidated as a result of this call.
+	 * @return True if spawning request was canceled
+	 */
+	bool CancelSpawning(const FMassHandle MassAgent, const int16 TemplateActorIndex, FMassHandle_ActorSpawnRequest & SpawnRequestHandle);
+
+	/**
+	 * Release an actor that is matching the TemplateActorIndex
+	 * @param MassAgent is the handle to the associated mass agent
+	 * @param TemplateActorIndex is the template type of the actor to release in case it was successfully spawned
+	 * @param ActorToRelease is the actual actor to release if any
+	 * @param bImmediate means it needs to be done immediately and not queue for later
+	 * @return True if actor was released
+	 */
+	bool ReleaseTemplateActor(const FMassHandle MassAgent, const int16 TemplateActorIndex, AActor* ActorToRelease, bool bImmediate);
+
+	/**
+	 * Release an actor or cancel its spawning if it is matching the TemplateActorIndex
+	 * @param MassAgent is the handle to the associated mass agent
+	 * @param TemplateActorIndex is the template type of the actor to release in case it was successfully spawned
+	 * @param ActorToRelease is the actual actor to release if any
+	 * @param SpawnRequestHandle [IN/OUT] previously requested spawn, gets invalidated as a result of this call.
+	 * @return True if actor was released or spawning request was canceled
+	 */
+	bool ReleaseTemplateActorOrCancelSpawning(const FMassHandle MassAgent, const int16 TemplateActorIndex, AActor* ActorToRelease, FMassHandle_ActorSpawnRequest& SpawnRequestHandle);
+
+
+	/**
+	 * Compare if an actor matches the registered template actor
+	 * @param Actor to compare its class against the template
+	 * @param TemplateActorIndex is the template type of the actor to compare against
+	 * @return True if actor matches the template
+	 */
+	bool DoesActorMatchTemplate(const AActor& Actor, const int16 TemplateActorIndex);
+
+	TSubclassOf<AActor> GetTemplateActorClass(const int16 TemplateActorIndex);
+
+	bool IsCollisionLoaded(const FName TargetGrid, const FTransform& Transform) const;
+
+	/**
+	 * Release all references to static meshes and template actors
+	 * Use with caution, all entities using this representation subsystem must be destroy otherwise they will point to invalid resources */
+	 void ReleaseAllResources();
+
+protected:
+	// USubsystem BEGIN
+	virtual void Initialize(FSubsystemCollectionBase& Collection) override;
+	virtual void Deinitialize() override;
+	// USubsystem END
+
+	/** Needed for batching the update of static mesh transform */
+	void OnProcessingPhaseStarted(const float DeltaSeconds, const EPipeProcessingPhase Phase) const;
+
+	void OnMassAgentComponentEntityAssociated(const UMassAgentComponent& AgentComponent);
+	void OnMassAgentComponentEntityDetaching(const UMassAgentComponent& AgentComponent);
+
+	bool ReleaseTemplateActorInternal(const int16 TemplateActorIndex, AActor* ActorToRelease, bool bImmediate);
+	bool CancelSpawningInternal(const int16 TemplateActorIndex, FMassHandle_ActorSpawnRequest& SpawnRequestHandle);
+
+protected:
+
+	/** The array of all the template actors */
+	UPROPERTY(Transient)
+	TArray<TSubclassOf<AActor>> TemplateActors;
+	UE_MT_DECLARE_RW_ACCESS_DETECTOR(TemplateActorsMTAccessDetector);
+
+	/** The component that handles all the static mesh instances */
+	UPROPERTY(Transient)
+	UMassVisualizationComponent* VisualizationComponent;
+
+	/** The actor owning the above visualization component */
+	UPROPERTY(Transient)
+	AMassVisualizer* Visualizer;
+
+	UPROPERTY(Transient)
+	UMassActorSpawnerSubsystem* ActorSpawnerSubsystem;
+
+	UPROPERTY(Transient)
+	UPipeEntitySubsystem* EntitySubsystem;
+
+	UPROPERTY(Transient)
+	UWorldPartitionSubsystem* WorldPartitionSubsystem;
+
+	/** The time to wait before retrying a to spawn actor that failed */
+	float RetryMovedDistanceSq = 1000000.0f;
+
+	/** The distance a failed spawned actor needs to move before we retry */
+	float RetryTimeInterval = 10.0f;
+
+	/** Keeping track of all the mass agent this subsystem is responsible for spawning actors */
+	TMap<FMassHandle, int32> HandledMassAgents;
+};
+
diff --git a/Engine/Plugins/Runtime/MassGameplay/Source/MassRepresentation/Public/MassRepresentationTypes.h b/Engine/Plugins/Runtime/MassGameplay/Source/MassRepresentation/Public/MassRepresentationTypes.h
new file mode 100644
--- /dev/null
+++ b/Engine/Plugins/Runtime/MassGameplay/Source/MassRepresentation/Public/MassRepresentationTypes.h
@@ -0,0 +1,302 @@
+// Copyright Epic Games, Inc. All Rights Reserved.
+
+#pragma once
+
+#include "MassLODTypes.h"
+#include "Engine/DataTable.h"
+#include "Misc/MTAccessDetector.h"
+
+#include "MassRepresentationTypes.generated.h"
+
+class UMaterialInterface;
+class UStaticMesh;
+
+DECLARE_LOG_CATEGORY_EXTERN(LogMassRepresentation, Log, All);
+
+namespace UE::Mass::ProcessorGroupNames
+{
+	const FName Representation = FName(TEXT("Representation"));
+}
+
+UENUM()
+enum class ERepresentationType : uint8
+{
+	HighResSpawnedActor,
+	LowResSpawnedActor,
+	StaticMeshInstance,
+	None,
+};
+
+USTRUCT()
+struct FStaticMeshInstanceVisualizationMeshDesc
+{
+	GENERATED_BODY()
+	
+	/** The static mesh visual representation */
+	UPROPERTY(EditAnywhere, Category = "Mass|Visual")
+	UStaticMesh* Mesh = nullptr;
+
+	/**
+	 * Material overrides for the static mesh visual representation. 
+	 * 
+	 * Array indices correspond to material slot indices on the static mesh.
+	 */
+	UPROPERTY(EditAnywhere, Category = "Mass|Visual")
+	TArray<UMaterialInterface*> MaterialOverrides;
+
+	/** The minimum inclusive LOD significance to start using this static mesh */
+	UPROPERTY(EditAnywhere, Category = "Mass|Visual")
+	float MinLODSignificance = float(EMassLOD::High);
+
+	/** The maximum exclusive LOD significance to stop using this static mesh */
+	UPROPERTY(EditAnywhere, Category = "Mass|Visual")
+	float MaxLODSignificance = float(EMassLOD::Max);
+
+	/** Controls whether the ISM can cast shadow or not */
+	UPROPERTY(EditAnywhere, Category = "Mass|Visual")
+	bool bCastShadows = false;
+
+	bool operator==(const FStaticMeshInstanceVisualizationMeshDesc& Other) const
+	{
+		return Mesh == Other.Mesh && 
+			MaterialOverrides == Other.MaterialOverrides && 
+			FMath::IsNearlyEqual(MinLODSignificance, Other.MinLODSignificance, KINDA_SMALL_NUMBER) &&
+			FMath::IsNearlyEqual(MaxLODSignificance, Other.MaxLODSignificance, KINDA_SMALL_NUMBER) &&
+			bCastShadows == Other.bCastShadows;
+	}
+};
+
+USTRUCT()
+struct FStaticMeshInstanceVisualizationDesc : public FTableRowBase
+{
+	GENERATED_BODY()
+
+	/** 
+	 * Mesh descriptions. These will be instanced together using the same transform for each, to 
+	 * visualize the given agent.
+	 */
+	UPROPERTY(EditAnywhere, Category = "Mass|Visual")
+	TArray<FStaticMeshInstanceVisualizationMeshDesc> Meshes;
+
+	/** Boolean to enable code to transform the static meshes if not align the mass agent transform */
+	UPROPERTY(EditAnywhere, Category = "Mass|Visual")
+	bool bUseTransformOffset = false;
+
+	/** Transform to offset the static meshes if not align the mass agent transform */
+	UPROPERTY(EditAnywhere, Category = "Mass|Visual", meta=(EditCondition="bUseTransformOffset"))
+	FTransform TransformOffset;
+
+	bool operator==(const FStaticMeshInstanceVisualizationDesc& Other) const
+	{
+		return Meshes == Other.Meshes;
+	}
+};
+
+class UInstancedStaticMeshComponent;
+
+USTRUCT()
+struct MASSREPRESENTATION_API FISMCSharedData
+{
+	GENERATED_BODY()
+
+	FISMCSharedData() = default;
+
+	FISMCSharedData(UInstancedStaticMeshComponent* InISMC)
+		: ISMC(InISMC),
+		RefCount(1)
+	{
+
+	}
+
+	UInstancedStaticMeshComponent* ISMC = nullptr;
+
+	int32 RefCount = 1;
+
+	/** Buffer holding current frame transforms for the static mesh instances, used to batch update the transforms */
+	TArray<int32> UpdateInstanceIds;
+	TArray<FTransform> StaticMeshInstanceTransforms;
+	TArray<FTransform> StaticMeshInstancePrevTransforms;
+
+	/** Buffer holding current frame custom floats for the static mesh instances, used to batch update the ISMs custom data */
+	TArray<float> StaticMeshInstanceCustomFloats;
+
+	// When initially adding to StaticMeshInstanceCustomFloats, can use the size as the write iterator, but on subsequent processors, we need to know where to start writing
+	int32 WriteIterator = 0;
+};
+
+USTRUCT()
+struct MASSREPRESENTATION_API FMassLODSignificanceRange
+{
+	GENERATED_BODY()
+public:
+
+	void AddBatchedTransform(const int32 InstanceId, const FTransform& Transform, const FTransform& PrevTransform);
+
+	// Adds the specified struct reinterpreted as custom floats to our custom data. Individual members of the specified struct should always fit into a float.
+	// When adding any custom data, the custom data must be added for every instance.
+	template<typename InCustomDataType>
+	void AddBatchedCustomData(InCustomDataType InCustomData, int32 NumFloatsToPad = 0)
+	{
+		check(ISMCSharedDataPtr);
+		static_assert((sizeof(InCustomDataType) % sizeof(float)) == 0, "AddBatchedCustomData: InCustomDataType should have a total size multiple of sizeof(float), and have members that fit in a float's boundaries");
+		const size_t StructSize = sizeof(InCustomDataType);
+		const size_t StructSizeInFloats = StructSize / sizeof(float);
+		for (int i = 0; i < StaticMeshRefs.Num(); i++)
+		{
+			FISMCSharedData& SharedData = (*ISMCSharedDataPtr)[StaticMeshRefs[i]];
+			const int32 StartIndex = SharedData.StaticMeshInstanceCustomFloats.AddDefaulted(StructSizeInFloats + NumFloatsToPad);
+			InCustomDataType* CustomData = reinterpret_cast<InCustomDataType*>(&SharedData.StaticMeshInstanceCustomFloats[StartIndex]);
+			*CustomData = InCustomData;
+		}
+	}
+
+	void AddBatchedCustomDataFloats(const TArray<float>& CustomFloats);
+
+	void WriteCustomDataFloatsAtStartIndex(int32 StaticMeshIndex, const TArrayView<float> CustomFloats, int32 FloatsPerInstance, int32 StartIndex);
+
+	/** LOD Significance range */
+	float MinSignificance;
+	float MaxSignificance;
+
+	/** The component handling these instances */
+	UPROPERTY()
+	TArray<UStaticMesh*> StaticMeshRefs;
+
+	TMap<UStaticMesh*, FISMCSharedData>* ISMCSharedDataPtr = nullptr;
+};
+
+USTRUCT()
+struct MASSREPRESENTATION_API FMassInstancedStaticMeshInfo
+{
+	GENERATED_BODY()
+public:
+
+	FMassInstancedStaticMeshInfo() = default;
+
+	FMassInstancedStaticMeshInfo(const FStaticMeshInstanceVisualizationDesc& InDesc)
+		: Desc(InDesc)
+	{
+	}
+
+	bool operator==(const FStaticMeshInstanceVisualizationDesc& OtherDesc) const
+	{
+		return Desc == OtherDesc;
+	}
+
+	/** Whether or not to transform the static meshes if not align the mass agent transform */
+	bool ShouldUseTransformOffset() const { return Desc.bUseTransformOffset; }
+	FTransform GetTransformOffset() const { return Desc.TransformOffset; }
+
+	FORCEINLINE FMassLODSignificanceRange* GetLODSignificanceRange(float LODSignificance)
+	{
+		for (FMassLODSignificanceRange& Range : LODSignificanceRanges)
+		{
+			if (LODSignificance >= Range.MinSignificance && LODSignificance < Range.MaxSignificance)
+			{
+				return &Range;
+			}
+		}
+		return nullptr;
+	}
+
+	FORCEINLINE void AddBatchedTransform(const int32 InstanceId, const FTransform& Transform, const FTransform& PrevTransform, float LODSignificance)
+	{
+		if (FMassLODSignificanceRange* Range = GetLODSignificanceRange(LODSignificance))
+		{
+			Range->AddBatchedTransform(InstanceId, Transform, PrevTransform);
+		}
+	}
+
+	// Adds the specified struct reinterpreted as custom floats to our custom data. Individual members of the specified struct should always fit into a float.
+	// When adding any custom data, the custom data must be added for every instance.
+	template<typename InCustomDataType>
+	void AddBatchedCustomData(InCustomDataType InCustomData, float LODSignificance, int32 NumFloatsToPad = 0)
+	{
+		if (FMassLODSignificanceRange* Range = GetLODSignificanceRange(LODSignificance))
+		{
+			Range->AddBatchedCustomData(InCustomData, NumFloatsToPad);
+		}
+	}
+
+	FORCEINLINE void AddBatchedCustomDataFloats(const TArray<float>& CustomFloats, float LODSignificance)
+	{
+		if (FMassLODSignificanceRange* Range = GetLODSignificanceRange(LODSignificance))
+		{
+			Range->AddBatchedCustomDataFloats(CustomFloats);
+		}
+	}
+
+	FORCEINLINE void WriteCustomDataFloatsAtStartIndex(int32 StaticMeshIndex, const TArrayView<float> CustomFloats, float LODSignificance, int32 FloatsPerInstance, int32 FloatStartIndex)
+	{
+		if (FMassLODSignificanceRange* Range = GetLODSignificanceRange(LODSignificance))
+		{
+			Range->WriteCustomDataFloatsAtStartIndex(StaticMeshIndex, CustomFloats, FloatsPerInstance, FloatStartIndex);
+		}
+	}
+
+protected:
+
+	/** Destroy the visual instance */
+	void ClearVisualInstance(TMap<UStaticMesh*, FISMCSharedData>& ISMCSharedData);
+
+	/** Information about this static mesh which will represent all instances */
+	UPROPERTY()
+		FStaticMeshInstanceVisualizationDesc Desc;
+
+	/** The component handling these instances */
+	UPROPERTY()
+		TArray<UInstancedStaticMeshComponent*> InstancedStaticMeshComponents;
+
+	UPROPERTY()
+		TArray<FMassLODSignificanceRange> LODSignificanceRanges;
+	friend class UMassVisualizationComponent;
+};
+
+#if ENABLE_MT_DETECTOR
+
+#define MAKE_MASS_INSTANCED_STATIC_MESH_INFO_ARRAY_VIEW(ArrayView, AccessDetector) FMassInstancedStaticMeshInfoArrayView(ArrayView, AccessDetector)
+struct FMassInstancedStaticMeshInfoArrayViewAccessDetector
+{
+	FMassInstancedStaticMeshInfoArrayViewAccessDetector(TArrayView<FMassInstancedStaticMeshInfo> InInstancedStaticMeshInfos, const FRWRecursiveAccessDetector& InAccessDetector)
+		: InstancedStaticMeshInfos(InInstancedStaticMeshInfos)
+		, AccessDetector(&InAccessDetector)
+	{
+		UE_MT_ACQUIRE_WRITE_ACCESS(*AccessDetector);
+	}
+
+	FMassInstancedStaticMeshInfoArrayViewAccessDetector(FMassInstancedStaticMeshInfoArrayViewAccessDetector&& Other)
+		: InstancedStaticMeshInfos(Other.InstancedStaticMeshInfos)
+		, AccessDetector(Other.AccessDetector)
+	{
+		Other.AccessDetector = nullptr;
+	}
+	FMassInstancedStaticMeshInfoArrayViewAccessDetector(const FMassInstancedStaticMeshInfoArrayViewAccessDetector& Other) = delete;
+	void operator=(const FMassInstancedStaticMeshInfoArrayViewAccessDetector& Other) = delete;
+
+	~FMassInstancedStaticMeshInfoArrayViewAccessDetector()
+	{
+		if(AccessDetector)
+		{
+			UE_MT_RELEASE_WRITE_ACCESS(*AccessDetector);
+		}
+	}
+
+	FORCEINLINE FMassInstancedStaticMeshInfo& operator[](int32 Index) const
+	{
+		return InstancedStaticMeshInfos[Index];
+	}
+
+private:
+	TArrayView<FMassInstancedStaticMeshInfo> InstancedStaticMeshInfos;
+	const FRWRecursiveAccessDetector* AccessDetector;
+};
+
+typedef FMassInstancedStaticMeshInfoArrayViewAccessDetector FMassInstancedStaticMeshInfoArrayView;
+
+#else // ENABLE_MT_DETECTOR
+
+#define MAKE_MASS_INSTANCED_STATIC_MESH_INFO_ARRAY_VIEW(ArrayView, AccessDetector) ArrayView
+
+typedef TArrayView<FMassInstancedStaticMeshInfo> FMassInstancedStaticMeshInfoArrayView;
+
+#endif // ENABLE_MT_DETECTOR
diff --git a/Engine/Plugins/Runtime/MassGameplay/Source/MassRepresentation/Public/MassUpdateISMProcessor.h b/Engine/Plugins/Runtime/MassGameplay/Source/MassRepresentation/Public/MassUpdateISMProcessor.h
new file mode 100644
--- /dev/null
+++ b/Engine/Plugins/Runtime/MassGameplay/Source/MassRepresentation/Public/MassUpdateISMProcessor.h
@@ -0,0 +1,47 @@
+// Copyright Epic Games, Inc. All Rights Reserved.
+
+#pragma once
+
+#include "MassProcessor.h"
+
+#include "MassUpdateISMProcessor.generated.h"
+
+class UMassRepresentationSubsystem;
+
+struct FMassInstancedStaticMeshInfo;
+
+UCLASS()
+class MASSREPRESENTATION_API UMassUpdateISMProcessor : public UPipeProcessor
+{
+	GENERATED_BODY()
+
+public:
+	UMassUpdateISMProcessor();
+
+	static void UpdateISMTransform(int32 EntityId, FMassInstancedStaticMeshInfo& ISMInfo, const FTransform& Transform, const FTransform& PrevTransform, float LODSignificance);
+
+protected:
+	/**
+	 * Initialize the processor
+	 * @param Owner of the Processor */
+	virtual void Initialize(UObject& Owner) override;
+
+	/** Configure the owned FLWComponentQuery instances to express processor's requirements */
+	virtual void ConfigureQueries() override;
+
+	/**
+	 * Execution method for this processor
+	 * @param EntitySubsystem is the system to execute the lambdas on each entity chunk
+	 * @param Context is the execution context to be passed when executing the lambdas */
+	virtual void Execute(UEntitySubsystem& EntitySubsystem, FLWComponentSystemExecutionContext& Context) override;
+
+	/** Component Tag that will be used to filter out entities to update ISM */
+	UPROPERTY(EditAnywhere, Category = "Mass|LOD", config, meta = (BaseStruct = "ComponentTag"))
+	TArray<FInstancedStruct> TagFilters;
+
+	/** A cache pointer to the representation subsystem */
+	UPROPERTY(Transient)
+	UMassRepresentationSubsystem* RepresentationSubsystem;
+
+	TArray<FLWComponentQuery> EntityQueries;
+};
\ No newline at end of file
