AggregateTicking and Pipe merged into Engine/Plugins/Runtime/MassEntity

diff --git a/Engine/Plugins/Runtime/MassGameplay/Source/MassRepresentation/Public/MassVisualizationLODProcessor.h b/Engine/Plugins/Runtime/MassGameplay/Source/MassRepresentation/Public/MassVisualizationLODProcessor.h
new file mode 100644
--- /dev/null
+++ b/Engine/Plugins/Runtime/MassGameplay/Source/MassRepresentation/Public/MassVisualizationLODProcessor.h
@@ -0,0 +1,146 @@
+// Copyright Epic Games, Inc. All Rights Reserved.
+
+#pragma once
+
+#include "MassCommonFragments.h"
+#include "MassRepresentationFragments.h"
+#include "MassLODCalculator.h"
+
+#include "MassVisualizationLODProcessor.generated.h"
+
+UCLASS()
+class MASSREPRESENTATION_API UMassVisualizationLODProcessor : public UMassProcessor_LODBase
+{
+	GENERATED_BODY()
+
+public:
+	UMassVisualizationLODProcessor();
+
+protected:
+
+	/** Configure the owned FLWComponentQuery instances to express processor's requirements */
+	virtual void ConfigureQueries() override;
+
+	/**
+	 * Initialize the processor 
+	 * @param Owner of the Processor
+	 */
+	virtual void Initialize(UObject& Owner) override;
+
+	/** 
+	 * Execution method for this processor 
+	 * @param EntitySubsystem is the system to execute the lambdas on each entity chunk
+	 * @param Context is the execution context to be passed when executing the lambdas
+	 */
+	virtual void Execute(UEntitySubsystem& EntitySubsystem, FLWComponentSystemExecutionContext& Context) override;
+
+	void PrepareExecution();
+
+	template <typename TMassViewerLODInfoFragment = FMassLODInfoFragment>
+	void ExecuteInternal(UEntitySubsystem& EntitySubsystem, FLWComponentSystemExecutionContext& Context);
+
+	
+	/**
+	 * Forces Off LOD on all calculation
+	 * @param bForce to whether force or not Off LOD
+	 */
+	void ForceOffLOD(bool bForce) { bForceOFFLOD = bForce; }
+
+protected:
+	/** Distances where each LOD becomes relevant */
+	UPROPERTY(EditAnywhere, Category = "Mass|LOD", config)
+	float BaseLODDistance[EMassLOD::Max];
+	UPROPERTY(EditAnywhere, Category = "Mass|LOD", config)
+	float VisibleLODDistance[EMassLOD::Max];
+	UPROPERTY(EditAnywhere, Category = "Mass|LOD", meta = (ClampMin = "0.0", UIMin = "0.0"), config)
+	float BufferHysteresisOnDistancePercentage = 10.0f;
+
+	/** Maximum limit for each entity per LOD */
+	UPROPERTY(EditAnywhere, Category = "Mass|LOD", config)
+	int32 LODMaxCount[EMassLOD::Max];
+
+	TMassLODCalculator<FMassRepresentationLODLogic> LODCalculator;
+
+	FLWComponentQuery CloseEntityQuery;
+	FLWComponentQuery FarEntityQuery;
+
+	bool bForceOFFLOD = false;
+};
+
+namespace UE::MassRepresentation
+{
+	extern MASSREPRESENTATION_API int32 bDebugRepresentationLOD;
+} // UE::MassRepresentation
+
+template <typename TMassViewerLODInfoFragment>
+void UMassVisualizationLODProcessor::ExecuteInternal(UEntitySubsystem& EntitySubsystem, FLWComponentSystemExecutionContext& Context)
+{
+	if (bForceOFFLOD)
+	{
+		CloseEntityQuery.ForEachEntityChunk(EntitySubsystem, Context, [this](FLWComponentSystemExecutionContext& Context)
+		{
+			TArrayView<FMassRepresentationLODFragment> RepresentationLODList = Context.GetMutableComponentView<FMassRepresentationLODFragment>();
+			LODCalculator.ForceOffLOD(Context, RepresentationLODList);
+		});
+		return;
+	}
+
+	{
+		TRACE_CPUPROFILER_EVENT_SCOPE(PrepareExecution)
+
+		PrepareExecution();
+	}
+
+	{
+		TRACE_CPUPROFILER_EVENT_SCOPE(CalculateLOD)
+
+		// @todo this is the only block in this processor actually using FMassRepresentationFragment. Consider refactoring it into a separate Processor
+		// the "debug" section below is non-production, should be separated as well
+		auto CalculateLOD = [this](FLWComponentSystemExecutionContext& Context)
+		{
+			TArrayView<FMassRepresentationLODFragment> RepresentationLODList = Context.GetMutableComponentView<FMassRepresentationLODFragment>();
+			TConstArrayView<TMassViewerLODInfoFragment> ViewerInfoList = Context.GetComponentView<TMassViewerLODInfoFragment>();
+			LODCalculator.CalculateLOD(Context, RepresentationLODList, ViewerInfoList);
+		};
+		CloseEntityQuery.ForEachEntityChunk(EntitySubsystem, Context, CalculateLOD);
+
+		// @todo this is the only block in this processor actually using FMassRepresentationFragment. Consider refactoring it into a separate Processor
+		// the "debug" section below is non-production, should be separated as well
+		FarEntityQuery.SetChunkFilter(&FMassVisualizationChunkFragment::ShouldUpdateVisualizationForChunk);
+		FarEntityQuery.ForEachEntityChunk(EntitySubsystem, Context, [this, &CalculateLOD](FLWComponentSystemExecutionContext& Context)
+			{
+				CalculateLOD(Context);
+			});
+		FarEntityQuery.ClearChunkFilter();
+	}
+
+	{
+		TRACE_CPUPROFILER_EVENT_SCOPE(AdjustDistanceAndLODFromCount)
+		if (LODCalculator.AdjustDistancesFromCount())
+		{
+			CloseEntityQuery.ForEachEntityChunk(EntitySubsystem, Context, [this](FLWComponentSystemExecutionContext& Context)
+			{
+				TConstArrayView<TMassViewerLODInfoFragment> ViewerInfoList = Context.GetComponentView<TMassViewerLODInfoFragment>();
+				TArrayView<FMassRepresentationLODFragment> RepresentationLODList = Context.GetMutableComponentView<FMassRepresentationLODFragment>();
+				LODCalculator.AdjustLODFromCount(Context, RepresentationLODList, ViewerInfoList);
+			});
+			// Far entities do not need to maximice count
+		}
+	}
+
+	// Optional debug display
+	if(UE::MassRepresentation::bDebugRepresentationLOD)
+	{
+		TRACE_CPUPROFILER_EVENT_SCOPE(DebugDisplayLOD)
+
+		auto DebugDisplayLOD = [this](FLWComponentSystemExecutionContext& Context)
+		{
+			TConstArrayView<FMassRepresentationLODFragment> RepresentationLODList = Context.GetComponentView<FMassRepresentationLODFragment>();
+			TConstArrayView<FDataFragment_Transform> TransformList = Context.GetComponentView<FDataFragment_Transform>();
+			LODCalculator.DebugDisplayLOD(Context, RepresentationLODList, TransformList, World);
+		};
+
+		CloseEntityQuery.ForEachEntityChunk(EntitySubsystem, Context, DebugDisplayLOD);
+		FarEntityQuery.ForEachEntityChunk(EntitySubsystem, Context, DebugDisplayLOD);
+	}
+}
\ No newline at end of file
diff --git a/Engine/Plugins/Runtime/MassGameplay/Source/MassRepresentation/Public/MassVisualizationTrait.h b/Engine/Plugins/Runtime/MassGameplay/Source/MassRepresentation/Public/MassVisualizationTrait.h
new file mode 100644
--- /dev/null
+++ b/Engine/Plugins/Runtime/MassGameplay/Source/MassRepresentation/Public/MassVisualizationTrait.h
@@ -0,0 +1,42 @@
+﻿// Copyright Epic Games, Inc. All Rights Reserved.
+
+#pragma once 
+
+#include "MassEntityTraitBase.h"
+#include "MassRepresentationTypes.h"
+#include "GameFramework/Actor.h"
+#include "MassVisualizationTrait.generated.h"
+
+class UMassRepresentationSubsystem;
+class UPipeProcessor;
+
+UCLASS(BlueprintType, EditInlineNew, CollapseCategories, meta=(DisplayName="Visualization"))
+class MASSREPRESENTATION_API UMassVisualizationTrait : public UMassEntityTraitBase
+{
+	GENERATED_BODY()
+public:
+	UMassVisualizationTrait();
+
+protected:
+
+	virtual void BuildTemplate(FMassEntityTemplateBuildContext& BuildContext, UWorld& World) const override;
+
+	/** Instanced static mesh information for this agent */
+	UPROPERTY(EditAnywhere, Category = "Mass|Visual")
+	FStaticMeshInstanceVisualizationDesc StaticMeshInstanceDesc;
+
+	/** Actor class of this agent when spawned in high resolution*/
+	UPROPERTY(EditAnywhere, Category = "Mass|Visual")
+	TSubclassOf<AActor> HighResTemplateActor;
+
+	/** Actor class of this agent when spawned in low resolution*/
+	UPROPERTY(EditAnywhere, Category = "Mass|Visual")
+	TSubclassOf<AActor> LowResTemplateActor;
+
+	/** Allow subclasses to override the representation subsystem to use */
+	TSubclassOf<UMassRepresentationSubsystem> RepresentationSubsystemClass;
+
+	/** Allow subclasses to override the representation fragment deinitializer */
+	TSubclassOf<UPipeProcessor> RepresentationFragmentDeinitializerClass;
+
+};
diff --git a/Engine/Plugins/Runtime/MassGameplay/Source/MassRepresentation/Public/MassVisualizer.h b/Engine/Plugins/Runtime/MassGameplay/Source/MassRepresentation/Public/MassVisualizer.h
new file mode 100644
--- /dev/null
+++ b/Engine/Plugins/Runtime/MassGameplay/Source/MassRepresentation/Public/MassVisualizer.h
@@ -0,0 +1,29 @@
+// Copyright Epic Games, Inc. All Rights Reserved.
+
+#pragma once
+
+#include "GameFramework/Actor.h"
+
+#include "MassVisualizer.generated.h"
+
+
+class UMassVisualizationComponent;
+
+/**
+ * Actor holding the mass visual component responsible to handle the representation of the mass agent as the static mesh instances 
+ * There may be a separate instance of these for different types of Agents (Cars, NPC's etc)
+ */
+UCLASS(NotPlaceable, Transient)
+class MASSREPRESENTATION_API AMassVisualizer : public AActor
+{
+	GENERATED_BODY()
+public:
+	AMassVisualizer();
+
+	/** Visualization component is garantee to exist if this class is created */
+	class UMassVisualizationComponent& GetVisualizationComponent() const { return *VisComponent; }
+
+protected:
+	UPROPERTY()
+	class UMassVisualizationComponent* VisComponent;
+};
diff --git a/Engine/Plugins/Runtime/MassGameplay/Source/MassSignals/Private/MassSignalProcessorBase.cpp b/Engine/Plugins/Runtime/MassGameplay/Source/MassSignals/Private/MassSignalProcessorBase.cpp
new file mode 100644
--- /dev/null
+++ b/Engine/Plugins/Runtime/MassGameplay/Source/MassSignals/Private/MassSignalProcessorBase.cpp
@@ -0,0 +1,148 @@
+﻿// Copyright Epic Games, Inc. All Rights Reserved.
+
+#include "MassSignalProcessorBase.h"
+#include "MassSignalSubsystem.h"
+#include "ArchetypeTypes.h"
+#include "Engine/World.h"
+
+void UMassSignalProcessorBase::Initialize(UObject& Owner)
+{
+	SignalSubsystem = UWorld::GetSubsystem<UMassSignalSubsystem>(Owner.GetWorld());
+}
+
+void UMassSignalProcessorBase::BeginDestroy()
+{
+	if (SignalSubsystem)
+	{
+		for (const FName& SignalName : RegisteredSignals)
+		{
+			SignalSubsystem->GetSignalDelegateByName(SignalName).RemoveAll(this);
+		}
+	}
+
+	Super::BeginDestroy();
+}
+
+void UMassSignalProcessorBase::SubscribeToSignal(const FName SignalName)
+{
+	check(SignalSubsystem);
+	check(!RegisteredSignals.Contains(SignalName));
+	RegisteredSignals.Add(SignalName);
+	SignalSubsystem->GetSignalDelegateByName(SignalName).AddUObject(this, &UMassSignalProcessorBase::OnSignalReceived);
+}
+
+void UMassSignalProcessorBase::Execute(UEntitySubsystem& EntitySubsystem, FLWComponentSystemExecutionContext& Context)
+{
+	QUICK_SCOPE_CYCLE_COUNTER(SignalEntities);
+
+	// Frame buffer handling
+	const int32 ProcessingFrameBufferIndex = CurrentFrameBufferIndex;
+	CurrentFrameBufferIndex = (CurrentFrameBufferIndex + 1) % 2;
+	FFrameReceivedSignals& ProcessingFrameBuffer = FrameReceivedSignals[ProcessingFrameBufferIndex];
+	TArray<FEntitySignalRange>& ReceivedSignalRanges = ProcessingFrameBuffer.ReceivedSignalRanges;
+	TArray<FLWEntity>& SignaledEntities = ProcessingFrameBuffer.SignaledEntities;
+
+	if (ReceivedSignalRanges.IsEmpty())
+	{
+		return;
+	}
+
+	EntityQuery.CacheArchetypes(EntitySubsystem);
+	if (EntityQuery.GetArchetypes().Num() > 0)
+	{
+		// EntitySet stores unique array of entities per specified archetype.
+		// FArchetypeChunkCollection expects an array of entities, a set is used to detect unique ones.
+		struct FEntitySet
+		{
+			void Reset()
+			{
+				Entities.Reset();
+				AddedEntities.Reset();
+			}
+
+			FArchetypeHandle Archetype;
+			TArray<FLWEntity> Entities;
+			TSet<FLWEntity> AddedEntities;
+		};
+		TArray<FEntitySet> EntitySets;
+
+		for (const FArchetypeHandle& Archetype : EntityQuery.GetArchetypes())
+		{
+			FEntitySet& Set = EntitySets.AddDefaulted_GetRef();
+			Set.Archetype = Archetype;
+		}
+
+		// SignalNameLookup has limit of how many signals it can handle at once, we'll do passes until all signals are processed.
+		int32 SignalsToProcess = ReceivedSignalRanges.Num();
+		while(SignalsToProcess > 0)
+		{
+			SignalNameLookup.Reset();
+
+			// Convert signals with entity ids into arrays of entities per archetype.
+			for (FEntitySignalRange& Range : ReceivedSignalRanges)
+			{
+				if (Range.bProcessed)
+				{
+					continue;
+				}
+				// Get bitflag for the signal name
+				const uint64 SignalFlag = SignalNameLookup.GetOrAddSignalName(Range.SignalName);
+				if (SignalFlag == 0)
+				{
+					// Will process that signal in a second iteration
+					continue;
+				}
+
+				// Entities for this signal
+				TArrayView<FLWEntity> Entities(&SignaledEntities[Range.Begin], Range.End - Range.Begin);
+				FEntitySet* PrevSet = &EntitySets[0];
+				for (const FLWEntity Entity : Entities)
+				{
+					// Add to set of supported archetypes. Dont process if we don't care about the type.
+					FArchetypeHandle Archetype = EntitySubsystem.GetArchetypeForEntity(Entity);
+					FEntitySet* Set = PrevSet->Archetype == Archetype ? PrevSet : EntitySets.FindByPredicate([&Archetype](const FEntitySet& Set) { return Archetype == Set.Archetype; });
+					if (Set != nullptr)
+					{
+						// Add only unique entities in the array.
+						if (!Set->AddedEntities.Contains(Entity))
+						{
+							Set->Entities.Add(Entity);
+							Set->AddedEntities.Add(Entity);
+						}
+						SignalNameLookup.AddSignalToEntity(Entity, SignalFlag);
+						PrevSet = Set;
+					}
+				}
+
+				Range.bProcessed = true;
+				SignalsToProcess--;
+			}
+
+			// Execute per archetype.
+			for (FEntitySet& Set : EntitySets)
+			{
+				if (Set.Entities.Num() > 0)
+				{
+					Context.SetChunkCollection(FArchetypeChunkCollection(Set.Archetype, Set.Entities));
+					SignalEntities(EntitySubsystem, Context, SignalNameLookup);
+					Context.ClearChunkCollection();
+				}
+				Set.Reset();
+			}
+		}
+	}
+
+	ReceivedSignalRanges.Reset();
+	SignaledEntities.Reset();
+}
+
+void UMassSignalProcessorBase::OnSignalReceived(FName SignalName, TConstArrayView<FLWEntity> Entities)
+{
+	FFrameReceivedSignals& CurrentFrameBuffer = FrameReceivedSignals[CurrentFrameBufferIndex];
+
+	FEntitySignalRange& Range = CurrentFrameBuffer.ReceivedSignalRanges.AddDefaulted_GetRef();
+	Range.SignalName = SignalName;
+	Range.Begin = CurrentFrameBuffer.SignaledEntities.Num();
+	CurrentFrameBuffer.SignaledEntities.Append(Entities.GetData(), Entities.Num());
+	Range.End = CurrentFrameBuffer.SignaledEntities.Num();
+}
diff --git a/Engine/Plugins/Runtime/MassGameplay/Source/MassSignals/Private/MassSignalSubsystem.cpp b/Engine/Plugins/Runtime/MassGameplay/Source/MassSignals/Private/MassSignalSubsystem.cpp
new file mode 100644
--- /dev/null
+++ b/Engine/Plugins/Runtime/MassGameplay/Source/MassSignals/Private/MassSignalSubsystem.cpp
@@ -0,0 +1,107 @@
+// Copyright Epic Games, Inc. All Rights Reserved.
+
+#include "MassSignalSubsystem.h"
+#include "LWCCommandBuffer.h"
+#include "MassSignalTypes.h"
+#include "VisualLogger/VisualLogger.h"
+#include "ProfilingDebugging/CsvProfiler.h"
+#include "Engine/World.h"
+
+CSV_DEFINE_CATEGORY(MassSignalsCounters, true);
+
+void UMassSignalSubsystem::Tick(float DeltaTime)
+{
+	for (int i = 0; i < DelayedSignals.Num();)
+	{
+		FDelayedSignal& DelayedSignal = DelayedSignals[i];
+		DelayedSignal.DelayInSeconds -= DeltaTime;
+		if(DelayedSignal.DelayInSeconds <= 0.0f)
+		{
+			SignalEntities(DelayedSignal.SignalName, MakeArrayView(DelayedSignal.Entities));
+			DelayedSignals.RemoveAtSwap(i);
+		}
+		else
+		{
+			i++;
+		}
+	}
+}
+
+TStatId UMassSignalSubsystem::GetStatId() const
+{
+	RETURN_QUICK_DECLARE_CYCLE_STAT(UMassSignalSubsystem, STATGROUP_Tickables);
+}
+
+void UMassSignalSubsystem::SignalEntity(FName SignalName, const FLWEntity Entity)
+{
+	checkf(Entity.IsSet(), TEXT("Expecting a valid entity to signal"));
+	SignalEntities(SignalName, MakeArrayView(&Entity, 1));
+}
+
+void UMassSignalSubsystem::SignalEntities(FName SignalName, TConstArrayView<FLWEntity> Entities)
+{
+	checkf(Entities.Num() > 0, TEXT("Expecting entities to signal"));
+	const UE::MassSignal::FSignalDelegate& SignalDelegate = GetSignalDelegateByName(SignalName);
+	SignalDelegate.Broadcast(SignalName, Entities);
+
+#if CSV_PROFILER
+	FCsvProfiler::RecordCustomStat(*SignalName.ToString(), CSV_CATEGORY_INDEX(MassSignalsCounters), Entities.Num(), ECsvCustomStatOp::Accumulate);
+#endif
+
+	UE_CVLOG(Entities.Num() == 1, this, LogMassSignals, Log, TEXT("Raising signal [%s] to entity [%s]"), *SignalName.ToString(), *Entities[0].DebugGetDescription());
+	UE_CVLOG(Entities.Num() > 1, this, LogMassSignals, Log, TEXT("Raising signal [%s] to %d entities"), *SignalName.ToString(), Entities.Num());
+}
+
+void UMassSignalSubsystem::DelaySignalEntity(FName SignalName, const FLWEntity Entity, const float DelayInSeconds)
+{
+	checkf(Entity.IsSet(), TEXT("Expecting a valid entity to signal"));
+	DelaySignalEntities(SignalName, MakeArrayView(&Entity, 1), DelayInSeconds);
+}
+
+void UMassSignalSubsystem::DelaySignalEntities(FName SignalName, TConstArrayView<FLWEntity> Entities, const float DelayInSeconds)
+{
+	FDelayedSignal& DelayedSignal = DelayedSignals.Emplace_GetRef();
+	DelayedSignal.SignalName = SignalName;
+	DelayedSignal.Entities = Entities;
+	DelayedSignal.DelayInSeconds = DelayInSeconds;
+	UE_CVLOG(Entities.Num() == 1, this, LogMassSignals, Log, TEXT("Delay signal [%s] to entity [%s] in %.2f"), *SignalName.ToString(), *Entities[0].DebugGetDescription(), DelayInSeconds);
+	UE_CVLOG(Entities.Num() > 1,this, LogMassSignals, Log, TEXT("Delay signal [%s] to %d entities in %.2f"), *SignalName.ToString(), Entities.Num(), DelayInSeconds);
+}
+
+void UMassSignalSubsystem::SignalEntityDeferred(FLWComponentSystemExecutionContext& Context, FName SignalName, const FLWEntity Entity)
+{
+	checkf(Entity.IsSet(), TEXT("Expecting a valid entity to signal"));
+	SignalEntitiesDeferred(Context, SignalName, MakeArrayView(&Entity, 1));
+}
+
+void UMassSignalSubsystem::SignalEntitiesDeferred(FLWComponentSystemExecutionContext& Context, FName SignalName, TConstArrayView<FLWEntity> Entities)
+{
+	checkf(Entities.Num() > 0, TEXT("Expecting entities to signal"));
+	Context.Defer().EmplaceCommand<FDeferredCommand>([SignalName, Entities](const UEntitySubsystem& System)
+	{
+		UMassSignalSubsystem* SignalSubsystem = UWorld::GetSubsystem<UMassSignalSubsystem>(System.GetWorld());
+		SignalSubsystem->SignalEntities(SignalName, Entities);
+	});
+
+	UE_CVLOG(Entities.Num() == 1, this, LogMassSignals, Log, TEXT("Raising deferred signal [%s] to entity [%s]"), *SignalName.ToString(), *Entities[0].DebugGetDescription());
+	UE_CVLOG(Entities.Num() > 1, this, LogMassSignals, Log, TEXT("Raising deferred signal [%s] to %d entities"), *SignalName.ToString(), Entities.Num());
+}
+
+void UMassSignalSubsystem::DelaySignalEntityDeferred(FLWComponentSystemExecutionContext& Context, FName SignalName, const FLWEntity Entity, const float DelayInSeconds)
+{
+	checkf(Entity.IsSet(), TEXT("Expecting a valid entity to signal"));
+	DelaySignalEntitiesDeferred(Context, SignalName, MakeArrayView(&Entity, 1), DelayInSeconds);
+}
+
+void UMassSignalSubsystem::DelaySignalEntitiesDeferred(FLWComponentSystemExecutionContext& Context, FName SignalName, TConstArrayView<FLWEntity> Entities, const float DelayInSeconds)
+{
+	checkf(Entities.Num() > 0, TEXT("Expecting entities to signal"));
+	Context.Defer().EmplaceCommand<FDeferredCommand>([SignalName, Entities, DelayInSeconds](const UEntitySubsystem& System)
+	{
+		UMassSignalSubsystem* SignalSubsystem = UWorld::GetSubsystem<UMassSignalSubsystem>(System.GetWorld());
+		SignalSubsystem->DelaySignalEntities(SignalName, Entities, DelayInSeconds);
+	});
+
+	UE_CVLOG(Entities.Num() == 1, this, LogMassSignals, Log, TEXT("Delay deferred signal [%s] to entity [%s] in %.2f"), *SignalName.ToString(), *Entities[0].DebugGetDescription(), DelayInSeconds);
+	UE_CVLOG(Entities.Num() > 1,this, LogMassSignals, Log, TEXT("Delay deferred signal [%s] to %d entities in %.2f"), *SignalName.ToString(), Entities.Num(), DelayInSeconds);
+}
diff --git a/Engine/Plugins/Runtime/MassGameplay/Source/MassSignals/Private/MassSignalTypes.cpp b/Engine/Plugins/Runtime/MassGameplay/Source/MassSignals/Private/MassSignalTypes.cpp
new file mode 100644
--- /dev/null
+++ b/Engine/Plugins/Runtime/MassGameplay/Source/MassSignals/Private/MassSignalTypes.cpp
@@ -0,0 +1,48 @@
+// Copyright Epic Games, Inc. All Rights Reserved.
+
+#include "MassSignalTypes.h"
+#include "EntitySubsystem.h"
+
+uint64 FMassSignalNameLookup::GetOrAddSignalName(const FName SignalName)
+{
+	int32 Index = SignalNames.Find(SignalName);
+	if (Index == INDEX_NONE)
+	{
+		if (SignalNames.Num() >= MaxSignalNames)
+		{
+			return 0;
+		}
+		Index = SignalNames.Add(SignalName);
+	}
+	return 1ULL << Index;
+}
+
+void FMassSignalNameLookup::AddSignalToEntity(const FLWEntity Entity, const uint64 SignalFlag)
+{
+	uint64& Signals = EntitySignals.FindOrAdd(Entity, 0);
+	Signals |= SignalFlag;
+}
+
+void FMassSignalNameLookup::GetSignalsForEntity(const FLWEntity Entity, TArray<FName>& OutSignals) const
+{
+	OutSignals.Reset();
+	if (const uint64* Signals = EntitySignals.Find(Entity))
+	{
+		int Start = FMath::CountTrailingZeros64(*Signals);
+		int End = FMath::CountLeadingZeros64(*Signals);
+		for (int64 i = Start; i < 64 - End; i++)
+		{
+			const uint64 SignalFlag = 1ULL << i;
+			if ((*Signals) & SignalFlag)
+			{
+				OutSignals.Add(SignalNames[i]);
+			}
+		}
+	}
+}
+
+void FMassSignalNameLookup::Reset()
+{
+	SignalNames.Reset();
+	EntitySignals.Reset();
+}
\ No newline at end of file
diff --git a/Engine/Plugins/Runtime/MassGameplay/Source/MassSignals/Public/MassSignalProcessorBase.h b/Engine/Plugins/Runtime/MassGameplay/Source/MassSignals/Public/MassSignalProcessorBase.h
new file mode 100644
--- /dev/null
+++ b/Engine/Plugins/Runtime/MassGameplay/Source/MassSignals/Public/MassSignalProcessorBase.h
@@ -0,0 +1,100 @@
+﻿// Copyright Epic Games, Inc. All Rights Reserved.
+
+#pragma once
+
+#include "MassProcessor.h"
+#include "EntitySubsystem.h"
+#include "MassSignalTypes.h"
+#include "Containers/StaticArray.h"
+#include "MassSignalProcessorBase.generated.h"
+
+class UMassSignalSubsystem;
+
+/** 
+ * Processor for executing signals on each targeted entities
+ * The derived classes only need to implement the method SignalEntities to actually received the raised signals for the entities they subscribed to 
+ */
+UCLASS(abstract)
+class MASSSIGNALS_API UMassSignalProcessorBase : public UPipeProcessor
+{
+	GENERATED_BODY()
+
+protected:
+	/**
+	 * Initialize the processor
+	 * @param Owner of the Processor
+	 */
+	virtual void Initialize(UObject& Owner) override;
+
+	virtual void BeginDestroy() override;
+
+	/** Configure the owned FLWComponentQuery instances to express processor queries requirements */
+	virtual void ConfigureQueries() override {}
+
+	/**
+	 * Actual method that derived class needs to implement to act on a signal that is raised for that frame
+	 * @param EntitySubsystem is the system to execute the lambdas on each entity chunk
+	 * @param Context is the execution context to be passed when executing the lambdas
+	 * @param EntitySignals Look up to retrieve for each entities their raised signal via GetSignalsForEntity
+	 */
+	 virtual void SignalEntities(UEntitySubsystem& EntitySubsystem, FLWComponentSystemExecutionContext& Context, FMassSignalNameLookup& EntitySignals) PURE_VIRTUAL(UMassSignalProcessorBase::SignalEntities, );
+
+	/**
+	 * Callback that is being called when new signal is raised
+	 * @param SignalName is the name of the signal being raised
+	 * @param Entities are the targeted entities for this signal
+	 */
+	 virtual void OnSignalReceived(FName SignalName, TConstArrayView<FLWEntity> Entities);
+
+	/**
+	 * Execution method for this processor
+	 * @param EntitySubsystem is the system to execute the lambdas on each entity chunk
+	 * @param Context is the execution context to be passed when executing the lambdas
+	 */
+	virtual void Execute(UEntitySubsystem& EntitySubsystem, FLWComponentSystemExecutionContext& Context) override;
+
+	/**
+	 * To receive notification about a particular signal, you need to subscribe to it.
+	 * @param SignalName is the name of the signal to receive notification about
+	 */
+	void SubscribeToSignal(const FName SignalName);
+
+	
+	UPROPERTY(Transient)
+	UMassSignalSubsystem* SignalSubsystem = nullptr;
+
+	FLWComponentQuery EntityQuery;
+
+private:
+
+	/** Stores a range of indices in the SignaledEntities TArray of Entities and the associated signal name */
+	struct FEntitySignalRange
+	{
+		FName SignalName;
+		int32 Begin = 0;
+		int32 End = 0;
+		bool bProcessed = false;
+	};
+
+	struct FFrameReceivedSignals
+	{
+		/** Received signals are double buffered as we can receive new one while processing them */
+		TArray<FEntitySignalRange> ReceivedSignalRanges;
+
+		/** the list of all signaled entities, can contain duplicates */
+		TArray<FLWEntity> SignaledEntities;
+	};
+
+	/** Double buffer frame received signal as we can receive new signals as we are processing them */
+	TStaticArray<FFrameReceivedSignals, 2> FrameReceivedSignals;
+	
+	/** Current frame buffer index of FrameReceivedSignals */
+	int32 CurrentFrameBufferIndex = 0;
+
+	/** Lookup used to store and retrieve signals per entity, only used during processing */
+	FMassSignalNameLookup SignalNameLookup;
+
+	/** List of all the registered signal names*/
+	TArray<FName> RegisteredSignals;
+};
+
diff --git a/Engine/Plugins/Runtime/MassGameplay/Source/MassSignals/Public/MassSignalSubsystem.h b/Engine/Plugins/Runtime/MassGameplay/Source/MassSignals/Public/MassSignalSubsystem.h
new file mode 100644
--- /dev/null
+++ b/Engine/Plugins/Runtime/MassGameplay/Source/MassSignals/Public/MassSignalSubsystem.h
@@ -0,0 +1,111 @@
+// Copyright Epic Games, Inc. All Rights Reserved.
+
+#pragma once
+
+#include "EntitySubsystem.h"
+#include "Subsystems/WorldSubsystem.h"
+#include "MassSignalSubsystem.generated.h"
+
+namespace UE::MassSignal 
+{
+	DECLARE_MULTICAST_DELEGATE_TwoParams(FSignalDelegate, FName /*SignalName*/, TConstArrayView<FLWEntity> /*Entities*/);
+} // UE::MassSignal
+
+/**
+* A subsystem for handling Signals in Mass
+*/
+UCLASS()
+class MASSSIGNALS_API UMassSignalSubsystem : public UTickableWorldSubsystem
+{
+	GENERATED_BODY()
+	
+public:
+
+	/** 
+	 * Retrieve the delegate dispatcher from the signal name
+	 * @param SignalName is the name of the signal to get the delegate dispatcher from
+	 */
+	UE::MassSignal::FSignalDelegate& GetSignalDelegateByName(FName SignalName)
+	{
+		return NamedSignals.FindOrAdd(SignalName);
+	}
+
+	/**
+	 * Inform a single entity of a signal being raised
+	 * @param SignalName is the name of the signal raised
+	 * @param Entity entity that should be informed that signal 'SignalName' was raised
+	 */
+	void SignalEntity(FName SignalName, const FLWEntity Entity);
+
+	/**
+	 * Inform multiple entities of a signal being raised
+	 * @param SignalName is the name of the signal raised
+	 * @param Entities list of entities that should be informed that signal 'SignalName' was raised
+	 */
+	void SignalEntities(FName SignalName, TConstArrayView<FLWEntity> Entities);
+
+	/**
+	 * Inform a single entity of a signal being raised in a certain amount of seconds
+	 * @param SignalName is the name of the signal raised
+	 * @param Entity entity that should be informed that signal 'SignalName' was raised
+	 * @param DelayInSeconds is the amount of time before signaling the entity
+	 */
+	void DelaySignalEntity(FName SignalName, const FLWEntity Entity, const float DelayInSeconds);
+
+ 	/**
+	 * Inform multiple entities of a signal being raised in a certain amount of seconds
+	 * @param SignalName is the name of the signal raised
+	 * @param Entities being informed of the raised signal
+	 * @param DelayInSeconds is the amount of time before signaling the entities
+	 */
+	void DelaySignalEntities(FName SignalName, TConstArrayView<FLWEntity> Entities, const float DelayInSeconds);
+
+	/**
+	 * Inform single entity of a signal being raised asynchronously using the LWComponent Command Buffer
+	 * @param Context is the Entity System execution context to push the async command
+	 * @param SignalName is the name of the signal raised
+	 * @param Entity entity that should be informed that signal 'SignalName' was raised
+	 */
+	void SignalEntityDeferred(FLWComponentSystemExecutionContext& Context, FName SignalName, const FLWEntity Entity);
+
+	/**
+	 * Inform multiple entities of a signal being raised asynchronously using the LWComponent Command Buffer
+	 * @param Context is the Entity System execution context to push the async command
+	 * @param SignalName is the name of the signal raised
+	 * @param Entities list of entities that should be informed that signal 'SignalName' was raised
+	 */
+	void SignalEntitiesDeferred(FLWComponentSystemExecutionContext& Context, FName SignalName, TConstArrayView<FLWEntity> Entities);
+
+ 	/**
+	 * Inform single entity of a signal being raised asynchronously using the LWComponent Command Buffer
+	 * @param Context is the Entity System execution context to push the async command
+	 * @param SignalName is the name of the signal raised
+	 * @param Entity entity that should be informed that signal 'SignalName' was raised
+	 * @param DelayInSeconds is the amount of time before signaling the entities
+	 */
+	void DelaySignalEntityDeferred(FLWComponentSystemExecutionContext& Context, FName SignalName, const FLWEntity Entity, const float DelayInSeconds);
+
+ 	/**
+	 * Inform multiple entities of a signal being raised asynchronously using the LWComponent Command Buffer
+	 * @param Context is the Entity System execution context to push the async command
+	 * @param SignalName is the name of the signal raised
+	 * @param Entities being informed of that signal was raised
+	 * @param DelayInSeconds is the amount of time before signaling the entities
+	 */
+	void DelaySignalEntitiesDeferred(FLWComponentSystemExecutionContext& Context, FName SignalName, TConstArrayView<FLWEntity> Entities, const float DelayInSeconds);
+
+protected:
+	virtual void Tick(float DeltaTime) override;
+	virtual TStatId GetStatId() const override;
+
+	TMap<FName, UE::MassSignal::FSignalDelegate> NamedSignals;
+
+	struct FDelayedSignal
+	{
+		FName SignalName;
+		TArray<FLWEntity> Entities;
+		float DelayInSeconds;
+	};
+
+	TArray<FDelayedSignal> DelayedSignals;
+};
diff --git a/Engine/Plugins/Runtime/MassGameplay/Source/MassSimulation/Private/MassSimulationLocalCoordinator.cpp b/Engine/Plugins/Runtime/MassGameplay/Source/MassSimulation/Private/MassSimulationLocalCoordinator.cpp
new file mode 100644
--- /dev/null
+++ b/Engine/Plugins/Runtime/MassGameplay/Source/MassSimulation/Private/MassSimulationLocalCoordinator.cpp
@@ -0,0 +1,67 @@
+// Copyright Epic Games, Inc. All Rights Reserved.
+
+#include "MassSimulationLocalCoordinator.h"
+#include "NavigationSystem.h"
+#include "MassProcessor.h"
+#include "MassSimulationSubsystem.h"
+
+#if WITH_EDITORONLY_DATA
+#include "UObject/ConstructorHelpers.h"
+#include "Components/BillboardComponent.h"
+#endif // WITH_EDITORONLY_DATA
+
+#if WITH_EDITOR
+#endif // WITH_EDITOR
+
+AMassSimulationLocalCoordinator::AMassSimulationLocalCoordinator(const FObjectInitializer& ObjectInitializer)
+	: Super(ObjectInitializer)
+{
+	USceneComponent* SceneComponent = CreateDefaultSubobject<USceneComponent>(TEXT("SceneComp"));
+	RootComponent = SceneComponent;
+	RootComponent->Mobility = EComponentMobility::Static;
+	
+#if WITH_EDITORONLY_DATA
+	SpriteComponent = CreateEditorOnlyDefaultSubobject<UBillboardComponent>(TEXT("Sprite"));
+	if (!IsRunningCommandlet())
+	{
+		// Structure to hold one-time initialization
+		struct FConstructorStatics
+		{
+			ConstructorHelpers::FObjectFinderOptional<UTexture2D> NoteTextureObject;
+			FName NotesID;
+			FText NotesName;
+			FConstructorStatics()
+				: NoteTextureObject(TEXT("/Engine/EditorResources/S_Note"))
+				, NotesID(TEXT("Notes"))
+				, NotesName(NSLOCTEXT("SpriteCategory", "Notes", "Notes"))
+			{
+			}
+		};
+		static FConstructorStatics ConstructorStatics;
+		if (SpriteComponent)
+		{
+			SpriteComponent->Sprite = ConstructorStatics.NoteTextureObject.Get();
+			SpriteComponent->SetRelativeScale3D(FVector(0.5f, 0.5f, 0.5f));
+			SpriteComponent->SpriteInfo.Category = ConstructorStatics.NotesID;
+			SpriteComponent->SpriteInfo.DisplayName = ConstructorStatics.NotesName;
+			SpriteComponent->SetupAttachment(RootComponent);
+			SpriteComponent->Mobility = EComponentMobility::Static;
+		}
+	}
+#endif // WITH_EDITORONLY_DATA
+	SetCanBeDamaged(false);
+
+	PrimaryActorTick.bCanEverTick = true;
+	PrimaryActorTick.bStartWithTickEnabled = true;
+
+#if WITH_EDITORONLY_DATA
+	bTickInEditor = false;
+	bDrawDebug = false;
+#endif // WITH_EDITORONLY_DATA
+}
+
+void AMassSimulationLocalCoordinator::PostLoad()
+{
+	Super::PostLoad();
+	UE_LOG(LogMass, Warning, TEXT("%s has been loaded but will be discarded during next level saving. This class is not longer being used"), *GetName());
+}
diff --git a/Engine/Plugins/Runtime/MassGameplay/Source/MassSimulation/Private/MassSimulationSettings.cpp b/Engine/Plugins/Runtime/MassGameplay/Source/MassSimulation/Private/MassSimulationSettings.cpp
new file mode 100644
--- /dev/null
+++ b/Engine/Plugins/Runtime/MassGameplay/Source/MassSimulation/Private/MassSimulationSettings.cpp
@@ -0,0 +1,43 @@
+// Copyright Epic Games, Inc. All Rights Reserved.
+
+#include "MassSimulationSettings.h"
+#include "MassSchematic.h"
+
+#if WITH_EDITOR
+#include "MassEntityEditorModule.h"
+
+void UMassSimulationSettings::PostInitProperties()
+{
+	Super::PostInitProperties();
+
+	FPipeEditorModule* PipeEditorModule = FModuleManager::GetModulePtr<FPipeEditorModule>("MassEntityEditor");
+	if (PipeEditorModule)
+	{
+		PipeEditorModule->GetOnAssetPropertiesChanged().AddUObject(this, &UMassSimulationSettings::OnAssetPropertiesChanged);
+	}
+}
+
+void UMassSimulationSettings::PostEditChangeProperty(struct FPropertyChangedEvent& PropertyChangedEvent)
+{
+	static const FName TickSchematicsName = GET_MEMBER_NAME_CHECKED(UMassSimulationSettings, TickSchematics);
+
+	Super::PostEditChangeProperty(PropertyChangedEvent);
+
+	if (PropertyChangedEvent.Property)
+	{
+		const FName PropName = PropertyChangedEvent.Property->GetFName();
+		if (PropName == TickSchematicsName)
+		{
+			OnTickSchematicChanged.Broadcast();
+		}
+	}
+}
+
+void UMassSimulationSettings::OnAssetPropertiesChanged(UPipeSchematic* PipeSchematic, const FPropertyChangedEvent& PropertyChangedEvent)
+{
+	if (TickSchematics.Find(PipeSchematic) != INDEX_NONE)
+	{
+		OnTickSchematicChanged.Broadcast();
+	}
+}
+#endif // WITH_EDITOR
diff --git a/Engine/Plugins/Runtime/MassGameplay/Source/MassSimulation/Private/MassSimulationSubsystem.cpp b/Engine/Plugins/Runtime/MassGameplay/Source/MassSimulation/Private/MassSimulationSubsystem.cpp
new file mode 100644
--- /dev/null
+++ b/Engine/Plugins/Runtime/MassGameplay/Source/MassSimulation/Private/MassSimulationSubsystem.cpp
@@ -0,0 +1,208 @@
+// Copyright Epic Games, Inc. All Rights Reserved.
+
+#include "MassSimulationSubsystem.h"
+#include "MassSimulationLocalCoordinator.h"
+#include "MassEntitySystem.h"
+#include "MassExecutor.h"
+#include "MassSimulationSettings.h"
+#include "VisualLogger/VisualLogger.h"
+#if WITH_EDITOR
+#include "Editor.h"
+#endif // WITH_EDITOR
+
+DEFINE_LOG_CATEGORY(LogMassSim);
+
+namespace FMassSimulationSubsystemHelper
+{
+	UPipeEntitySubsystem* GetEntitySubsystem(const UWorld* World)
+	{
+		UMassSimulationSubsystem* Instance = UWorld::GetSubsystem<UMassSimulationSubsystem>(World);
+		return Instance ? Instance->GetEntitySubsystem() : nullptr;
+	}
+}
+
+namespace UE::MassSimulation
+{
+	int32 bDoEntityCompaction = 1;
+	FAutoConsoleVariableRef CVarEntityCompaction(TEXT("ai.mass.EntityCompaction"), bDoEntityCompaction, TEXT("Maximize the nubmer of entities per chunk"), ECVF_Cheat);
+}
+
+//----------------------------------------------------------------------//
+// UMassSimulationSubsystem
+//----------------------------------------------------------------------//
+UMassSimulationSubsystem::UMassSimulationSubsystem(const FObjectInitializer& ObjectInitializer)
+	: Super()
+{
+	PhaseManager = CreateDefaultSubobject<UPipeProcessingPhaseManager>(TEXT("PhaseManager"));
+}
+
+void UMassSimulationSubsystem::PostInitProperties()
+{
+	REDIRECT_OBJECT_TO_VLOG(PhaseManager, this);
+
+	Super::PostInitProperties();
+
+	if (HasAnyFlags(RF_ClassDefaultObject))
+	{
+		UPipeEntitySubsystem::SetGetter_Internal([](const UWorld* World)
+		{
+			return FMassSimulationSubsystemHelper::GetEntitySubsystem(World);
+		});
+	}
+	else
+	{
+		CreateEntitySubsystem();
+#if WITH_EDITOR
+		GET_MASS_CONFIG_VALUE(GetOnTickSchematicChanged()).AddUObject(this, &UMassSimulationSubsystem::RebuildTickPipeline);
+#endif // WITH_EDITOR
+	}
+}
+
+void UMassSimulationSubsystem::BeginDestroy()
+{
+#if WITH_EDITOR
+	if (PieBeginEventHandle.IsValid())
+	{
+		FEditorDelegates::BeginPIE.Remove(PieBeginEventHandle);
+	}
+	if (PieEndedEventHandle.IsValid())
+	{
+		FEditorDelegates::PrePIEEnded.Remove(PieEndedEventHandle);
+	}
+#endif //  WITH_EDITOR
+	Super::BeginDestroy();
+}
+
+FPipeProcessingPhase::FOnPhaseEvent& UMassSimulationSubsystem::GetOnProcessingPhaseStarted(const EPipeProcessingPhase Phase) const
+{
+	check(PhaseManager);
+	return PhaseManager->GetOnPhaseStart(Phase);
+}
+
+FPipeProcessingPhase::FOnPhaseEvent& UMassSimulationSubsystem::GetOnProcessingPhaseFinished(const EPipeProcessingPhase Phase) const
+{
+	check(PhaseManager);
+	return PhaseManager->GetOnPhaseEnd(Phase);
+}
+
+void UMassSimulationSubsystem::CreateEntitySubsystem()
+{
+	check(HasAnyFlags(RF_ClassDefaultObject) == false);
+	if (ensureMsgf(EntitySubsystem == nullptr, TEXT("Trying to create a second instance of EntitySubsystem. The call will be ignored.")))
+	{
+		UWorld* World = GetWorld();
+		check(World);
+		EntitySubsystem = NewObject<UPipeEntitySubsystem>(World);
+	}
+}
+
+void UMassSimulationSubsystem::Initialize(FSubsystemCollectionBase& Collection)
+{
+	Super::Initialize(Collection);
+
+	check(EntitySubsystem);
+	EntitySubsystem->Initialize(Collection);
+	GetOnProcessingPhaseStarted(EPipeProcessingPhase::PrePhysics).AddUObject(this, &UMassSimulationSubsystem::OnProcessingPhaseStarted, EPipeProcessingPhase::PrePhysics);
+}
+
+void UMassSimulationSubsystem::Deinitialize()
+{
+#if WITH_EDITOR
+	GET_MASS_CONFIG_VALUE(GetOnTickSchematicChanged()).RemoveAll(this);
+#endif // WITH_EDITOR
+	GetOnProcessingPhaseStarted(EPipeProcessingPhase::PrePhysics).RemoveAll(this);
+	StopSimulation();
+
+	Super::Deinitialize();
+}
+
+void UMassSimulationSubsystem::PostInitialize()
+{
+	Super::PostInitialize();
+
+#if WITH_EDITOR
+	UWorld* World = GetWorld();
+	if (GEditor && World!= nullptr && !World->IsGameWorld())
+	{
+		// in editor worlds we need to rebuild the pipeline at this point since OnWorldBeginPlay won't be called
+		RebuildTickPipeline();
+
+		PieBeginEventHandle = FEditorDelegates::BeginPIE.AddUObject(this, &UMassSimulationSubsystem::OnPieBegin);
+		PieEndedEventHandle = FEditorDelegates::PrePIEEnded.AddUObject(this, &UMassSimulationSubsystem::OnPieEnded);
+
+		// note that this starts ticking for the editor world
+		StartSimulation(*World);
+	}
+#endif // WITH_EDITOR
+}
+
+void UMassSimulationSubsystem::OnWorldBeginPlay(UWorld& InWorld)
+{
+	// To evaluate the effective processors execution mode, we need to wait on OnWorldBeginPlay before calling RebuildTickPipeline as we are sure by this time the network is setup correctly.
+	RebuildTickPipeline();
+	// note that we're running for a game world right now. This means the PhaseManager->Start in OnPostWorldInit won't get called
+	StartSimulation(InWorld);
+}
+
+void UMassSimulationSubsystem::RebuildTickPipeline()
+{
+	check(PhaseManager);
+	PhaseManager->InitializePhases(*this);
+}
+
+void UMassSimulationSubsystem::StartSimulation(UWorld& InWorld)
+{
+	check(PhaseManager);
+	PhaseManager->Start(InWorld);
+
+	bSimulationStarted = true;
+
+	OnSimulationStarted.Broadcast(&InWorld);
+}
+
+void UMassSimulationSubsystem::StopSimulation()
+{
+	check(PhaseManager);
+	PhaseManager->Stop();
+
+	bSimulationStarted = false;
+}
+
+void UMassSimulationSubsystem::OnProcessingPhaseStarted(const float DeltaSeconds, const EPipeProcessingPhase Phase) const
+{
+	switch (Phase)
+	{
+		case EPipeProcessingPhase::PrePhysics:
+			{
+				TRACE_CPUPROFILER_EVENT_SCOPE(DoEntityCompation);
+				check(EntitySubsystem);
+				if (UE::MassSimulation::bDoEntityCompaction)
+				{
+					EntitySubsystem->DoEntityCompaction(GET_MASS_CONFIG_VALUE(DesiredEntityCompactionTimeSlicePerTick));
+				}
+			}
+			break;
+		default:
+			// unhandled phases, by design, not every phase needs to be handled by the Actor subsystem
+			break;
+	}
+}
+
+
+#if WITH_EDITOR
+void UMassSimulationSubsystem::OnPieBegin(const bool bIsSimulation)
+{
+	check(PhaseManager);
+	// called so that we're not processing phases for the editor world while PIE/SIE is running
+	StopSimulation();
+}
+
+void UMassSimulationSubsystem::OnPieEnded(const bool bIsSimulation)
+{
+	if (UWorld * World = GetWorld())
+	{
+		// Resume processing phases in the editor world.
+		StartSimulation(*World);
+	}
+}
+#endif // WITH_EDITOR
\ No newline at end of file
diff --git a/Engine/Plugins/Runtime/MassGameplay/Source/MassSimulation/Public/MassSimulationLocalCoordinator.h b/Engine/Plugins/Runtime/MassGameplay/Source/MassSimulation/Public/MassSimulationLocalCoordinator.h
new file mode 100644
--- /dev/null
+++ b/Engine/Plugins/Runtime/MassGameplay/Source/MassSimulation/Public/MassSimulationLocalCoordinator.h
@@ -0,0 +1,71 @@
+// Copyright Epic Games, Inc. All Rights Reserved.
+
+#pragma once
+
+#include "CoreMinimal.h"
+#include "Stats/Stats.h"
+#include "UObject/ObjectMacros.h"
+#include "Templates/SubclassOf.h"
+#include "UObject/SoftObjectPtr.h"
+#include "Engine/World.h"
+#include "GameFramework/Actor.h"
+#include "Engine/DataTable.h"
+#include "MassEntityTypes.h"
+#include "MassCommonTypes.h"
+#include "MassSimulationLocalCoordinator.generated.h"
+
+
+class UPipeSchematic;
+class UPipeEntitySubsystem;
+class UHierarchicalInstancedStaticMeshComponent;
+class UMassSpawningManager;
+class UMassDebugVisualizationComponent;
+
+
+#if WITH_EDITORONLY_DATA
+class UBillboardComponent;
+#endif // WITH_EDITORONLY_DATA
+
+// @todo leaving editable properties for the reference during transition period
+UCLASS(hidecategories = (Object, Actor, Input, Rendering, LOD, Cooking, Collision, HLOD, Partition), NotPlaceable, Transient)
+class MASSSIMULATION_API AMassSimulationLocalCoordinator : public AActor
+{
+	GENERATED_BODY()
+	//----------------------------------------------------------------------//
+	// temporary agent config. Will be pulled from spawning components later 
+	//----------------------------------------------------------------------//
+public:
+	UPROPERTY(EditAnywhere, Category = "Mass|Agent", meta=(ClampMin=0, UIMin=0))
+	float MaxSpeed = 50;
+
+public:
+	AMassSimulationLocalCoordinator(const FObjectInitializer& ObjectInitializer = FObjectInitializer::Get());
+	virtual void PostLoad() override;
+	
+protected:
+	// all spawned entities will go through these pipelines
+	UPROPERTY(EditAnywhere, Category = "Mass|Spawning")
+	TArray<TSoftObjectPtr<UPipeSchematic>> PostSpawnSchematics;
+
+	// Every tick these pipelines will be spawned 
+	UPROPERTY(EditAnywhere, Category = "Mass|Simulation")
+	TArray<TSoftObjectPtr<UPipeSchematic>> TickSchematics;
+
+	UPROPERTY(EditAnywhere, Category = "Mass|Spawning", meta=(DisplayName="LEGACY_LocallyAppliedSpawnComponents"))
+	TArray<const UScriptStruct*> LocallyAppliedSpawnComponents;
+	
+	UPROPERTY(EditAnywhere, meta=(ClampMin=100, UIMin=100), Category = "Mass|Simulation")
+	float ControlRadius = 5000;
+
+#if WITH_EDITORONLY_DATA
+	UPROPERTY(EditAnywhere, Category = "Mass|Simulation")
+	uint32 bTickInEditor : 1;
+
+	UPROPERTY(EditAnywhere, Category = "Mass|Debug")
+	uint32 bDrawDebug : 1;
+
+private:
+	UPROPERTY()
+	UBillboardComponent* SpriteComponent;
+#endif // WITH_EDITORONLY_DATA
+};
diff --git a/Engine/Plugins/Runtime/MassGameplay/Source/MassSimulation/Public/MassSimulationSettings.h b/Engine/Plugins/Runtime/MassGameplay/Source/MassSimulation/Public/MassSimulationSettings.h
new file mode 100644
--- /dev/null
+++ b/Engine/Plugins/Runtime/MassGameplay/Source/MassSimulation/Public/MassSimulationSettings.h
@@ -0,0 +1,60 @@
+// Copyright Epic Games, Inc. All Rights Reserved.
+
+#pragma once
+
+#include "Engine/DeveloperSettings.h"
+#include "MassSimulationSettings.generated.h"
+
+#define GET_MASS_CONFIG_VALUE(a) (GetMutableDefault<UMassSimulationSettings>()->a)
+
+class UPipeSchematic;
+
+/**
+ * Implements the settings for MassSimulation
+ */
+UCLASS(config = Game, defaultconfig, DisplayName = "Mass")
+class MASSSIMULATION_API UMassSimulationSettings : public UDeveloperSettings
+{
+	GENERATED_BODY()
+public:
+	DECLARE_MULTICAST_DELEGATE(FOnTickSchematicChangedDelegate);
+
+#if WITH_EDITOR
+public:
+	FOnTickSchematicChangedDelegate& GetOnTickSchematicChanged() { return OnTickSchematicChanged; }
+
+protected:
+	virtual void PostInitProperties() override;
+	virtual void PostEditChangeProperty(struct FPropertyChangedEvent& PropertyChangedEvent) override;
+	void OnAssetPropertiesChanged(UPipeSchematic* PipeSchematic, const FPropertyChangedEvent& PropertyChangedEvent);	
+#endif // WITH_EDITOR
+
+public:
+	UPROPERTY(EditDefaultsOnly, Category = "Runtime", config, meta=(DisplayName="DEPRECATED_TickSchematics"))
+	TArray<TSoftObjectPtr<UPipeSchematic>> TickSchematics;
+
+	/** The desired budget in seconds allowed to do actor spawning per frame */
+	UPROPERTY(EditDefaultsOnly, Category = "Runtime", config)
+	double DesiredActorSpawningTimeSlicePerTick = 0.0015;
+
+	/** The desired budget in seconds allowed to do actor destruction per frame */
+	UPROPERTY(EditDefaultsOnly, Category = "Runtime", config)
+	double DesiredActorDestructionTimeSlicePerTick = 0.0005;
+
+	/** The desired budget in seconds allowed to do entity compaction per frame */
+	UPROPERTY(EditDefaultsOnly, Category = "Runtime", config)
+	double DesiredEntityCompactionTimeSlicePerTick = 0.005;
+
+	/** The time to wait before retrying to spawn an actor that previously failed, default 10 seconds */
+	UPROPERTY(EditDefaultsOnly, Category = "Runtime", config)
+	float DesiredActorFailedSpawningRetryTimeInterval = 5.0f;
+
+	/** The distance a failed spawned actor needs to move before we retry, default 10 meters */
+	UPROPERTY(EditDefaultsOnly, Category = "Runtime", config)
+	float DesiredActorFailedSpawningRetryMoveDistance = 500.0f;
+
+protected:
+#if WITH_EDITOR
+	FOnTickSchematicChangedDelegate OnTickSchematicChanged;
+#endif // WITH_EDITOR
+};
diff --git a/Engine/Plugins/Runtime/MassGameplay/Source/MassSimulation/Public/MassSimulationSubsystem.h b/Engine/Plugins/Runtime/MassGameplay/Source/MassSimulation/Public/MassSimulationSubsystem.h
new file mode 100644
--- /dev/null
+++ b/Engine/Plugins/Runtime/MassGameplay/Source/MassSimulation/Public/MassSimulationSubsystem.h
@@ -0,0 +1,85 @@
+// Copyright Epic Games, Inc. All Rights Reserved.
+
+#pragma once
+
+#include "Delegates/IDelegateInstance.h"
+#include "MassEntityTypes.h"
+#include "MassProcessingPhase.h"
+#include "Subsystems/WorldSubsystem.h"
+#include "MassSimulationSubsystem.generated.h"
+
+
+class AMassSimulationLocalCoordinator;
+class UPipeEntitySubsystem;
+class UPipeProcessingPhaseManager;
+
+DECLARE_LOG_CATEGORY_EXTERN(LogMassSim, Log, All);
+
+UCLASS(config = Game, defaultconfig)
+class MASSSIMULATION_API UMassSimulationSubsystem : public UWorldSubsystem
+{
+	GENERATED_BODY()
+public:
+	DECLARE_MULTICAST_DELEGATE_TwoParams(FOnAdjustTickSchematics, UWorld* /*World*/, TArray<TSoftObjectPtr<UPipeSchematic>>& /*InOutTickSchematics*/);
+	DECLARE_MULTICAST_DELEGATE_OneParam(FOnSimulationStarted, UWorld* /*World*/);
+	
+	UMassSimulationSubsystem(const FObjectInitializer& ObjectInitializer = FObjectInitializer::Get());
+
+	UPipeEntitySubsystem* GetEntitySubsystem() const { return EntitySubsystem; }
+	const UPipeProcessingPhaseManager& GetPhaseManager() const { check(PhaseManager); return *PhaseManager; }
+
+	FPipeProcessingPhase::FOnPhaseEvent& GetOnProcessingPhaseStarted(const EPipeProcessingPhase Phase) const;
+	FPipeProcessingPhase::FOnPhaseEvent& GetOnProcessingPhaseFinished(const EPipeProcessingPhase Phase) const;
+	static FOnAdjustTickSchematics& GetOnAdjustTickSchematics() { return OnAdjustTickSchematics; }
+	static FOnSimulationStarted& GetOnSimulationStarted() { return OnSimulationStarted; }
+
+	bool IsSimulationStarted() const { return bSimulationStarted; }
+
+protected:
+	virtual void PostInitProperties() override;
+	virtual void Initialize(FSubsystemCollectionBase& Collection) override;
+	virtual void PostInitialize() override;
+	virtual void OnWorldBeginPlay(UWorld& InWorld) override;
+	virtual void Deinitialize() override;
+	virtual void BeginDestroy() override;
+
+	/** Creates an instance of UPipeEntitySubsystem and assigns it to EntitySubsystem property. This is the entity 
+	 *  subsystem that Mass is going to be using. If you want to instantiate a different entity system class or point 
+	 *  at a preexisting instance just override this function. */
+	virtual void CreateEntitySubsystem();
+
+	void RebuildTickPipeline();
+
+	void StartSimulation(UWorld& InWorld);
+	void StopSimulation();
+
+	void OnProcessingPhaseStarted(const float DeltaSeconds, const EPipeProcessingPhase Phase) const;
+
+#if WITH_EDITOR
+	void OnPieBegin(const bool bIsSimulation);
+	void OnPieEnded(const bool bIsSimulation);
+#endif // WITH_EDITOR
+
+protected:
+
+	UPROPERTY()
+	UPipeEntitySubsystem* EntitySubsystem;
+
+	UPROPERTY()
+	UPipeProcessingPhaseManager* PhaseManager;
+
+	inline static FOnAdjustTickSchematics OnAdjustTickSchematics={};
+	inline static FOnSimulationStarted OnSimulationStarted={};
+
+	UPROPERTY()
+	FRuntimePipeline RuntimePipeline;
+
+	float CurrentDeltaSeconds = 0.f;
+	bool bTickInProgress = false;
+	bool bSimulationStarted = false;
+
+#if WITH_EDITOR
+	FDelegateHandle PieBeginEventHandle;
+	FDelegateHandle PieEndedEventHandle;
+#endif // WITH_EDITOR
+};
diff --git a/Engine/Plugins/Runtime/MassGameplay/Source/MassSmartObjects/MassSmartObjects.Build.cs b/Engine/Plugins/Runtime/MassGameplay/Source/MassSmartObjects/MassSmartObjects.Build.cs
new file mode 100644
--- /dev/null
+++ b/Engine/Plugins/Runtime/MassGameplay/Source/MassSmartObjects/MassSmartObjects.Build.cs
@@ -0,0 +1,44 @@
+// Copyright Epic Games, Inc. All Rights Reserved.
+
+namespace UnrealBuildTool.Rules
+{
+	public class MassSmartObjects : ModuleRules
+	{
+		public MassSmartObjects(ReadOnlyTargetRules Target) : base(Target)
+		{
+			PCHUsage = ModuleRules.PCHUsageMode.UseExplicitOrSharedPCHs;
+
+			PublicIncludePaths.AddRange(
+				new string[] {
+				"Runtime/AIModule/Public",
+				ModuleDirectory + "/Public",
+				}
+			);
+
+			PrivateIncludePaths.AddRange(
+				new string[] {
+				}
+			);
+
+			PublicDependencyModuleNames.AddRange(
+				new string[] {
+					"MassEntity",
+					"Core",
+					"CoreUObject",
+					"DeveloperSettings",
+					"Engine",
+					"GameplayTags",
+					"MassCommon",
+					"MassLOD",
+					"MassMovement",
+					"MassSignals",
+					"MassSimulation",
+					"SmartObjectsModule",
+					"StructUtils",
+					"ZoneGraph",
+					"ZoneGraphAnnotations"
+				}
+			);
+		}
+	}
+}
diff --git a/Engine/Plugins/Runtime/MassGameplay/Source/MassSmartObjects/Private/MassSmartObjectBehaviorConfig.cpp b/Engine/Plugins/Runtime/MassGameplay/Source/MassSmartObjects/Private/MassSmartObjectBehaviorConfig.cpp
new file mode 100644
--- /dev/null
+++ b/Engine/Plugins/Runtime/MassGameplay/Source/MassSmartObjects/Private/MassSmartObjectBehaviorConfig.cpp
@@ -0,0 +1,11 @@
+// Copyright Epic Games, Inc. All Rights Reserved.
+
+#include "MassSmartObjectBehaviorConfig.h"
+#include "MassSmartObjectProcessor.h"
+
+void USmartObjectMassBehaviorConfig::Activate(UEntitySubsystem& EntitySubsystem,
+                                              FLWComponentSystemExecutionContext& Context,
+                                              const FMassBehaviorEntityContext& EntityContext) const
+{
+	EntityContext.SOUser.SetUseTime(UseTime);
+}
diff --git a/Engine/Plugins/Runtime/MassGameplay/Source/MassSmartObjects/Private/MassSmartObjectHandler.cpp b/Engine/Plugins/Runtime/MassGameplay/Source/MassSmartObjects/Private/MassSmartObjectHandler.cpp
new file mode 100644
--- /dev/null
+++ b/Engine/Plugins/Runtime/MassGameplay/Source/MassSmartObjects/Private/MassSmartObjectHandler.cpp
@@ -0,0 +1,252 @@
+// Copyright Epic Games, Inc. All Rights Reserved.
+
+#include "MassSmartObjectHandler.h"
+#include "MassSmartObjectProcessor.h"
+#include "MassSmartObjectBehaviorConfig.h"
+#include "SmartObjectSubsystem.h"
+#include "VisualLogger/VisualLogger.h"
+#include "LWCCommandBuffer.h"
+
+namespace UE::Mass::SmartObject
+{
+struct FPayload
+{
+	FLWEntity Entity;
+	TWeakObjectPtr<UEntitySubsystem> EntitySubsystem;
+	TWeakObjectPtr<USmartObjectSubsystem> SmartObjectSubsystem;
+};
+
+void OnSlotInvalidated(const FSmartObjectClaimHandle& ClaimHandle, const ESmartObjectSlotState State, FPayload Payload)
+{
+	FDataFragment_SmartObjectUser& User = Payload.EntitySubsystem->GetComponentDataChecked<FDataFragment_SmartObjectUser>(Payload.Entity);
+	if (USmartObjectSubsystem* SmartObjectSubsystem = Payload.SmartObjectSubsystem.Get())
+	{
+		SmartObjectSubsystem->Release(User.ClaimHandle);
+	}
+
+	User.InteractionStatus = EMassSmartObjectInteractionStatus::Aborted;
+	User.ClaimHandle.Invalidate();
+}
+} // UE::Mass::SmartObject;
+
+//----------------------------------------------------------------------//
+// FMassSmartObjectHandler
+//----------------------------------------------------------------------//
+FMassSmartObjectRequestID FMassSmartObjectHandler::FindCandidatesAsync(const FLWEntity RequestingEntity, const FVector& Location) const
+{
+	const FLWEntity ReservedEntity = EntitySubsystem.ReserveEntity();
+
+	FMassSmartObjectWorldLocationRequestFragment RequestFragment;
+	RequestFragment.SearchOrigin = Location;
+	RequestFragment.RequestingEntity = RequestingEntity;
+
+	FMassSmartObjectRequestResultFragment ResultFragment;
+
+	ExecutionContext.Defer().PushCommand(
+		FBuildEntityFromComponentInstances(ReservedEntity,
+			{
+				FStructView::Make(RequestFragment),
+				FStructView::Make(ResultFragment)
+			}));
+
+	return FMassSmartObjectRequestID(ReservedEntity);
+}
+
+FMassSmartObjectRequestID FMassSmartObjectHandler::FindCandidatesAsync(const FLWEntity RequestingEntity, const FZoneGraphCompactLaneLocation& LaneLocation) const
+{
+	const FLWEntity ReservedEntity = EntitySubsystem.ReserveEntity();
+
+	FMassSmartObjectLaneLocationRequestFragment RequestFragment;
+	RequestFragment.CompactLaneLocation = LaneLocation;
+	RequestFragment.RequestingEntity = RequestingEntity;
+
+	FMassSmartObjectRequestResultFragment ResultFragment;
+
+	ExecutionContext.Defer().PushCommand(
+		FBuildEntityFromComponentInstances(ReservedEntity,
+			{
+				FStructView::Make(RequestFragment),
+				FStructView::Make(ResultFragment)
+			}));
+
+	return FMassSmartObjectRequestID(ReservedEntity);
+}
+
+FMassSmartObjectRequestResult FMassSmartObjectHandler::GetRequestResult(const FMassSmartObjectRequestID& RequestID) const
+{
+	const FLWEntity RequestEntity = static_cast<FLWEntity>(RequestID);
+	if (!ensureMsgf(EntitySubsystem.IsEntityValid(RequestEntity), TEXT("Invalid request.")))
+	{
+		return {};
+	}
+
+	// Check if entity is built by now.
+	if (!EntitySubsystem.IsEntityBuilt(RequestEntity))
+	{
+		return {};
+	}
+
+	const FMassSmartObjectRequestResultFragment& RequestFragment = EntitySubsystem.GetComponentDataChecked<FMassSmartObjectRequestResultFragment>(RequestEntity);
+	return RequestFragment.Result;
+}
+
+void FMassSmartObjectHandler::RemoveRequest(const FMassSmartObjectRequestID& RequestID) const
+{
+	const FLWEntity RequestEntity = static_cast<FLWEntity>(RequestID);
+	ExecutionContext.Defer().DestroyEntity(RequestEntity);
+}
+
+EMassSmartObjectClaimResult FMassSmartObjectHandler::ClaimCandidate(const FLWEntity Entity, FDataFragment_SmartObjectUser& User, const FMassSmartObjectRequestID& RequestID) const
+{
+	if (!ensureMsgf(RequestID.IsSet(), TEXT("Trying to claim using an invalid request.")))
+	{
+		return EMassSmartObjectClaimResult::Failed_InvalidRequest;
+	}
+
+	const FLWEntity RequestEntity = static_cast<FLWEntity>(RequestID);
+	const FMassSmartObjectRequestResultFragment& RequestFragment = EntitySubsystem.GetComponentDataChecked<FMassSmartObjectRequestResultFragment>(RequestEntity);
+	if (!ensureMsgf(RequestFragment.Result.bProcessed, TEXT("Trying to claim using a request that is not processed.")))
+	{
+		return EMassSmartObjectClaimResult::Failed_UnprocessedRequest;
+	}
+
+	return ClaimCandidate(Entity, User, RequestFragment.Result);
+}
+
+EMassSmartObjectClaimResult FMassSmartObjectHandler::ClaimCandidate(const FLWEntity Entity, FDataFragment_SmartObjectUser& User, const FMassSmartObjectRequestResult& SearchRequestResult) const
+{
+	if (!ensureMsgf(SearchRequestResult.bProcessed, TEXT("Trying to claim using a search request that is not processed.")))
+	{
+		return EMassSmartObjectClaimResult::Failed_UnprocessedRequest;
+	}
+
+	checkf(!User.ClaimHandle.IsValid(), TEXT("User should not already have a valid handle."));
+	TConstArrayView<FSmartObjectCandidate> View = MakeArrayView(SearchRequestResult.Candidates.GetData(), SearchRequestResult.NumCandidates);
+
+	for (const FSmartObjectCandidate& Candidate : View)
+	{
+		if (ClaimSmartObject(Entity, User, Candidate.ID))
+		{
+#if WITH_MASS_DEBUG
+			UE_CVLOG(UE::MassDebug::IsDebuggingEntity(Entity),
+				&SmartObjectSubsystem,
+				LogSmartObject,
+				Log,
+				TEXT("[%s] claimed [%s]"),
+				*Entity.DebugGetDescription(),
+				*Candidate.ID.Describe());
+#endif // WITH_MASS_DEBUG
+			break;
+		}
+	}
+
+	return User.ClaimHandle.IsValid() ? EMassSmartObjectClaimResult::Succeeded : EMassSmartObjectClaimResult::Failed_NoAvailableCandidate;
+}
+
+bool FMassSmartObjectHandler::ClaimSmartObject(const FLWEntity Entity, FDataFragment_SmartObjectUser& User, const FSmartObjectID& ObjectID) const
+{
+	FSmartObjectRequestFilter Filter;
+	Filter.BehaviorConfigurationClass = USmartObjectMassBehaviorConfig::StaticClass();
+
+	const FSmartObjectClaimHandle ClaimHandle = SmartObjectSubsystem.Claim(ObjectID, Filter);
+	const bool bSuccess = ClaimHandle.IsValid();
+	if (bSuccess)
+	{
+		User.ClaimHandle = ClaimHandle;
+		User.InteractionStatus = EMassSmartObjectInteractionStatus::Unset;
+		User.TargetLocation = SmartObjectSubsystem.GetSlotLocation(User.ClaimHandle).Get(FVector::ZeroVector);
+
+		// @todo SO: need to unregister callback if entity or component gets removed. Or at least support getting called with a no longer valid LWEntity.
+		UE::Mass::SmartObject::FPayload Payload;
+		Payload.Entity = Entity;
+		Payload.EntitySubsystem = &EntitySubsystem;
+		Payload.SmartObjectSubsystem = &SmartObjectSubsystem;
+		SmartObjectSubsystem.RegisterSlotInvalidationCallback(ClaimHandle, FOnSlotInvalidated::CreateStatic(&UE::Mass::SmartObject::OnSlotInvalidated, Payload));
+	}
+
+#if WITH_MASS_DEBUG
+	UE_CVLOG(UE::MassDebug::IsDebuggingEntity(Entity),
+		&SmartObjectSubsystem,
+		LogSmartObject,
+		Log,
+		TEXT("[%s] claim for [%s] %s"),
+		*Entity.DebugGetDescription(),
+		*ObjectID.Describe(),
+		bSuccess ? TEXT("Succeeded") : TEXT("Failed"));
+#endif // WITH_MASS_DEBUG
+
+	return bSuccess;
+}
+
+bool FMassSmartObjectHandler::UseSmartObject(
+	const FLWEntity Entity,
+	FDataFragment_SmartObjectUser& User,
+	const FDataFragment_Transform& Transform) const
+{
+#if WITH_MASS_DEBUG
+	UE_CVLOG(UE::MassDebug::IsDebuggingEntity(Entity),
+		&SmartObjectSubsystem,
+		LogSmartObject,
+		Log,
+		TEXT("[%s] starts using [%s]"),
+		*Entity.DebugGetDescription(),
+		*User.ClaimHandle.Describe());
+#endif // WITH_MASS_DEBUG
+
+	const USmartObjectMassBehaviorConfig* Config = SmartObjectSubsystem.Use<USmartObjectMassBehaviorConfig>(User.ClaimHandle);
+	if (Config == nullptr)
+	{
+		return false;
+	}
+
+	Config->Activate(EntitySubsystem, ExecutionContext, FMassBehaviorEntityContext(Entity, Transform, User, SmartObjectSubsystem));
+
+	User.InteractionStatus = EMassSmartObjectInteractionStatus::InProgress;
+	ExecutionContext.Defer().AddTag<FMassSmartObjectTimedBehaviorTag>(Entity);
+
+	return true;
+}
+
+void FMassSmartObjectHandler::ReleaseSmartObject(const FLWEntity Entity, FDataFragment_SmartObjectUser& User, const EMassSmartObjectInteractionStatus NewStatus) const
+{
+#if WITH_MASS_DEBUG
+	UE_CVLOG(UE::MassDebug::IsDebuggingEntity(Entity),
+		&SmartObjectSubsystem,
+		LogSmartObject,
+		Log,
+		TEXT("[%s] releases handle [%s]"),
+		*Entity.DebugGetDescription(),
+		*User.ClaimHandle.Describe());
+#endif // WITH_MASS_DEBUG
+
+#if DO_ENSURE
+	const EMassSmartObjectInteractionStatus CurrentStatus = User.InteractionStatus;
+
+	switch (CurrentStatus)
+	{
+	case EMassSmartObjectInteractionStatus::Unset:
+		ensureMsgf(NewStatus == EMassSmartObjectInteractionStatus::Unset,
+			TEXT("Not expecting new status for inactive interaction. Received %s"),
+			*UEnum::GetValueAsString(NewStatus));
+		break;
+
+	case EMassSmartObjectInteractionStatus::InProgress:
+		ensureMsgf(NewStatus == EMassSmartObjectInteractionStatus::Completed || NewStatus == EMassSmartObjectInteractionStatus::Aborted,
+			TEXT("Expecting status 'Completed' or 'Aborted' for in progress interaction. Received %s"),
+			*UEnum::GetValueAsString(NewStatus));
+		break;
+
+	case EMassSmartObjectInteractionStatus::Completed:
+	case EMassSmartObjectInteractionStatus::Aborted:
+		ensureMsgf(false, TEXT("Not expecting status changes for'Completed' or 'Aborted' interaction. Received %s"), *UEnum::GetValueAsString(NewStatus));
+		break;
+
+	default:
+		ensureMsgf(false, TEXT("Unhandled state: %s"), *UEnum::GetValueAsString(NewStatus));
+	}
+#endif
+
+	SmartObjectSubsystem.Release(User.ClaimHandle);
+	User.InteractionStatus = NewStatus;
+	User.ClaimHandle.Invalidate();
+}
diff --git a/Engine/Plugins/Runtime/MassGameplay/Source/MassSmartObjects/Private/MassSmartObjectProcessor.cpp b/Engine/Plugins/Runtime/MassGameplay/Source/MassSmartObjects/Private/MassSmartObjectProcessor.cpp
new file mode 100644
--- /dev/null
+++ b/Engine/Plugins/Runtime/MassGameplay/Source/MassSmartObjects/Private/MassSmartObjectProcessor.cpp
@@ -0,0 +1,348 @@
+// Copyright Epic Games, Inc. All Rights Reserved.
+
+#include "MassSmartObjectProcessor.h"
+#include "LWCCommandBuffer.h"
+#include "MassCommonTypes.h"
+#include "MassSignalSubsystem.h"
+#include "MassSmartObjectBehaviorConfig.h"
+#include "MassSmartObjectSettings.h"
+#include "SmartObjectZoneAnnotations.h"
+#include "Misc/ScopeExit.h"
+#include "SmartObjectOctree.h"
+#include "SmartObjectSubsystem.h"
+#include "VisualLogger/VisualLogger.h"
+#include "ZoneGraphAnnotationSubsystem.h"
+#include "ZoneGraphSubsystem.h"
+
+//----------------------------------------------------------------------//
+// UMassProcessor_SmartObjectBase
+//----------------------------------------------------------------------//
+void UMassProcessor_SmartObjectBase::Initialize(UObject& Owner)
+{
+	Super::Initialize(Owner);
+
+	SmartObjectSubsystem = UWorld::GetSubsystem<USmartObjectSubsystem>(GetWorld());
+	SignalSubsystem = UWorld::GetSubsystem<UMassSignalSubsystem>(Owner.GetWorld());
+}
+
+//----------------------------------------------------------------------//
+// UMassProcessor_SmartObjectCandidatesFinder
+//----------------------------------------------------------------------//
+void UMassProcessor_SmartObjectCandidatesFinder::Initialize(UObject& Owner)
+{
+	Super::Initialize(Owner);
+
+	AnnotationSubsystem = UWorld::GetSubsystem<UZoneGraphAnnotationSubsystem>(GetWorld());
+	ZoneGraphSubsystem = UWorld::GetSubsystem<UZoneGraphSubsystem>(GetWorld());
+}
+
+void UMassProcessor_SmartObjectCandidatesFinder::ConfigureQueries()
+{
+	WorldRequestQuery.AddRequirement<FMassSmartObjectWorldLocationRequestFragment>(ELWComponentAccess::ReadOnly);
+	WorldRequestQuery.AddRequirement<FMassSmartObjectRequestResultFragment>(ELWComponentAccess::ReadWrite);
+	WorldRequestQuery.AddTagRequirement<FMassSmartObjectCompletedRequestTag>(ELWComponentPresence::None);
+
+	LaneRequestQuery.AddRequirement<FMassSmartObjectLaneLocationRequestFragment>(ELWComponentAccess::ReadOnly);
+	LaneRequestQuery.AddRequirement<FMassSmartObjectRequestResultFragment>(ELWComponentAccess::ReadWrite);
+	LaneRequestQuery.AddTagRequirement<FMassSmartObjectCompletedRequestTag>(ELWComponentPresence::None);
+}
+
+UMassProcessor_SmartObjectCandidatesFinder::UMassProcessor_SmartObjectCandidatesFinder()
+{
+	// 1. Frame T Behavior create a request(deferred entity creation)
+	// 2. Frame T+1: Processor execute the request might mark it as done(deferred add tag flushed at the end of the frame)
+	// 3. Frame T+1: Behavior could cancel request(deferred destroy entity)
+	// If the processor does not run before the behaviors, step 2 and 3 are flipped and it will crash while flushing the deferred commands
+	ExecutionOrder.ExecuteBefore.Add(UE::Mass::ProcessorGroupNames::Behavior);
+}
+
+void UMassProcessor_SmartObjectCandidatesFinder::Execute(UEntitySubsystem& EntitySubsystem, FLWComponentSystemExecutionContext& Context)
+{
+	checkf(SmartObjectSubsystem != nullptr, TEXT("SmartObjectSubsystem should exist when executing processors."));
+	checkf(SignalSubsystem != nullptr, TEXT("MassSignalSubsystem should exist when executing processors."));
+	checkf(ZoneGraphSubsystem != nullptr, TEXT("ZoneGraphSubsystem should exist when executing processors."));
+	checkf(AnnotationSubsystem != nullptr, TEXT("ZoneGraphAnnotationSubsystem should exist when executing processors."));
+
+	// Create filter
+	FSmartObjectRequestFilter Filter;
+	Filter.BehaviorConfigurationClass = USmartObjectMassBehaviorConfig::StaticClass();
+
+	// Build list of request owner entities to send a completion signal
+	TArray<FLWEntity> EntitiesToSignal;
+
+	auto BeginRequestProcessing = [](const FLWEntity Entity, FLWComponentSystemExecutionContext& Context, FMassSmartObjectRequestResult& Result)
+	{
+		Context.Defer().AddTag<FMassSmartObjectCompletedRequestTag>(Entity);
+		Result.NumCandidates = 0;
+	};
+
+	auto EndRequestProcessing = [](const UObject* LogOwner, const FLWEntity Entity, FMassSmartObjectRequestResult& Result)
+	{
+		if (Result.NumCandidates > 0)
+		{
+			TArrayView<FSmartObjectCandidate> View = MakeArrayView(Result.Candidates.GetData(), Result.NumCandidates);
+			Algo::Sort(View, [](const FSmartObjectCandidate& LHS, const FSmartObjectCandidate& RHS) { return LHS.Cost < RHS.Cost; });
+		}
+		Result.bProcessed = true;
+
+#if WITH_MASS_DEBUG
+		UE_VLOG(LogOwner, LogSmartObject, Verbose, TEXT("[%s] search completed: found %d"), *Entity.DebugGetDescription(), Result.NumCandidates);
+#endif // WITH_MASS_DEBUG
+	};
+
+	// Process world location based requests
+	WorldRequestQuery.ForEachEntityChunk(EntitySubsystem, Context, [this, &Filter, &EntitiesToSignal, &BeginRequestProcessing, &EndRequestProcessing](FLWComponentSystemExecutionContext& Context)
+	{
+		const FSmartObjectOctree& Octree = SmartObjectSubsystem->GetOctree();
+
+		const int32 NumEntities = Context.GetEntitiesNum();
+		EntitiesToSignal.Reserve(EntitiesToSignal.Num() + NumEntities);
+
+		const TConstArrayView<FMassSmartObjectWorldLocationRequestFragment> RequestList = Context.GetComponentView<FMassSmartObjectWorldLocationRequestFragment>();
+		const TArrayView<FMassSmartObjectRequestResultFragment> ResultList = Context.GetMutableComponentView<FMassSmartObjectRequestResultFragment>();
+
+		for (int32 i = 0; i < NumEntities; ++i)
+		{
+			const FMassSmartObjectWorldLocationRequestFragment& RequestFragment = RequestList[i];
+			FMassSmartObjectRequestResult& Result = ResultList[i].Result;
+			EntitiesToSignal.Add(RequestFragment.RequestingEntity);
+
+			const FVector& SearchOrigin = RequestFragment.SearchOrigin;
+			const FBox& SearchBounds = FBox::BuildAABB(SearchOrigin, FVector(SearchExtents));
+
+			const FLWEntity Entity = Context.GetEntity(i);
+			bool bDisplayDebug = false;
+			FColor DebugColor(FColor::White);
+
+#if WITH_MASS_DEBUG
+			bDisplayDebug = UE::MassDebug::IsDebuggingEntity(Entity, &DebugColor);
+#endif // WITH_MASS_DEBUG
+
+			BeginRequestProcessing(Entity, Context, Result);
+			ON_SCOPE_EXIT{ EndRequestProcessing(SmartObjectSubsystem, Entity, Result);	};
+
+			Octree.FindElementsWithBoundsTest(SearchBounds, [this, &Filter, &Entity, &Result, &SearchOrigin, &bDisplayDebug, &DebugColor](const FSmartObjectOctreeElement& Element)
+			{
+				if (Result.NumCandidates < FMassSmartObjectRequestResult::MaxNumCandidates)
+				{
+					const FSmartObjectID ID = Element.SmartObjectID;
+
+					// Make sure that we can use a slot in that object (availability with supported config, etc.)
+					if (SmartObjectSubsystem->FindSlot(ID, Filter).IsValid())
+					{
+						const FVector ObjectLocation = Element.Bounds.Center;
+						Result.Candidates[Result.NumCandidates++] = FSmartObjectCandidate(ID, FVector::DistSquared(SearchOrigin, ObjectLocation));
+
+#if WITH_MASS_DEBUG
+						if (bDisplayDebug)
+						{
+							constexpr float DebugRadius = 10.f;
+							UE_VLOG_LOCATION(SmartObjectSubsystem, LogSmartObject, Display, ObjectLocation, DebugRadius, DebugColor, TEXT("%s"), *ID.Describe());
+							UE_VLOG_SEGMENT(SmartObjectSubsystem, LogSmartObject, Display, SearchOrigin, ObjectLocation, DebugColor, TEXT(""));
+						}
+#endif // WITH_MASS_DEBUG
+					}
+				}
+			});
+		}
+	});
+
+	// Process lane based requests
+	const FZoneGraphTag SmartObjectTag = GetDefault<UMassSmartObjectSettings>()->SmartObjectTag;
+	USmartObjectZoneAnnotations* SmartObjectBehavior = Cast<USmartObjectZoneAnnotations>(AnnotationSubsystem->GetFirstAnnotationForTag(SmartObjectTag));
+	if (SmartObjectBehavior != nullptr)
+	{
+		LaneRequestQuery.ForEachEntityChunk(EntitySubsystem, Context, [this, SmartObjectBehavior, &Filter, SmartObjectTag, &EntitiesToSignal, &BeginRequestProcessing, &EndRequestProcessing](FLWComponentSystemExecutionContext& Context)
+		{
+			const int32 NumEntities = Context.GetEntitiesNum();
+			EntitiesToSignal.Reserve(EntitiesToSignal.Num() + NumEntities);
+
+			TConstArrayView<FMassSmartObjectLaneLocationRequestFragment> RequestList = Context.GetComponentView<FMassSmartObjectLaneLocationRequestFragment>();
+			TArrayView<FMassSmartObjectRequestResultFragment> ResultList = Context.GetMutableComponentView<FMassSmartObjectRequestResultFragment>();
+
+			// Cache latest used data since request are most of the time on the same zone graph
+			FZoneGraphDataHandle LastUsedDataHandle;
+			const FSmartObjectAnnotationData* GraphData = nullptr;
+
+			for (int32 i = 0; i < NumEntities; ++i)
+			{
+				const FMassSmartObjectLaneLocationRequestFragment& RequestFragment = RequestList[i];
+				FMassSmartObjectRequestResult& Result = ResultList[i].Result;
+				EntitiesToSignal.Add(RequestFragment.RequestingEntity);
+
+				const FZoneGraphCompactLaneLocation RequestLocation = RequestFragment.CompactLaneLocation;
+				const FZoneGraphLaneHandle RequestLaneHandle = RequestLocation.LaneHandle;
+
+				const FLWEntity Entity = Context.GetEntity(i);
+				bool bDisplayDebug = false;
+#if WITH_MASS_DEBUG
+				FColor DebugColor(FColor::White);
+				bDisplayDebug = UE::MassDebug::IsDebuggingEntity(Entity, &DebugColor);
+#endif // WITH_MASS_DEBUG
+
+				BeginRequestProcessing(Entity, Context, Result);
+				ON_SCOPE_EXIT{ EndRequestProcessing(SmartObjectSubsystem, Entity, Result); };
+
+				if (!ensureMsgf(RequestLaneHandle.IsValid(), TEXT("Requesting smart objects using an invalid handle")))
+				{
+					continue;
+				}
+
+				// Fetch smart object data associated to the current graph if different than last used one
+				if (LastUsedDataHandle != RequestLaneHandle.DataHandle)
+				{
+					LastUsedDataHandle = RequestLaneHandle.DataHandle;
+					GraphData = SmartObjectBehavior->GetAnnotationData(RequestLaneHandle.DataHandle);
+				}
+
+				if (GraphData == nullptr)
+				{
+					continue;
+				}
+
+				// Fetch current annotations for the specified lane and look for the smart object tag
+				const FZoneGraphTagMask LaneMask = AnnotationSubsystem->GetAnnotationTags(RequestLaneHandle);
+				if (!LaneMask.Contains(SmartObjectTag))
+				{
+					continue;
+				}
+
+				const FSmartObjectList* SmartObjectList = GraphData->LaneToSmartObjectsLookup.Find(RequestLaneHandle.Index);
+				if (SmartObjectList == nullptr || !ensureMsgf(SmartObjectList->SmartObjects.Num() > 0, TEXT("Lookup table should only contains lanes with one or more associated object(s).")))
+				{
+					continue;
+				}
+
+				for (const FSmartObjectID ID : SmartObjectList->SmartObjects)
+				{
+					// Find entry point using FindChecked since all smart objects added to LaneToSmartObjects lookup table
+					// were also added to the entry point lookup table
+					const FSmartObjectLaneLocation EntryPoint = GraphData->ObjectToEntryPointLookup.FindChecked(ID);
+
+					float Cost = 0.f;
+					if (ensureMsgf(EntryPoint.LaneIndex == RequestLocation.LaneHandle.Index, TEXT("Must be on same lane to be able to use distance along lane.")))
+					{
+						// Only consider object ahead
+						const float DistAhead = EntryPoint.DistanceAlongLane - RequestLocation.DistanceAlongLane;
+						if (DistAhead < 0)
+						{
+							continue;
+						}
+						Cost = DistAhead;
+					}
+
+					// Make sure that we can use a slot in that object (availability with supported config, etc.)
+					if (!SmartObjectSubsystem->FindSlot(ID, Filter).IsValid())
+					{
+						continue;
+					}
+
+					Result.Candidates[Result.NumCandidates++] = FSmartObjectCandidate(ID, Cost);
+
+#if WITH_MASS_DEBUG
+					if (bDisplayDebug)
+					{
+						FZoneGraphLaneLocation RequestLaneLocation, EntryPointLaneLocation;
+						ZoneGraphSubsystem->CalculateLocationAlongLane(RequestLaneHandle, RequestLocation.DistanceAlongLane, RequestLaneLocation);
+						ZoneGraphSubsystem->CalculateLocationAlongLane(RequestLaneHandle, EntryPoint.DistanceAlongLane, EntryPointLaneLocation);
+
+						constexpr float DebugRadius = 10.f;
+						UE_VLOG_LOCATION(SmartObjectSubsystem, LogSmartObject, Display, EntryPointLaneLocation.Position, DebugRadius, DebugColor, TEXT("%s"), *ID.Describe());
+						UE_VLOG_SEGMENT(SmartObjectSubsystem, LogSmartObject, Display, RequestLaneLocation.Position, EntryPointLaneLocation.Position, DebugColor, TEXT(""));
+					}
+#endif // WITH_MASS_DEBUG
+				}
+			}
+		});
+	}
+
+	// Signal entities that their search results are ready
+	if (EntitiesToSignal.Num())
+	{
+		SignalSubsystem->SignalEntities(UE::Mass::Signals::SmartObjectCandidatesReady, EntitiesToSignal);
+	}
+}
+
+//----------------------------------------------------------------------//
+// UMassProcessor_SmartObjectTimedBehavior
+//----------------------------------------------------------------------//
+void UMassProcessor_SmartObjectTimedBehavior::ConfigureQueries()
+{
+	EntityQuery.AddRequirement<FDataFragment_SmartObjectUser>(ELWComponentAccess::ReadWrite);
+	EntityQuery.AddTagRequirement<FMassSmartObjectTimedBehaviorTag>(ELWComponentPresence::All);
+}
+
+UMassProcessor_SmartObjectTimedBehavior::UMassProcessor_SmartObjectTimedBehavior()
+{
+	ExecutionOrder.ExecuteInGroup = UE::Mass::ProcessorGroupNames::Tasks;
+	ExecutionOrder.ExecuteAfter.Add(UE::Mass::ProcessorGroupNames::Behavior);
+}
+
+void UMassProcessor_SmartObjectTimedBehavior::Execute(UEntitySubsystem& EntitySubsystem, FLWComponentSystemExecutionContext& Context)
+{
+	checkf(SmartObjectSubsystem != nullptr, TEXT("SmartObjectSubsystem should exist when executing processors."));
+	checkf(SignalSubsystem != nullptr, TEXT("MassSignalSubsystem should exist when executing processors."));
+
+	TArray<TTuple<FLWEntity, FDataFragment_SmartObjectUser*>> ToRelease;
+
+	QUICK_SCOPE_CYCLE_COUNTER(UMassProcessor_SmartObjectTestBehavior_Run);
+
+	EntityQuery.ForEachEntityChunk(EntitySubsystem, Context, [this, &ToRelease](FLWComponentSystemExecutionContext& Context)
+	{
+		const int32 NumEntities = Context.GetEntitiesNum();
+		const TArrayView<FDataFragment_SmartObjectUser> UserList = Context.GetMutableComponentView<FDataFragment_SmartObjectUser>();
+
+		for (int32 i = 0; i < NumEntities; ++i)
+		{
+			FDataFragment_SmartObjectUser& SOUser = UserList[i];
+			if (SOUser.InteractionStatus != EMassSmartObjectInteractionStatus::InProgress)
+			{
+				// Only expecting in progress interaction status
+				continue;
+			}
+
+			const float DT = Context.GetDeltaTimeSeconds();
+			SOUser.SetUseTime(FMath::Max(SOUser.GetUseTime() - DT, 0.0f));
+			const bool bMustRelease = SOUser.GetUseTime() <= 0.f;
+
+#if WITH_MASS_DEBUG
+			const FLWEntity Entity = Context.GetEntity(i);
+			FColor DebugColor(FColor::White);
+			const bool bIsDebuggingEntity = UE::MassDebug::IsDebuggingEntity(Entity, &DebugColor);
+			if (bIsDebuggingEntity)
+			{
+				UE_CVLOG(bMustRelease, SmartObjectSubsystem, LogSmartObject, Log, TEXT("[%s] stops using [%s]"), *Entity.DebugGetDescription(), *SOUser.GetClaimHandle().Describe());
+				UE_CVLOG(!bMustRelease, SmartObjectSubsystem, LogSmartObject, Verbose, TEXT("[%s] using [%s] for %.1f"), *Entity.DebugGetDescription(), *SOUser.GetClaimHandle().Describe(), SOUser.GetUseTime());
+
+				const TOptional<FTransform> Transform = SmartObjectSubsystem->GetSlotTransform(SOUser.ClaimHandle);
+				if (Transform.IsSet())
+				{
+					constexpr float Radius = 40.f;
+					const FVector HalfHeightOffset(0.f, 0.f, 100.f);
+					const FVector Pos = Transform.GetValue().GetLocation();
+					const FVector Dir = Transform.GetValue().GetRotation().GetForwardVector();
+					UE_VLOG_CYLINDER(SmartObjectSubsystem, LogSmartObject, Display, Pos - HalfHeightOffset, Pos + HalfHeightOffset, Radius, DebugColor, TEXT(""));
+					UE_VLOG_ARROW(SmartObjectSubsystem, LogSmartObject, Display, Pos, Pos + Dir * 2.0f * Radius, DebugColor, TEXT(""));
+				}
+			}
+#endif // WITH_MASS_DEBUG
+
+			if (bMustRelease)
+			{
+				ToRelease.Add(MakeTuple(Context.GetEntity(i), &SOUser));
+			}
+		}
+	});
+
+	for (const TTuple<FLWEntity, FDataFragment_SmartObjectUser*> EntryToRelease : ToRelease)
+	{
+		const FLWEntity Entity = EntryToRelease.Get<0>();
+		SignalSubsystem->SignalEntity(UE::Mass::Signals::SmartObjectInteractionDone, Entity);
+
+		FDataFragment_SmartObjectUser* User = EntryToRelease.Get<1>();
+		SmartObjectSubsystem->Release(User->ClaimHandle);
+		User->InteractionStatus = EMassSmartObjectInteractionStatus::Completed;
+		Context.Defer().RemoveTag<FMassSmartObjectTimedBehaviorTag>(Entity);
+		User->ClaimHandle.Invalidate();
+	}
+}
diff --git a/Engine/Plugins/Runtime/MassGameplay/Source/MassSmartObjects/Public/MassSmartObjectBehaviorConfig.h b/Engine/Plugins/Runtime/MassGameplay/Source/MassSmartObjects/Public/MassSmartObjectBehaviorConfig.h
new file mode 100644
--- /dev/null
+++ b/Engine/Plugins/Runtime/MassGameplay/Source/MassSmartObjects/Public/MassSmartObjectBehaviorConfig.h
@@ -0,0 +1,54 @@
+// Copyright Epic Games, Inc. All Rights Reserved.
+
+#pragma once
+
+#include "SmartObjectConfig.h"
+#include "MassSmartObjectBehaviorConfig.generated.h"
+
+class UEntitySubsystem;
+class USmartObjectSubsystem;
+struct FLWEntity;
+struct FLWComponentSystemExecutionContext;
+struct FDataFragment_Transform;
+struct FDataFragment_SmartObjectUser;
+
+/**
+ * Struct to pass around the required set of information to activate a mass behavior configuration on a given entity.
+ */
+struct MASSSMARTOBJECTS_API FMassBehaviorEntityContext
+{
+	FMassBehaviorEntityContext() = delete;
+
+	FMassBehaviorEntityContext(const FLWEntity& InEntity, const FDataFragment_Transform& InTransformFragment, FDataFragment_SmartObjectUser& InSOUser, USmartObjectSubsystem& InSubsystem)
+		: Entity(InEntity), TransformFragment(InTransformFragment), SOUser(InSOUser), Subsystem(InSubsystem)
+	{}
+
+	const FLWEntity& Entity;
+	const FDataFragment_Transform& TransformFragment;
+	FDataFragment_SmartObjectUser& SOUser;
+	USmartObjectSubsystem& Subsystem;
+};
+
+/**
+ * Base class for MassAIBehavior configurations. This is the type of configuration that LW Entities queries will look for.
+ * Configuration subclass can parameterized its associated behavior by overriding method Activate.
+ */
+UCLASS(EditInlineNew)
+class MASSSMARTOBJECTS_API USmartObjectMassBehaviorConfig : public USmartObjectBehaviorConfigBase
+{
+	GENERATED_BODY()
+
+public:
+	/**
+	 * This virtual method allows subclasses to configure the LW Entity based on their
+	 * parameters (e.g. Add new fragments)
+	 */
+	virtual void Activate(UEntitySubsystem& EntitySubsystem, FLWComponentSystemExecutionContext& Context, const FMassBehaviorEntityContext& EntityContext) const;
+
+	/**
+	 * Indicates the amount of time the LW entity
+	 * will execute its behavior when reaching the smart object.
+	 */
+	UPROPERTY(EditDefaultsOnly, Category = SmartObject)
+	float UseTime;
+};
diff --git a/Engine/Plugins/Runtime/MassGameplay/Source/MassSmartObjects/Public/MassSmartObjectHandler.h b/Engine/Plugins/Runtime/MassGameplay/Source/MassSmartObjects/Public/MassSmartObjectHandler.h
new file mode 100644
--- /dev/null
+++ b/Engine/Plugins/Runtime/MassGameplay/Source/MassSmartObjects/Public/MassSmartObjectHandler.h
@@ -0,0 +1,124 @@
+// Copyright Epic Games, Inc. All Rights Reserved.
+
+#pragma once
+
+#include "MassSmartObjectRequest.h"
+#include "MassSmartObjectTypes.h"
+
+class UEntitySubsystem;
+class USmartObjectSubsystem;
+struct FLWComponentSystemExecutionContext;
+struct FLWEntity;
+struct FDataFragment_SmartObjectUser;
+struct FDataFragment_Transform;
+struct FSmartObjectClaimHandle;
+struct FSmartObjectID;
+struct FZoneGraphCompactLaneLocation;
+enum class ESmartObjectSlotState : uint8;
+
+/**
+ * Mediator struct that encapsulates communication between SmartObjectSubsystem and Mass.
+ * This object is meant to be created and used in method scope to guarantee subsystems validity.
+ */
+struct MASSSMARTOBJECTS_API FMassSmartObjectHandler
+{
+	/**
+	 * FMassSmartObjectHandler constructor
+	 * @param InEntitySubsystem is the entity subsystem that the smart object should belong to
+	 * @param InExecutionContext is the current execution context of the entity subsystem
+	 * @param InSmartObjectSubsystem is the smart object subsystem
+	 */
+	FMassSmartObjectHandler(UEntitySubsystem& InEntitySubsystem, FLWComponentSystemExecutionContext& InExecutionContext, USmartObjectSubsystem& InSmartObjectSubsystem)
+		: EntitySubsystem(InEntitySubsystem)
+		, ExecutionContext(InExecutionContext)
+		, SmartObjectSubsystem(InSmartObjectSubsystem)
+	{
+	}
+
+	/**
+	 * Creates an async request to build a list of compatible smart objects
+	 * around the provided location. The caller must poll using the request id
+	 * to know when the reservation can be done.
+	 * @param RequestingEntity Entity requesting the candidates list
+	 * @param Location The center of the query
+	 * @return Request identifier that can be used to try claiming a result once available
+	 */
+	FMassSmartObjectRequestID FindCandidatesAsync(const FLWEntity RequestingEntity, const FVector& Location) const;
+
+	/**
+	 * Creates an async request to build a list of compatible smart objects
+	 * around the provided lane location. The caller must poll using the request id
+	 * to know when the reservation can be done.
+	 * @param RequestingEntity Entity requesting the candidates list
+	 * @param LaneLocation The lane location as reference for the query
+	 * @return Request identifier that can be used to try claiming a result once available
+	 */
+	FMassSmartObjectRequestID FindCandidatesAsync(const FLWEntity RequestingEntity, const FZoneGraphCompactLaneLocation& LaneLocation) const;
+
+	/**
+	 * Provides the result of a previously created request from FindCandidatesAsync to indicate if it has been processed
+	 * and the results can be used by ClaimCandidate.
+	 * @param RequestID A valid request identifier (method will ensure otherwise)
+	 * @return The current request's result
+	 */
+	FMassSmartObjectRequestResult GetRequestResult(const FMassSmartObjectRequestID& RequestID) const;
+
+	/**
+	 * Deletes the request associated to the specified identifier
+	 * @param RequestID A valid request identifier (method will ensure otherwise)
+	 */
+	void RemoveRequest(const FMassSmartObjectRequestID& RequestID) const;
+
+	/**
+	 * Claims the first available smart object from the results of the find request. This method
+	 * can be called without calling IsRequestProcessed first but it will fail and return
+	 * Failed_UnprocessedRequest. It can then be called again on next frame until it succeeds or
+	 * returns a different error.
+	 * @param Entity LW Entity associated to the user fragment
+	 * @param User Fragment of the user claiming
+	 * @param RequestID A valid request identifier (method will ensure otherwise)
+	 * @return Whether the slot has been successfully claimed or not
+	 */
+	EMassSmartObjectClaimResult ClaimCandidate(const FLWEntity Entity, FDataFragment_SmartObjectUser& User, const FMassSmartObjectRequestID& RequestID) const;
+
+	/**
+	 * Claims the first available smart object from the provided results.
+	 * @param Entity LW Entity associated to the user fragment
+	 * @param User Fragment of the user claiming
+	 * @param SearchRequestResult Results of completed search request
+	 * @return Whether the slot has been successfully claimed or not
+	 */
+	EMassSmartObjectClaimResult ClaimCandidate(const FLWEntity Entity, FDataFragment_SmartObjectUser& User, const FMassSmartObjectRequestResult& SearchRequestResult) const;
+
+	/**
+	 * Claims the first available slot holding any type of USmartObjectMassBehaviorConfig in the smart object
+	 * associated to the provided identifier.
+	 * @param Entity LW Entity associated to the user fragment
+	 * @param User Fragment of the user claiming
+	 * @param ObjectID A valid smart object identifier (method will ensure otherwise)
+	 * @return Whether the slot has been successfully claimed or not
+	 */
+	bool ClaimSmartObject(const FLWEntity Entity, FDataFragment_SmartObjectUser& User, const FSmartObjectID& ObjectID) const;
+
+	/**
+	 * Activates the mass gameplay behavior associated to the previously claimed smart object.
+	 * @param Entity LW Entity associated to the user fragment
+	 * @param User Fragment of the user claiming
+	 * @param Transform Fragment holding the transform of the user claiming
+	 * @return Whether the slot has been successfully claimed or not
+	 */
+	bool UseSmartObject(const FLWEntity Entity, FDataFragment_SmartObjectUser& User, const FDataFragment_Transform& Transform) const;
+
+	/**
+	 * Releases a claimed/in-use smart object and update user fragment.
+	 * @param Entity LW Entity associated to the user fragment
+	 * @param User Fragment of the user claiming
+	 * @param Status The new status for in-progress interaction
+	 */
+	void ReleaseSmartObject(const FLWEntity Entity, FDataFragment_SmartObjectUser& User, const EMassSmartObjectInteractionStatus Status = EMassSmartObjectInteractionStatus::Unset) const;
+
+private:
+	UEntitySubsystem& EntitySubsystem;
+	FLWComponentSystemExecutionContext& ExecutionContext;
+	USmartObjectSubsystem& SmartObjectSubsystem;
+};
diff --git a/Engine/Plugins/Runtime/MassGameplay/Source/MassSmartObjects/Public/MassSmartObjectProcessor.h b/Engine/Plugins/Runtime/MassGameplay/Source/MassSmartObjects/Public/MassSmartObjectProcessor.h
new file mode 100644
--- /dev/null
+++ b/Engine/Plugins/Runtime/MassGameplay/Source/MassSmartObjects/Public/MassSmartObjectProcessor.h
@@ -0,0 +1,119 @@
+// Copyright Epic Games, Inc. All Rights Reserved.
+
+#pragma once
+
+#include "MassProcessor.h"
+#include "SmartObjectRuntime.h"
+#include "MassSmartObjectRequest.h"
+#include "MassSimulationLOD.h"
+#include "MassSmartObjectTypes.h"
+#include "MassSmartObjectProcessor.generated.h"
+
+class USmartObjectSubsystem;
+class UMassSignalSubsystem;
+class UZoneGraphSubsystem;
+class UZoneGraphAnnotationSubsystem;
+
+USTRUCT()
+struct MASSSMARTOBJECTS_API FDataFragment_SmartObjectUser : public FLWComponentData
+{
+	GENERATED_BODY()
+
+public:
+	FVector GetTargetLocation() const { return TargetLocation; }
+
+	const FSmartObjectClaimHandle& GetClaimHandle() const { return ClaimHandle; }
+
+	float GetCooldown() const { return Cooldown; }
+	void SetCooldown(const float InCooldown) { Cooldown = InCooldown; }
+
+	float GetUseTime() const { return UseTime; }
+	void SetUseTime(const float InUseTime) { UseTime = InUseTime; }
+
+	UPROPERTY(Transient)
+	FSmartObjectClaimHandle ClaimHandle;
+
+	UPROPERTY(Transient)
+	FVector TargetLocation = FVector::ZeroVector;
+
+	UPROPERTY(Transient)
+	EMassSmartObjectInteractionStatus InteractionStatus = EMassSmartObjectInteractionStatus::Unset;
+
+private:
+	UPROPERTY(Transient)
+	float Cooldown = 0.f;
+
+	UPROPERTY(Transient)
+	float UseTime = 0.f; 
+};
+
+/**
+ * Base class for smart object processors that takes care of caching the smart object manager.
+ */
+UCLASS(Abstract)
+class MASSSMARTOBJECTS_API UMassProcessor_SmartObjectBase : public UPipeProcessor
+{
+	GENERATED_BODY()
+
+protected:
+	virtual void Initialize(UObject& Owner) override;
+
+	UPROPERTY(Transient)
+	USmartObjectSubsystem* SmartObjectSubsystem;
+
+	UPROPERTY(Transient)
+	UMassSignalSubsystem* SignalSubsystem;
+};
+
+/** Processor that builds a list of candidates objects for each users. */
+UCLASS()
+class MASSSMARTOBJECTS_API UMassProcessor_SmartObjectCandidatesFinder : public UMassProcessor_SmartObjectBase
+{
+	GENERATED_BODY()
+
+public:
+	UMassProcessor_SmartObjectCandidatesFinder();
+
+protected:
+	virtual void Initialize(UObject& Owner) override;
+	virtual void ConfigureQueries() override;
+	virtual void Execute(UEntitySubsystem& EntitySubsystem, FLWComponentSystemExecutionContext& Context) override;
+
+	/** Extents used to perform the spatial query in the octree for world location queries. */
+	UPROPERTY(EditDefaultsOnly, Category = SmartObject, config)
+	float SearchExtents = 5000.f;
+
+	UPROPERTY(Transient)
+	UZoneGraphSubsystem* ZoneGraphSubsystem;
+
+	UPROPERTY(Transient)
+	UZoneGraphAnnotationSubsystem* AnnotationSubsystem;
+
+	/** Query to fetch and process requests to find smart objects using spacial query around a given world location. */
+	FLWComponentQuery WorldRequestQuery;
+
+	/** Query to fetch and process requests to find smart objects on zone graph lanes. */
+	FLWComponentQuery LaneRequestQuery;
+};
+
+USTRUCT()
+struct MASSSMARTOBJECTS_API FMassSmartObjectTimedBehaviorTag : public FComponentTag
+{
+	GENERATED_BODY()
+};
+
+
+/** Processor for time based user's behavior that waits x seconds then releases its claim on the object */
+UCLASS()
+class MASSSMARTOBJECTS_API UMassProcessor_SmartObjectTimedBehavior : public UMassProcessor_SmartObjectBase
+{
+	GENERATED_BODY()
+public:
+	UMassProcessor_SmartObjectTimedBehavior();
+
+protected:
+	virtual void Execute(UEntitySubsystem& EntitySubsystem, FLWComponentSystemExecutionContext& Context) override;
+	virtual void ConfigureQueries() override;
+
+	FLWComponentQuery EntityQuery;
+};
\ No newline at end of file
diff --git a/Engine/Plugins/Runtime/MassGameplay/Source/MassSmartObjects/Public/MassSmartObjectRequest.h b/Engine/Plugins/Runtime/MassGameplay/Source/MassSmartObjects/Public/MassSmartObjectRequest.h
new file mode 100644
--- /dev/null
+++ b/Engine/Plugins/Runtime/MassGameplay/Source/MassSmartObjects/Public/MassSmartObjectRequest.h
@@ -0,0 +1,99 @@
+// Copyright Epic Games, Inc. All Rights Reserved.
+
+#pragma once
+
+#include "SmartObjectTypes.h"
+#include "LWComponentTypes.h"
+#include "ZoneGraphTypes.h"
+#include "Containers/StaticArray.h"
+#include "MassSmartObjectRequest.generated.h"
+
+/**
+ * Structure that represents a potential smart object for a LW Entity during the search
+ */
+USTRUCT()
+struct MASSSMARTOBJECTS_API FSmartObjectCandidate
+{
+	GENERATED_BODY()
+
+	FSmartObjectCandidate() = default;
+	FSmartObjectCandidate(const FSmartObjectID InID, const float InCost) : ID(InID), Cost(InCost) {}
+
+	UPROPERTY(Transient)
+	FSmartObjectID ID;
+
+	UPROPERTY(Transient)
+	float Cost = 0.f;
+};
+
+/**
+ * Struct that holds status and results of a candidate finder request
+ */
+USTRUCT()
+struct MASSSMARTOBJECTS_API FMassSmartObjectRequestResult
+{
+	GENERATED_BODY()
+
+	static constexpr uint32 MaxNumCandidates = 4;
+	TStaticArray<FSmartObjectCandidate, MaxNumCandidates> Candidates;
+
+	UPROPERTY(Transient)
+	uint8 NumCandidates = 0;
+
+	UPROPERTY(Transient)
+	bool bProcessed = false;
+};
+
+/**
+ * Fragment that holds the result of a request to find candidates.
+ */
+USTRUCT()
+struct MASSSMARTOBJECTS_API FMassSmartObjectRequestResultFragment : public FLWComponentData
+{
+	GENERATED_BODY()
+
+	UPROPERTY(Transient)
+	FMassSmartObjectRequestResult Result;
+};
+
+/**
+ * Fragment used to build a list potential smart objects to use. Once added to an entity
+ * this will be processed by the candidates finder processor to fill a SmartObjectCandidates
+ * fragment that could then be processed by the reservation processor
+ */
+USTRUCT()
+struct MASSSMARTOBJECTS_API FMassSmartObjectWorldLocationRequestFragment : public FLWComponentData
+{
+	GENERATED_BODY()
+
+	UPROPERTY(Transient)
+	FVector SearchOrigin = FVector::ZeroVector;
+
+	UPROPERTY(Transient)
+	FLWEntity RequestingEntity;
+};
+
+/**
+ * Fragment used to build a list potential smart objects to use. Once added to an entity
+ * this will be processed by the candidates finder processor to fill a SmartObjectCandidates
+ * fragment that could then be processed by the reservation processor
+ */
+USTRUCT()
+struct MASSSMARTOBJECTS_API FMassSmartObjectLaneLocationRequestFragment : public FLWComponentData
+{
+	GENERATED_BODY()
+
+	FZoneGraphCompactLaneLocation CompactLaneLocation;
+
+	UPROPERTY(Transient)
+	FLWEntity RequestingEntity;
+};
+
+/**
+ * Special tag to mark processed requests
+ */
+USTRUCT()
+struct MASSSMARTOBJECTS_API FMassSmartObjectCompletedRequestTag : public FComponentTag
+{
+	GENERATED_BODY()
+};
diff --git a/Engine/Plugins/Runtime/MassGameplay/Source/MassSmartObjects/Public/MassSmartObjectTypes.h b/Engine/Plugins/Runtime/MassGameplay/Source/MassSmartObjects/Public/MassSmartObjectTypes.h
new file mode 100644
--- /dev/null
+++ b/Engine/Plugins/Runtime/MassGameplay/Source/MassSmartObjects/Public/MassSmartObjectTypes.h
@@ -0,0 +1,54 @@
+// Copyright Epic Games, Inc. All Rights Reserved.
+
+#pragma once
+
+#include "LWComponentTypes.h"
+#include "MassSmartObjectTypes.generated.h"
+
+namespace UE::Mass::Signals
+{
+	const FName SmartObjectRequestCandidates = FName(TEXT("SmartObjectRequestCandidates"));
+	const FName SmartObjectCandidatesReady = FName(TEXT("SmartObjectCandidatesReady"));
+	const FName SmartObjectInteractionDone = FName(TEXT("SmartObjectInteractionDone"));
+}
+
+/**
+ * Identifier associated to a request for smart object candidates. We use a 1:1 match
+ * with an LWEntity since all requests are batched together using the EntitySubsystem.
+ */
+USTRUCT()
+struct MASSSMARTOBJECTS_API FMassSmartObjectRequestID
+{
+	GENERATED_BODY()
+
+	FMassSmartObjectRequestID() = default;
+	FMassSmartObjectRequestID(const FLWEntity& InEntity) : Entity(InEntity) {}
+	
+	bool IsSet() const { return Entity.IsSet(); }
+	void Reset() { Entity.Reset(); }
+	
+	explicit operator FLWEntity() const { return Entity; }
+
+private:
+	UPROPERTY(Transient)
+	FLWEntity Entity;
+};
+
+UENUM()
+enum class EMassSmartObjectClaimResult: uint8
+{
+	Unset,
+	Failed_InvalidRequest,
+	Failed_UnprocessedRequest,
+	Failed_NoAvailableCandidate,
+	Succeeded
+};
+
+UENUM()
+enum class EMassSmartObjectInteractionStatus: uint8
+{
+	Unset,
+    InProgress,
+    Completed,
+    Aborted
+};
\ No newline at end of file
diff --git a/Engine/Plugins/Runtime/MassGameplay/Source/MassSmartObjects/Public/SmartObjectZoneAnnotations.h b/Engine/Plugins/Runtime/MassGameplay/Source/MassSmartObjects/Public/SmartObjectZoneAnnotations.h
new file mode 100644
--- /dev/null
+++ b/Engine/Plugins/Runtime/MassGameplay/Source/MassSmartObjects/Public/SmartObjectZoneAnnotations.h
@@ -0,0 +1,117 @@
+// Copyright Epic Games, Inc. All Rights Reserved.
+
+#pragma once
+
+#include "SmartObjectTypes.h"
+#include "ZoneGraphAnnotationComponent.h"
+#include "ZoneGraphTypes.h"
+#include "SmartObjectZoneAnnotations.generated.h"
+
+class AZoneGraphData;
+class USmartObjectSubsystem;
+
+/** Minimal amount of data to search and compare lane locations within the same graph. */
+USTRUCT()
+struct FSmartObjectLaneLocation
+{
+	GENERATED_BODY()
+
+	UPROPERTY()
+	int32 LaneIndex = INDEX_NONE;
+
+	UPROPERTY()
+	float DistanceAlongLane = 0.0f;
+};
+
+/** Simple struct to wrap the container to be able to use in a TMap */
+USTRUCT()
+struct FSmartObjectList
+{
+	GENERATED_BODY()
+
+	UPROPERTY(VisibleAnywhere, Category = SmartObject)
+	TArray<FSmartObjectID> SmartObjects;
+};
+
+/** Per ZoneGraphData smart object look up data. */
+USTRUCT()
+struct FSmartObjectAnnotationData
+{
+	GENERATED_BODY()
+
+	/** @return True if this entry is valid (associated to a valid zone graph data), false otherwise. */
+	bool IsValid() const { return DataHandle.IsValid(); }
+
+	/** Reset all internal data. */
+	void Reset()
+	{
+		DataHandle = {};
+		AffectedLanes.Reset();
+		ObjectToEntryPointLookup.Reset();
+		LaneToSmartObjectsLookup.Reset();
+	}
+
+	/** Handle of the ZoneGraphData that this smart object annotation data is associated to */
+	UPROPERTY(VisibleAnywhere, Category = SmartObject)
+	FZoneGraphDataHandle DataHandle;
+
+	UPROPERTY(VisibleAnywhere, Category = SmartObject)
+	TArray<int32> AffectedLanes;
+
+	UPROPERTY(VisibleAnywhere, Category = SmartObject)
+	TMap<FSmartObjectID, FSmartObjectLaneLocation> ObjectToEntryPointLookup;
+
+	UPROPERTY(VisibleAnywhere, Category = SmartObject)
+	TMap<int32, FSmartObjectList> LaneToSmartObjectsLookup;
+
+	bool bInitialTaggingCompleted = false;
+};
+
+/**
+ * ZoneGraph annotations for smart objects
+ */
+UCLASS(ClassGroup = AI, BlueprintType, meta = (BlueprintSpawnableComponent))
+class MASSSMARTOBJECTS_API USmartObjectZoneAnnotations : public UZoneGraphAnnotationComponent
+{
+	GENERATED_BODY()
+
+public:
+	const FSmartObjectAnnotationData* GetAnnotationData(FZoneGraphDataHandle DataHandle) const;
+
+protected:
+	virtual void PostSubsystemsInitialized() override;
+	virtual FZoneGraphTagMask GetAnnotationTags() const override;
+	virtual void TickAnnotation(const float DeltaTime, FZoneGraphAnnotationTagContainer& BehaviorTagContainer) override;
+
+	virtual void PostZoneGraphDataAdded(const AZoneGraphData& ZoneGraphData) override;
+	virtual void PreZoneGraphDataRemoved(const AZoneGraphData& ZoneGraphData) override;
+
+	/** Filter specifying which lanes the behavior is applied to. */
+	UPROPERTY(EditAnywhere, Category = SmartObject)
+	FZoneGraphTagFilter AffectedLaneTags;
+
+	/** Entry points graph for each ZoneGraphData. */
+	UPROPERTY(VisibleAnywhere, Category = SmartObject)
+	TArray<FSmartObjectAnnotationData> SmartObjectAnnotationDataArray;
+
+	/** Tag to mark the lanes that offers smart objects. */
+	UPROPERTY(VisibleAnywhere, Category = SmartObject)
+	FZoneGraphTag BehaviorTag;
+
+#if WITH_EDITOR
+	virtual void PostEditChangeChainProperty(FPropertyChangedChainEvent& PropertyChangedEvent) override;
+	virtual void OnUnregister() override;
+
+	void RebuildForSingleGraph(FSmartObjectAnnotationData& Data, const FZoneGraphStorage& Storage);
+	void RebuildForAllGraphs();
+
+	FDelegateHandle OnAnnotationSettingsChangedHandle;
+	FDelegateHandle OnGraphDataChangedHandle;
+#endif // WITH_EDITOR
+
+#if WITH_EDITORONLY_DATA
+	/** Cached SmartObjectSubsystem */
+	UPROPERTY(Transient)
+	USmartObjectSubsystem* SmartObjectSubsystem = nullptr;
+#endif // WITH_EDITORONLY_DATA
+};
