AggregateTicking and Pipe merged into Engine/Plugins/Runtime/MassEntity

diff --git a/Engine/Plugins/Runtime/MassGameplay/Source/MassSpawner/Private/MassAssortedFragmentsTrait.cpp b/Engine/Plugins/Runtime/MassGameplay/Source/MassSpawner/Private/MassAssortedFragmentsTrait.cpp
new file mode 100644
--- /dev/null
+++ b/Engine/Plugins/Runtime/MassGameplay/Source/MassSpawner/Private/MassAssortedFragmentsTrait.cpp
@@ -0,0 +1,29 @@
+// Copyright Epic Games, Inc. All Rights Reserved.
+
+#include "MassAssortedFragmentsTrait.h"
+#include "MassEntityTemplateRegistry.h"
+
+
+void UMassAssortedFragmentsTrait::BuildTemplate(FMassEntityTemplateBuildContext& BuildContext, UWorld& World) const
+{
+	for (const FInstancedStruct& Fragment : Fragments)
+	{
+		if (Fragment.IsValid())
+		{
+			const UScriptStruct* Type = Fragment.GetScriptStruct();
+			CA_ASSUME(Type);
+			if (Type->IsChildOf(FLWComponentData::StaticStruct()))
+			{
+				BuildContext.AddFragmentWithDefaultInitializer(Fragment);
+			}
+			else if (Type->IsChildOf(FComponentTag::StaticStruct()))
+			{
+				BuildContext.AddTag(*Type);
+			}
+			else
+			{
+				UE_LOG(LogMass, Error, TEXT("Invalid fragment type found: '%s', nor a data fragment nor a tag"), *GetPathNameSafe(Type));
+			}
+		}
+	}
+}
\ No newline at end of file
diff --git a/Engine/Plugins/Runtime/MassGameplay/Source/MassSpawner/Private/MassEntityConfigAsset.cpp b/Engine/Plugins/Runtime/MassGameplay/Source/MassSpawner/Private/MassEntityConfigAsset.cpp
new file mode 100644
--- /dev/null
+++ b/Engine/Plugins/Runtime/MassGameplay/Source/MassSpawner/Private/MassEntityConfigAsset.cpp
@@ -0,0 +1,188 @@
+// Copyright Epic Games, Inc. All Rights Reserved.
+
+#include "MassEntityConfigAsset.h"
+#include "MassEntityTraitBase.h"
+#include "MassSpawnerTypes.h"
+#include "MassSpawnerSubsystem.h"
+#include "MassEntityTemplateRegistry.h"
+#include "Serialization/ArchiveObjectCrc32.h"
+#include "VisualLogger/VisualLogger.h"
+#include "Engine/World.h"
+
+namespace UE::MassSpawner
+{
+	uint32 HashTraits(const TArray<UMassEntityTraitBase*>& CombinedTraits)
+	{
+		class FArchiveObjectCRC32AgentConfig : public FArchiveObjectCrc32
+		{
+		public:
+			virtual bool ShouldSkipProperty(const FProperty* InProperty) const override
+			{
+				check(InProperty);
+				return FArchiveObjectCrc32::ShouldSkipProperty(InProperty) || InProperty->HasAllPropertyFlags(CPF_Transient);
+			}
+		};
+
+		uint32 CRC = 0;
+		for (UMassEntityTraitBase* Trait : CombinedTraits)
+		{
+			FArchiveObjectCRC32AgentConfig Archive;
+			CRC = Archive.Crc32(Trait, CRC);
+		}
+		return CRC;
+	}
+} // UE::MassSpawner
+
+FMassEntityConfig::FMassEntityConfig(UMassEntityConfigAsset& InParent)
+	: Parent(&InParent)
+{
+
+}
+
+const FMassEntityTemplate* FMassEntityConfig::GetOrCreateEntityTemplate(AActor& OwnerActor, const UObject& ConfigOwner) const
+{
+	uint32 Hash;
+	FMassEntityTemplateID TemplateID;
+	TArray<UMassEntityTraitBase*> CombinedTraits;
+	if (const FMassEntityTemplate* ExistingTemplate = GetEntityTemplateInternal(OwnerActor, ConfigOwner, Hash, TemplateID, CombinedTraits))
+	{
+		return ExistingTemplate;
+	}
+
+	UWorld* World = OwnerActor.GetWorld();
+	UMassSpawnerSubsystem* SpawnerSystem = UWorld::GetSubsystem<UMassSpawnerSubsystem>(World);
+	check(SpawnerSystem);
+	UMassEntityTemplateRegistry& TemplateRegistry = SpawnerSystem->GetTemplateRegistryInstance();
+
+	// Build new template
+	// TODO: Add methods to FMassEntityTemplateBuildContext to indicate dependency vs setup.
+	// Dependency should add a fragment with default values (which later can be overridden),
+	// while setup would override values and should be run just once.
+	FMassEntityTemplate& Template = TemplateRegistry.CreateTemplate(Hash, TemplateID);
+	FMassEntityTemplateBuildContext BuildContext(Template);
+
+	for (const UMassEntityTraitBase* Trait : CombinedTraits)
+	{
+		check(Trait);
+		Trait->BuildTemplate(BuildContext, *World);
+	}
+
+	for (const UMassEntityTraitBase* Trait : CombinedTraits)
+	{
+		check(Trait);
+		Trait->ValidateTemplate(BuildContext, *World);
+	}
+
+	if (ensureMsgf(Template.GetFragmentsNum(), TEXT("Need at least one fragment to create an Archetype")))
+	{
+		Template.SetUpProcessors(BuildContext.Handlers, OwnerActor);
+		TemplateRegistry.InitializeEntityTemplate(Template);
+	}
+
+	return &Template;
+}
+
+void FMassEntityConfig::DestroyEntityTemplate(const AActor& OwnerActor, const UObject& ConfigOwner) const
+{
+	uint32 Hash;
+	FMassEntityTemplateID TemplateID;
+	TArray<UMassEntityTraitBase*> CombinedTraits;
+	const FMassEntityTemplate* Template = GetEntityTemplateInternal(OwnerActor, ConfigOwner, Hash, TemplateID, CombinedTraits);
+	if (Template == nullptr)
+	{
+		return;
+	}
+
+	const UWorld* World = OwnerActor.GetWorld();
+	UMassSpawnerSubsystem* SpawnerSystem = UWorld::GetSubsystem<UMassSpawnerSubsystem>(World);
+	check(SpawnerSystem);
+	UMassEntityTemplateRegistry& TemplateRegistry = SpawnerSystem->GetTemplateRegistryInstance();
+
+
+	for (const UMassEntityTraitBase* Trait : CombinedTraits)
+	{
+		check(Trait);
+		Trait->DestroyTemplate();
+	}
+
+	// TODO - The templates are not being torn down completely, resulting in traits that leave data in various subsystems. (Representation system)
+	
+	TemplateRegistry.DestroyTemplate(Hash, TemplateID);
+}
+
+const FMassEntityTemplate& FMassEntityConfig::GetEntityTemplateChecked(const AActor& OwnerActor, const UObject& ConfigOwner) const
+{
+	uint32 Hash;
+	FMassEntityTemplateID TemplateID;
+	TArray<UMassEntityTraitBase*> CombinedTraits;
+	const FMassEntityTemplate* ExistingTemplate = GetEntityTemplateInternal(OwnerActor, ConfigOwner, Hash, TemplateID, CombinedTraits);
+	check(ExistingTemplate);
+	return *ExistingTemplate;
+}
+
+const FMassEntityTemplate* FMassEntityConfig::GetEntityTemplateInternal(const AActor& OwnerActor, const UObject& ConfigOwner, uint32& HashOut, FMassEntityTemplateID& TemplateIDOut, TArray<UMassEntityTraitBase*>& CombinedTraitsOut) const
+{
+	const UWorld* World = OwnerActor.GetWorld();
+	UMassSpawnerSubsystem* SpawnerSystem = UWorld::GetSubsystem<UMassSpawnerSubsystem>(World);
+	check(SpawnerSystem);
+	const UMassEntityTemplateRegistry& TemplateRegistry = SpawnerSystem->GetTemplateRegistryInstance();
+
+	// Combine all the features into one array
+	// @todo this is an inefficient way assuming given template is expected to have already been created. Figure out a way to cache it.
+	TArray<const UObject*> Visited;
+	CombinedTraitsOut.Reset();
+	Visited.Add(&ConfigOwner);
+	GetCombinedTraits(CombinedTraitsOut, Visited, ConfigOwner);
+
+	// Return existing template if found.
+	// TODO: cache the hash.
+	HashOut = UE::MassSpawner::HashTraits(CombinedTraitsOut);
+	TemplateIDOut = FMassEntityTemplateID(HashOut, EMassEntityTemplateIDType::ScriptStruct); // TODO: add proper ID type
+	return TemplateRegistry.FindTemplateFromTemplateID(TemplateIDOut);
+}
+
+void FMassEntityConfig::GetCombinedTraits(TArray<UMassEntityTraitBase*>& OutTraits, TArray<const UObject*>& Visited, const UObject& ConfigOwner) const
+{
+	if (Parent)
+	{
+		if (Visited.IndexOfByKey(Parent) != INDEX_NONE)
+		{
+			// Infinite loop detected.
+			FString Path;
+			for (const UObject* Object : Visited)
+			{
+				Path += Object->GetName();
+				Path += TEXT("/");
+			}
+			UE_VLOG(&ConfigOwner, LogMassSpawner, Error, TEXT("%s: Encountered %s as parent second time (Infinite loop). %s"), *GetNameSafe(&ConfigOwner), *GetNameSafe(Parent), *Path);
+		}
+		else
+		{
+			Visited.Add(Parent);
+			Parent->GetConfig().GetCombinedTraits(OutTraits, Visited, ConfigOwner);
+		}
+	}
+
+	for (UMassEntityTraitBase* Trait : Traits)
+	{
+		if (!Trait)
+		{
+			continue;
+		}
+		// Allow only one feature per type. This is also used to allow child configs override parent features.
+		const int32 Index = OutTraits.IndexOfByPredicate([Trait](const UMassEntityTraitBase* ExistingFeature) -> bool { return Trait->GetClass() == ExistingFeature->GetClass(); });
+		if (Index != INDEX_NONE)
+		{
+			OutTraits[Index] = Trait;
+		}
+		else
+		{
+			OutTraits.Add(Trait);
+		}
+	}
+}
+
+void FMassEntityConfig::AddTrait(UMassEntityTraitBase& Trait)
+{
+	Traits.Add(&Trait);
+}
diff --git a/Engine/Plugins/Runtime/MassGameplay/Source/MassSpawner/Private/MassEntityEQSSpawnPointsGenerator.cpp b/Engine/Plugins/Runtime/MassGameplay/Source/MassSpawner/Private/MassEntityEQSSpawnPointsGenerator.cpp
new file mode 100644
--- /dev/null
+++ b/Engine/Plugins/Runtime/MassGameplay/Source/MassSpawner/Private/MassEntityEQSSpawnPointsGenerator.cpp
@@ -0,0 +1,52 @@
+// Copyright Epic Games, Inc. All Rights Reserved.
+
+#include "MassEntityEQSSpawnPointsGenerator.h"
+#include "MassSpawnerTypes.h"
+#include "VisualLogger/VisualLogger.h"
+
+void UMassEntityEQSSpawnPointsGenerator::GenerateSpawnPoints(UObject& QueryOwner, int32 Count, FFinishedGeneratingSpawnPointsSignature& FinishedGeneratingSpawnPointsDelegate) const
+{
+	// Need to copy the request as it is called inside a CDO and CDO states cannot be changed.
+	FEQSParametrizedQueryExecutionRequest EQSRequestInstanced = EQSRequest;
+	if (EQSRequestInstanced.IsValid() == false)
+	{
+		EQSRequestInstanced.InitForOwnerAndBlackboard(QueryOwner, /*BBAsset=*/nullptr);
+		if (!ensureMsgf(EQSRequestInstanced.IsValid(), TEXT("Query request initialization can fail due to missing parameters. See the runtime log for details")))
+		{
+			return;
+		}
+	}
+
+	FQueryFinishedSignature Delegate = FQueryFinishedSignature::CreateUObject(this, &UMassEntityEQSSpawnPointsGenerator::OnEQSQueryFinished, Count, FinishedGeneratingSpawnPointsDelegate);
+	EQSRequestInstanced.Execute(QueryOwner, /*BlackboardComponent=*/nullptr, Delegate);
+}
+
+void UMassEntityEQSSpawnPointsGenerator::OnEQSQueryFinished(TSharedPtr<FEnvQueryResult> Result, int32 Count, FFinishedGeneratingSpawnPointsSignature FinishedGeneratingSpawnPointsDelegate) const
+{
+	TArray<FVector> Locations;
+
+	if (Result.IsValid() == false || Result->IsSuccessful() == false)
+	{
+		UE_VLOG_UELOG(this, LogMassSpawner, Error, TEXT("EQS query failed or result is invalid"));
+	}
+	else
+	{
+		Result->GetAllAsLocations(Locations);
+	}
+
+	// Randomize them
+	FRandomStream RandomStream(GFrameNumber);
+	for (int32 I = 0; I < Locations.Num(); ++I)
+	{
+		const int32 J = RandomStream.RandHelper(Locations.Num());
+		Locations.Swap(I, J);
+	}
+
+	// If we generated too many, shrink it.
+	if (Locations.Num() > Count)
+	{
+		Locations.SetNum(Count);
+	}
+
+	FinishedGeneratingSpawnPointsDelegate.Execute(Locations);
+}
\ No newline at end of file
diff --git a/Engine/Plugins/Runtime/MassGameplay/Source/MassSpawner/Private/MassEntityTemplate.cpp b/Engine/Plugins/Runtime/MassGameplay/Source/MassSpawner/Private/MassEntityTemplate.cpp
new file mode 100644
--- /dev/null
+++ b/Engine/Plugins/Runtime/MassGameplay/Source/MassSpawner/Private/MassEntityTemplate.cpp
@@ -0,0 +1,60 @@
+// Copyright Epic Games, Inc. All Rights Reserved.
+
+#include "MassEntityTemplate.h"
+#include "VisualLogger/VisualLoggerTypes.h"
+
+DEFINE_ENUM_TO_STRING(EMassEntityTemplateIDType);
+
+//----------------------------------------------------------------------//
+//  FMassEntityTemplateID
+//----------------------------------------------------------------------//
+FString FMassEntityTemplateID::ToString() const
+{
+	return FString::Printf(TEXT("[%s %d]"), *EnumToString(Type), Hash);
+}
+
+//----------------------------------------------------------------------//
+//  FMassEntityTemplate
+//----------------------------------------------------------------------//
+void FMassEntityTemplate::SetArchetype(const FArchetypeHandle& InArchetype)
+{
+	check(InArchetype.IsValid());
+	Archetype = InArchetype;
+}
+
+void FMassEntityTemplate::SetUpProcessors(const FMassObjectHandlers& ObjectHandlers, UObject& PipelineOwner)
+{
+	// @todo do we want to copy the Processors or just use the CDOs?
+	InitializationPipeline.InitializeFromArray(ObjectHandlers.Initializers, PipelineOwner);
+	DeinitializationPipeline.InitializeFromArray(ObjectHandlers.Deinitializers, PipelineOwner);
+}
+
+FString FMassEntityTemplate::DebugGetDescription(UEntitySubsystem* EntitySubsystem) const
+{ 
+	FStringOutputDevice Ar;
+#if WITH_MASS_DEBUG
+	Ar.SetAutoEmitLineTerminator(true);
+
+	Ar += TEXT("Fragments:\n");
+	FragmentCollection.DebugOutputDescription(Ar);
+
+	Ar += TEXT("\nProcessors:\n");
+	InitializationPipeline.DebugOutputDescription(Ar);
+
+	if (EntitySubsystem)
+	{
+		Ar += TEXT("Archetype detais:\n");
+		Ar += DebugGetArchetypeDescription(*EntitySubsystem);
+	}
+#endif // WITH_MASS_DEBUG
+	return MoveTemp(Ar);
+}
+
+FString FMassEntityTemplate::DebugGetArchetypeDescription(UEntitySubsystem& EntitySubsystem) const
+{
+	FStringOutputDevice OutDescription;
+#if WITH_MASS_DEBUG
+	EntitySubsystem.DebugGetStringDesc(Archetype, OutDescription);
+#endif // WITH_MASS_DEBUG
+	return MoveTemp(OutDescription);
+}
diff --git a/Engine/Plugins/Runtime/MassGameplay/Source/MassSpawner/Private/MassEntityTemplateRegistry.cpp b/Engine/Plugins/Runtime/MassGameplay/Source/MassSpawner/Private/MassEntityTemplateRegistry.cpp
new file mode 100644
--- /dev/null
+++ b/Engine/Plugins/Runtime/MassGameplay/Source/MassSpawner/Private/MassEntityTemplateRegistry.cpp
@@ -0,0 +1,300 @@
+// Copyright Epic Games, Inc. All Rights Reserved.
+
+#include "MassEntityTemplateRegistry.h"
+#include "MassSpawnerTypes.h"
+#include "MassEntitySystem.h"
+#include "Engine/World.h"
+#include "VisualLogger/VisualLogger.h"
+#include "HAL/IConsoleManager.h"
+#include "MassSpawnerSubsystem.h"
+
+namespace FTemplateRegistryHelpers
+{
+	uint32 CalcHash(const FInstancedStruct& StructInstance)
+	{
+		const UScriptStruct* Type = StructInstance.GetScriptStruct();
+		const uint8* Memory = StructInstance.GetMemory();
+		return Type && Memory ? Type->GetStructTypeHash(Memory) : 0;
+	}
+
+	void FragmentInstancesToTypes(TArrayView<const FInstancedStruct> FragmentList, TArray<const UScriptStruct*>& OutFragmentTypes)
+	{
+		for (const FInstancedStruct& Instance : FragmentList)
+		{
+			// at this point FragmentList is assumed to have no duplicates nor nulls
+			OutFragmentTypes.Add(Instance.GetScriptStruct());
+		}
+	}
+
+	void ResetEntityTemplates(const TArray<FString>& Args, UWorld* InWorld)
+	{
+		UMassSpawnerSubsystem* SpawnerSystem = UWorld::GetSubsystem<UMassSpawnerSubsystem>(InWorld);
+		if (ensure(SpawnerSystem))
+		{
+			UMassEntityTemplateRegistry& Registry = SpawnerSystem->GetTemplateRegistryInstance();
+			Registry.DebugReset();
+		}
+	}
+
+	FAutoConsoleCommandWithWorldAndArgs EnableCategoryNameCmd(
+		TEXT("ai.mass.reset_entity_templates"),
+		TEXT("Clears all the runtime information cached by MassEntityTemplateRegistry. Will result in lazily building all entity templates again."),
+		FConsoleCommandWithWorldAndArgsDelegate::CreateStatic(&ResetEntityTemplates)
+	);
+
+}
+
+//----------------------------------------------------------------------//
+// UMassEntityTemplateRegistry 
+//----------------------------------------------------------------------//
+TMap<const UScriptStruct*, UMassEntityTemplateRegistry::FStructToTemplateBuilderDelegate> UMassEntityTemplateRegistry::StructBasedBuilders;
+TMap<const UClass*, UMassEntityTemplateRegistry::FClassToTemplateBuilderDelegate> UMassEntityTemplateRegistry::ClassBasedBuilders;
+
+void UMassEntityTemplateRegistry::BeginDestroy()
+{
+	// force release of memory owned by individual templates (especially the hosted InstancedScriptStructs).
+	TemplateIDToTemplateMap.Reset();
+
+	Super::BeginDestroy();
+}
+
+UWorld* UMassEntityTemplateRegistry::GetWorld() const 
+{
+	const UObject* Outer = GetOuter();
+	return Outer ? Outer->GetWorld() : nullptr;
+}
+
+UMassEntityTemplateRegistry::FStructToTemplateBuilderDelegate& UMassEntityTemplateRegistry::FindOrAdd(const UScriptStruct& DataType)
+{
+	return StructBasedBuilders.FindOrAdd(&DataType);
+}
+
+UMassEntityTemplateRegistry::FClassToTemplateBuilderDelegate& UMassEntityTemplateRegistry::FindOrAdd(const UClass& Class)
+{
+	return ClassBasedBuilders.FindOrAdd(&Class);
+}
+
+const FMassEntityTemplate* UMassEntityTemplateRegistry::FindOrBuildStructTemplate(const FInstancedStruct& StructInstance)
+{
+	// thou shall not call this function on CDO
+	check(HasAnyFlags(RF_ClassDefaultObject) == false);
+
+	const UScriptStruct* Type = StructInstance.GetScriptStruct();
+	check(Type);
+	// 1. Check if we already have the template stored.
+	// 2. If not, 
+	//	a. build it
+	//	b. store it
+
+	const uint32 StructInstanceHash = FTemplateRegistryHelpers::CalcHash(StructInstance);
+	const uint32 HashLookup = HashCombine(GetTypeHash(Type->GetFName()), StructInstanceHash);
+
+	FMassEntityTemplateID* TemplateID = LookupTemplateIDMap.Find(HashLookup);
+
+	if (TemplateID != nullptr)
+	{
+		if (const FMassEntityTemplate* TemplateFound = TemplateIDToTemplateMap.Find(*TemplateID))
+		{
+			return TemplateFound;
+		}
+	}
+
+	// this means we don't have an entry for given struct. Let's see if we know how to make one
+	FMassEntityTemplate* NewTemplate = nullptr;
+	FStructToTemplateBuilderDelegate* Builder = StructBasedBuilders.Find(StructInstance.GetScriptStruct());
+	if (Builder)
+	{
+		if (TemplateID == nullptr)
+		{
+			// TODO consider removing the need for strings here
+			// Use the class name string for the hash here so the hash can be deterministic between client and server
+			const uint32 NameStringHash = GetTypeHash(Type->GetName());
+			const uint32 Hash = HashCombine(NameStringHash, StructInstanceHash);
+
+			TemplateID = &LookupTemplateIDMap.Add(HashLookup, FMassEntityTemplateID(Hash, EMassEntityTemplateIDType::ScriptStruct));
+		}
+
+		NewTemplate = &TemplateIDToTemplateMap.Add(*TemplateID);
+
+		check(NewTemplate);
+		NewTemplate->SetTemplateID(*TemplateID);
+
+		BuildTemplateImpl(*Builder, StructInstance, *NewTemplate);
+	}
+	UE_CVLOG_UELOG(Builder == nullptr, this, LogMassSpawner, Warning, TEXT("Attempting to build a MassAgentTemplate for struct type %s while template builder has not been registered for this type")
+		, *GetNameSafe(Type));
+
+	return NewTemplate;
+}
+
+bool UMassEntityTemplateRegistry::BuildTemplateImpl(const FStructToTemplateBuilderDelegate& Builder, const FInstancedStruct& StructInstance, FMassEntityTemplate& OutTemplate)
+{
+	UWorld* World = GetWorld();
+	FMassEntityTemplateBuildContext Context(OutTemplate);
+	Builder.Execute(World, StructInstance, Context);
+	if (ensure(OutTemplate.GetFragmentsNum())) // need at least one fragment to create an Archetype
+	{
+		OutTemplate.SetUpProcessors(Context.Handlers, *this);
+		InitializeEntityTemplate(OutTemplate);
+
+		UE_VLOG(this, LogMassSpawner, Log, TEXT("Created entity template for %s:\n%s"), *GetNameSafe(StructInstance.GetScriptStruct())
+			, *OutTemplate.DebugGetDescription(UPipeEntitySubsystem::GetCurrent(World)));
+
+		return true;
+	}
+	return false;
+}
+
+const FMassEntityTemplate* UMassEntityTemplateRegistry::FindOrBuildInstanceTemplate(const UObject& Instance)
+{
+	checkf(HasAnyFlags(RF_ClassDefaultObject) == false, TEXT("Thou shall not call %s on UMassEntityTemplateRegistry's CDO"), ANSI_TO_TCHAR(__FUNCTION__));
+	check(Instance.GetClass());
+
+	const FMassEntityTemplateID* const TemplateID = LookupTemplateIDMap.Find(GetTypeHash(Instance.GetClass()->GetFName()));
+
+	if (TemplateID != nullptr)
+	{
+		if (const FMassEntityTemplate* TemplateFound = TemplateIDToTemplateMap.Find(*TemplateID))
+		{
+			return TemplateFound;
+		}
+	}
+
+	// TODO this code needs refactoring as its not accounting for instances needing a different hash to classes
+	// 2. If not: build it & store it
+	return BuildClassTemplate(*Instance.GetClass(), &Instance);
+}
+
+const FMassEntityTemplate& UMassEntityTemplateRegistry::FindInstanceTemplateChecked(const UObject& Instance) const
+{
+	checkf(HasAnyFlags(RF_ClassDefaultObject) == false,  TEXT("Thou shall not call %s on UMassEntityTemplateRegistry's CDO"), ANSI_TO_TCHAR(__FUNCTION__));
+	check(Instance.GetClass());
+
+	const FMassEntityTemplateID& TemplateID = LookupTemplateIDMap.FindChecked(GetTypeHash(Instance.GetClass()->GetFName()));
+
+	return TemplateIDToTemplateMap.FindChecked(TemplateID);
+}
+
+const FMassEntityTemplate* UMassEntityTemplateRegistry::FindOrBuildClassTemplate(const UClass& Class)
+{
+	// thou shall not call this function on UMassEntityTemplateRegistry's CDO
+	check(HasAnyFlags(RF_ClassDefaultObject) == false);
+
+	// 1. Check if we already have the template stored.
+	const FMassEntityTemplateID* const TemplateID = LookupTemplateIDMap.Find(GetTypeHash(Class.GetFName()));
+
+	if (TemplateID != nullptr)
+	{
+		if (const FMassEntityTemplate* TemplateFound = TemplateIDToTemplateMap.Find(*TemplateID))
+		{
+			return TemplateFound;
+		}
+	}
+
+	// 2. If not: build it & store it
+	return BuildClassTemplate(Class);
+}
+
+const FMassEntityTemplate* UMassEntityTemplateRegistry::BuildClassTemplate(const UClass& Class, const UObject* Instance)
+{
+	FMassEntityTemplate* EntityTemplate = nullptr;
+	// this function will be called when there's been no EntityTemplate created for Class/Instance just yet
+	// Let's see if we know how to make one	
+	FClassToTemplateBuilderDelegate* Builder = GetBuilderForClass(Class);
+	if (Builder)
+	{
+		TArray<const UPipeProcessor*> Initializers;
+		
+		// TODO consider removing the need for strings here
+		// Use the class name string for the hash here, so the hash can be deterministic between client and server
+		const uint32 NameStringHash = GetTypeHash(Class.GetName());
+
+		const EMassEntityTemplateIDType TemplateType = (Instance != nullptr && Instance->HasAnyFlags(RF_ClassDefaultObject) == false) ? EMassEntityTemplateIDType::Instance : EMassEntityTemplateIDType::Class;
+		const FMassEntityTemplateID TemplateID = LookupTemplateIDMap.Add(GetTypeHash(Class.GetFName()), FMassEntityTemplateID(NameStringHash, TemplateType));
+
+		EntityTemplate = &TemplateIDToTemplateMap.Add(TemplateID);
+
+		check(EntityTemplate);
+
+		EntityTemplate->SetTemplateID(TemplateID);
+
+		FMassEntityTemplateBuildContext BuildContext(*EntityTemplate);
+		Builder->Execute(GetWorld(), Class, Instance, BuildContext);
+
+		if (ensureMsgf(EntityTemplate->GetFragmentsNum(), TEXT("Need at least one fragment to create an Archetype")))
+		{
+			EntityTemplate->SetUpProcessors(BuildContext.Handlers, *this);
+			InitializeEntityTemplate(*EntityTemplate);
+		}
+	}
+	UE_CVLOG_UELOG(Builder == nullptr, this, LogMassSpawner, Warning, TEXT("Attempting to build a MassAgentTemplate for uclass %s (instance %s) while template builder has not been registered for this type")
+		, *Class.GetName(), *GetNameSafe(Instance));
+
+	return EntityTemplate;
+}
+
+UMassEntityTemplateRegistry::FClassToTemplateBuilderDelegate* UMassEntityTemplateRegistry::GetBuilderForClass(const UClass& Class)
+{
+	FClassToTemplateBuilderDelegate* Builder = nullptr;
+	const UClass* TmpClass = &Class;
+	do
+	{ 
+		Builder = ClassBasedBuilders.Find(TmpClass);
+		TmpClass = TmpClass->GetSuperClass();
+	} while (Builder == nullptr && TmpClass != nullptr);
+	return Builder;
+}
+
+void UMassEntityTemplateRegistry::InitializeEntityTemplate(FMassEntityTemplate& OutTemplate) const
+{
+	// expected to be ensured by the caller
+	check(OutTemplate.GetFragmentsNum());
+
+	UWorld* World = GetWorld();
+	// find or create template
+	UPipeEntitySubsystem* EntitySys = UPipeEntitySubsystem::GetCurrent(World);
+	check(EntitySys);
+
+	FTemplateRegistryHelpers::FragmentInstancesToTypes(OutTemplate.GetFragments(), OutTemplate.GetMutableFragmentTypes());
+	const FArchetypeHandle ArchetypeHandle = EntitySys->CreateArchetype(OutTemplate.GetFragmentTypes(), OutTemplate.GetTags(), FLWChunkComponentBitSet(OutTemplate.GetChunkFragments()));
+	OutTemplate.SetArchetype(ArchetypeHandle);
+	for (const FInstancedStruct& InstancedStruct : OutTemplate.GetChunkFragments())
+	{
+		EntitySys->SetDefaultChunkComponentValue(ArchetypeHandle, InstancedStruct);
+	}
+}
+
+void UMassEntityTemplateRegistry::DebugReset()
+{
+#if WITH_MASS_DEBUG
+	LookupTemplateIDMap.Reset();
+	TemplateIDToTemplateMap.Reset();
+#endif // WITH_MASS_DEBUG
+}
+
+const FMassEntityTemplate* UMassEntityTemplateRegistry::FindTemplateFromTemplateID(FMassEntityTemplateID TemplateID) const 
+{
+	return TemplateIDToTemplateMap.Find(TemplateID);
+}
+
+FMassEntityTemplate* UMassEntityTemplateRegistry::FindMutableTemplateFromTemplateID(FMassEntityTemplateID TemplateID) 
+{
+	return TemplateIDToTemplateMap.Find(TemplateID);
+}
+
+FMassEntityTemplate& UMassEntityTemplateRegistry::CreateTemplate(const uint32 HashLookup, FMassEntityTemplateID TemplateID)
+{
+	checkSlow(!LookupTemplateIDMap.Contains(HashLookup));
+	LookupTemplateIDMap.Add(HashLookup, TemplateID);
+	checkSlow(!TemplateIDToTemplateMap.Contains(TemplateID));
+	FMassEntityTemplate& NewTemplate = TemplateIDToTemplateMap.Add(TemplateID);
+	NewTemplate.SetTemplateID(TemplateID);
+	return NewTemplate;
+}
+
+void UMassEntityTemplateRegistry::DestroyTemplate(const uint32 HashLookup, FMassEntityTemplateID TemplateID)
+{
+	LookupTemplateIDMap.Remove(HashLookup);
+	TemplateIDToTemplateMap.Remove(TemplateID);
+}
+
diff --git a/Engine/Plugins/Runtime/MassGameplay/Source/MassSpawner/Private/MassEntityZoneGraphSpawnPointsGenerator.cpp b/Engine/Plugins/Runtime/MassGameplay/Source/MassSpawner/Private/MassEntityZoneGraphSpawnPointsGenerator.cpp
new file mode 100644
--- /dev/null
+++ b/Engine/Plugins/Runtime/MassGameplay/Source/MassSpawner/Private/MassEntityZoneGraphSpawnPointsGenerator.cpp
@@ -0,0 +1,82 @@
+// Copyright Epic Games, Inc. All Rights Reserved.
+
+#include "MassEntityZoneGraphSpawnPointsGenerator.h"
+#include "ZoneGraphData.h"
+#include "ZoneGraphSubsystem.h"
+#include "ZoneGraphQuery.h"
+#include "MassSpawnerTypes.h"
+#include "VisualLogger/VisualLogger.h"
+#include "Engine/World.h"
+
+void UMassEntityZoneGraphSpawnPointsGenerator::GenerateSpawnPoints(UObject& QueryOwner, int32 Count, FFinishedGeneratingSpawnPointsSignature& FinishedGeneratingSpawnPointsDelegate) const
+{
+	TArray<FVector> Locations;
+	if (const UWorld* World = QueryOwner.GetWorld())
+	{
+		if (const UZoneGraphSubsystem* ZoneGraph = UWorld::GetSubsystem<UZoneGraphSubsystem>(World))
+		{
+			const FRandomStream RandomStream(GFrameNumber);
+			const TConstArrayView<FRegisteredZoneGraphData> RegisteredZoneGraphs = ZoneGraph->GetRegisteredZoneGraphData();
+			for (const FRegisteredZoneGraphData& Registered : RegisteredZoneGraphs)
+			{
+				if (Registered.bInUse && Registered.ZoneGraphData)
+				{
+					GeneratePointsForZoneGraphData(*Registered.ZoneGraphData, Locations, RandomStream);
+				}
+			}
+
+			// Randomize them
+			for (int32 I = 0; I < Locations.Num(); ++I)
+			{
+				const int32 J = RandomStream.RandHelper(Locations.Num());
+				Locations.Swap(I, J);
+			}
+
+			// If we generated too many, shrink it.
+			if (Locations.Num() > Count)
+			{
+				Locations.SetNum(Count);
+			}
+		}
+		else
+		{
+			UE_VLOG_UELOG(&QueryOwner, LogMassSpawner, Error, TEXT("No zone graph found in world"));
+		}
+	}
+	else
+	{
+		UE_VLOG_UELOG(&QueryOwner, LogMassSpawner, Error, TEXT("Unable to retrieve world from QueryOwner"));
+	}
+
+	FinishedGeneratingSpawnPointsDelegate.Execute(Locations);
+}
+
+void UMassEntityZoneGraphSpawnPointsGenerator::GeneratePointsForZoneGraphData(const ::AZoneGraphData& ZoneGraphData, TArray<FVector>& Locations, const FRandomStream& RandomStream) const
+{
+	const FZoneGraphStorage& ZoneGraphStorage = ZoneGraphData.GetStorage();
+
+	// Loop through all lanes
+	for (int32 LaneIndex = 0; LaneIndex < ZoneGraphStorage.Lanes.Num(); ++LaneIndex)
+	{
+		const FZoneLaneData& Lane = ZoneGraphStorage.Lanes[LaneIndex];
+		const float LaneHalfWidth = Lane.Width / 2.0f;
+		if (TagFilter.Pass(Lane.Tags))
+		{
+			float LaneLength = 0.0f;
+			UE::ZoneGraph::Query::GetLaneLength(ZoneGraphStorage, LaneIndex, LaneLength);
+
+			float Distance = RandomStream.FRandRange(MinGap, MaxGap); // ..initially
+			while (Distance <= LaneLength)
+			{
+				// Add location at the center of this space.
+				FZoneGraphLaneLocation LaneLocation;
+				UE::ZoneGraph::Query::CalculateLocationAlongLane(ZoneGraphStorage, LaneIndex, Distance, LaneLocation);
+				const FVector Perp = LaneLocation.Direction ^ LaneLocation.Up;
+				Locations.Add(LaneLocation.Position + Perp * RandomStream.RandRange(-LaneHalfWidth, LaneHalfWidth));
+
+				// Advance ahead past the space we just consumed, plus a random gap.
+				Distance += RandomStream.FRandRange(MinGap, MaxGap);
+			}
+		}
+	}
+}
diff --git a/Engine/Plugins/Runtime/MassGameplay/Source/MassSpawner/Private/MassSpawner.cpp b/Engine/Plugins/Runtime/MassGameplay/Source/MassSpawner/Private/MassSpawner.cpp
new file mode 100644
--- /dev/null
+++ b/Engine/Plugins/Runtime/MassGameplay/Source/MassSpawner/Private/MassSpawner.cpp
@@ -0,0 +1,557 @@
+// Copyright Epic Games, Inc. All Rights Reserved.
+
+#include "MassSpawner.h"
+#include "Engine/World.h"
+#include "UObject/ConstructorHelpers.h"
+#include "Components/SceneComponent.h"
+#include "Components/BillboardComponent.h"
+#include "MassSpawnerTypes.h"
+#include "MassSpawnerSubsystem.h"
+#include "MassSimulationSubsystem.h"
+#include "EnvironmentQuery/EnvQuery.h"
+#include "VisualLogger/VisualLogger.h"
+#include "MassEntityConfigAsset.h"
+#include "EngineUtils.h"
+
+namespace UE::MassSpawner
+{
+	float ScalabilitySpawnDensityMultiplier = 1.f;
+	FAutoConsoleVariableRef CVarScalabilitySpawnDensityMultiplier(TEXT("ai.mass.scalability.SpawnDensityMultiplier"), ScalabilitySpawnDensityMultiplier, TEXT("Spawn Density Multiplier, must be set before Mass Spawn Init"), ECVF_Scalability);
+
+#if WITH_EDITOR
+	static FAutoConsoleCommandWithWorld ForceSpawningCommand(
+	TEXT("ai.mass.ForceSpawn"),
+	TEXT("Command to Force Spawn all mass entities generated by MassSpawners"),
+	FConsoleCommandWithWorldDelegate::CreateLambda([](UWorld* World)
+	{
+		for (TActorIterator<AActor> It(World, AMassSpawner::StaticClass()); It; ++It)
+		{
+			if (AMassSpawner* Spawner = Cast<AMassSpawner>(*It))
+			{
+				Spawner->DoSpawning();
+			}
+		}
+	}));
+
+	static FAutoConsoleCommandWithWorld ForceDespawningCommand(
+	TEXT("ai.mass.ForceDespawn"),
+	TEXT("Command to Force Despawn all mass entities generated by MassSpawners"),
+	FConsoleCommandWithWorldDelegate::CreateLambda([](UWorld* World)
+	{
+		for (TActorIterator<AActor> It(World, AMassSpawner::StaticClass()); It; ++It)
+		{
+			if (AMassSpawner* Spawner = Cast<AMassSpawner>(*It))
+			{
+				Spawner->DoDespawning();
+			}
+		}
+	}));
+
+	static FAutoConsoleCommandWithWorld ResetSpawningCommand(
+	TEXT("ai.mass.ResetSpawning"),
+	TEXT("Command to Force Despawn and Respawn all mass entities generated by MassSpawners"),
+	FConsoleCommandWithWorldDelegate::CreateLambda([](UWorld* World)
+	{
+		for (TActorIterator<AActor> It(World, AMassSpawner::StaticClass()); It; ++It)
+		{
+			if (AMassSpawner* Spawner = Cast<AMassSpawner>(*It))
+			{
+				Spawner->DoDespawning();
+				Spawner->DoSpawning();
+			}
+		}
+	}));
+#endif // WITH_EDITOR
+}
+
+AMassSpawner::AMassSpawner()
+{
+	RootComponent = CreateDefaultSubobject<USceneComponent>(TEXT("SceneComp"));
+	RootComponent->Mobility = EComponentMobility::Static;
+
+#if WITH_EDITOR
+	SpriteComponent = CreateEditorOnlyDefaultSubobject<UBillboardComponent>(TEXT("Sprite"));
+
+	// SpriteComponent can be null for editor builds running "-game|server"
+	if (!IsRunningCommandlet() && SpriteComponent != nullptr)
+	{
+		// Structure to hold one-time initialization
+		struct FConstructorStatics
+		{
+			ConstructorHelpers::FObjectFinderOptional<UTexture2D> IconTextureObject;
+			FName MassSpawnerID;
+			FText MassSpawnerName;
+			FConstructorStatics()
+				: IconTextureObject(TEXT("/MassGameplay/S_MassCrowd"))
+				, MassSpawnerID(TEXT("MassSpawner"))
+				, MassSpawnerName(NSLOCTEXT("SpriteCategory", "MassSpawner", "MassSpawner"))
+			{
+			}
+		};
+		static FConstructorStatics ConstructorStatics;
+		
+		SpriteComponent->Sprite = ConstructorStatics.IconTextureObject.Get();
+		SpriteComponent->SetRelativeScale3D(FVector(0.5f, 0.5f, 0.5f));
+		SpriteComponent->SpriteInfo.Category = ConstructorStatics.MassSpawnerID;
+		SpriteComponent->SpriteInfo.DisplayName = ConstructorStatics.MassSpawnerName;
+		SpriteComponent->SetupAttachment(RootComponent);
+		SpriteComponent->Mobility = EComponentMobility::Static;
+	}
+#endif // WITH_EDITOR
+
+	SetCanBeDamaged(false);
+	SetActorHiddenInGame(true);
+
+	bAutoSpawnOnBeginPlay = true;
+	bOverrideSchematics = false;
+}
+
+void AMassSpawner::PostInitProperties()
+{
+	Super::PostInitProperties();
+
+	if (HasAnyFlags(RF_ClassDefaultObject) == false)
+	{
+		TickSchematicsAdjustHandle = UMassSimulationSubsystem::GetOnAdjustTickSchematics().AddUObject(this, &AMassSpawner::OnAdjustTickSchematics);
+	}
+}
+
+void AMassSpawner::PostRegisterAllComponents()
+{
+	Super::PostRegisterAllComponents();
+
+	// This is a temp fix for streaming levels async loading MassSpawners after UMassSpawnerSubsystem::OnPostWorldInit, in the long run we are going to need a better system
+	// for making sure all the spawnsets are registered on the clients before replication of Agents occurs. This is only required to be done for clients
+	if (HasAnyFlags(RF_ClassDefaultObject) == false)
+	{
+		UWorld* World = GetWorld();
+		check(World);
+
+		if (GEngine->GetNetMode(GetWorld()) == NM_Client)
+		{
+			UMassSpawnerSubsystem* MassSpawnerSubsystem = UWorld::GetSubsystem<UMassSpawnerSubsystem>(World);
+			if (MassSpawnerSubsystem)
+			{
+				if (bUseEntityConfig)
+				{
+					RegisterEntityTemplates();
+				}
+				else
+				{
+					MassSpawnerSubsystem->RegisterCollection(GetSpawnSets());
+				}
+			}
+			else
+			{
+				FWorldDelegates::OnPostWorldInitialization.AddUObject(this, &AMassSpawner::OnPostWorldInit);
+			}
+		}
+	}
+}
+
+void AMassSpawner::OnPostWorldInit(UWorld* World, const UWorld::InitializationValues)
+{
+	if (World == GetWorld())
+	{
+		UMassSpawnerSubsystem* MassSpawnerSubsystem = UWorld::GetSubsystem<UMassSpawnerSubsystem>(World);
+		check(MassSpawnerSubsystem);
+
+		if (bUseEntityConfig)
+		{
+			RegisterEntityTemplates();
+		}
+		else
+		{
+			MassSpawnerSubsystem->RegisterCollection(GetSpawnSets());
+		}
+
+		FWorldDelegates::OnPostWorldInitialization.Remove(OnPostWorldInitDelegateHandle);
+	}
+}
+
+void AMassSpawner::BeginDestroy()
+{
+	FWorldDelegates::OnPostWorldInitialization.Remove(OnPostWorldInitDelegateHandle);
+	UMassSimulationSubsystem::GetOnAdjustTickSchematics().Remove(TickSchematicsAdjustHandle);
+
+	DoDespawning();
+
+	Super::BeginDestroy();
+}
+
+void AMassSpawner::BeginPlay()
+{
+	check(GEngine);
+
+	Super::BeginPlay();
+
+	const ENetMode NetMode = GEngine->GetNetMode(GetWorld());
+
+	if (bAutoSpawnOnBeginPlay && NetMode != NM_Client)
+	{
+		const UMassSimulationSubsystem* MassSimulationSubsystem = UWorld::GetSubsystem<UMassSimulationSubsystem>(GetWorld());
+		if (MassSimulationSubsystem == nullptr || MassSimulationSubsystem->IsSimulationStarted())
+		{
+			DoSpawning();
+		}
+		else
+		{
+			
+			SimulationStartedHandle = UMassSimulationSubsystem::GetOnSimulationStarted().AddLambda([this](UWorld* InWorld)
+				{
+					UWorld* World = GetWorld();
+
+					if (World == InWorld)
+					{
+						DoSpawning();
+					}
+				});
+		}
+	}
+}
+
+void AMassSpawner::EndPlay(const EEndPlayReason::Type EndPlayReason)
+{
+	UMassSimulationSubsystem::GetOnSimulationStarted().Remove(SimulationStartedHandle);
+
+	Super::EndPlay(EndPlayReason);
+}
+
+#if WITH_EDITOR
+void AMassSpawner::PostEditChangeProperty(FPropertyChangedEvent& PropertyChangedEvent)
+{
+	static const FName SpawnSetsName = GET_MEMBER_NAME_CHECKED(AMassSpawner, SpawnSets);
+	static const FName EQSRequestName = GET_MEMBER_NAME_CHECKED(AMassSpawner, EQSRequest);
+	static const FName EntityTypesName = GET_MEMBER_NAME_CHECKED(AMassSpawner, EntityTypes);
+
+	Super::PostEditChangeProperty(PropertyChangedEvent);
+
+	if (PropertyChangedEvent.Property)
+	{
+		const FName PropName = PropertyChangedEvent.Property->GetFName();
+		if (PropName == SpawnSetsName)
+		{
+			ValidateSpawnSets();
+		}
+		else if (PropName == EQSRequestName)
+		{
+			EQSRequest.InitForOwnerAndBlackboard(*this, /*BBAsset=*/nullptr);
+		}
+		else if (PropName == EntityTypesName)
+		{
+			// TODO: Should optimize this, i.e. set a dirty flag and update only when needed.
+			RegisterEntityTemplates();
+		}
+	}
+
+	if (PropertyChangedEvent.MemberProperty &&
+		PropertyChangedEvent.MemberProperty->GetFName() == EQSRequestName)
+	{
+		EQSRequest.PostEditChangeProperty(*this, PropertyChangedEvent);
+	}
+}
+
+void AMassSpawner::ValidateSpawnSets()
+{
+	// @todo shortly
+}
+
+void AMassSpawner::DEBUG_Spawn()
+{
+	DoSpawning();
+}
+
+void AMassSpawner::DEBUG_Clear()
+{
+	DoDespawning();
+}
+
+#endif // WITH_EDITOR
+
+void AMassSpawner::RegisterEntityTemplates()
+{
+	for (FMassSpawnedEntityType& EntityType : EntityTypes)
+	{
+		if (const UMassEntityConfigAsset* EntityConfig = EntityType.GetEntityConfig())
+		{
+			EntityConfig->GetConfig().GetOrCreateEntityTemplate(*this, *EntityConfig);
+		}
+	}
+}
+
+void AMassSpawner::DoSpawning()
+{
+	if (bUseSpawnPointsGenerators)
+	{
+		const int32 SpawnCount = GetSpawnCount();
+		if (SpawnCount > 0)
+		{
+			float TotalProportion = 0.0f;
+			for (FMassSpawnPointGenerator& SpawnPointsGenerator : SpawnPointsGenerators)
+			{
+				if (SpawnPointsGenerator.GeneratorClass.GetDefaultObject())
+				{
+					SpawnPointsGenerator.bPointsGenerated = false;
+					TotalProportion += SpawnPointsGenerator.Proportion;
+				}
+			}
+
+			if (TotalProportion <= 0.0f)
+			{
+				UE_VLOG_UELOG(this, LogMassSpawner, Error, TEXT("The total combined porportion of all the generator needs to be greater than 0.0f."));
+				return;
+			}
+
+			int32 SpawnPointCountRemaining = SpawnCount;
+			float ProportionRemaining = TotalProportion;
+			for(FMassSpawnPointGenerator& SpawnPointsGenerator : SpawnPointsGenerators)
+			{
+				// Still needs to call the OnSpawnPointGenerationFinished for Generator we don't want to output Spawn Points for as the system will wait
+				// for the bPointsGenerated == true on every Generators before proceeding to the next steps.
+				if (SpawnPointsGenerator.Proportion == 0.0f || ProportionRemaining <= 0.0f || SpawnPointCountRemaining <= 0)
+				{
+					const TArray<FVector> EmptyLoc;
+					OnSpawnPointGenerationFinished(EmptyLoc, &SpawnPointsGenerator);
+					continue;
+				}
+
+				if(const UMassEntitySpawnPointsGeneratorBase* Generator = SpawnPointsGenerator.GeneratorClass.GetDefaultObject())
+				{
+					const float ProportionRatio = FMath::Min(SpawnPointsGenerator.Proportion / ProportionRemaining, 1.0f);
+					const int32 SpawnPointCount = FMath::CeilToInt(SpawnPointCountRemaining * ProportionRatio);
+					if(SpawnPointCount > 0)
+					{
+						FFinishedGeneratingSpawnPointsSignature Delegate = FFinishedGeneratingSpawnPointsSignature::CreateUObject(this, &AMassSpawner::OnSpawnPointGenerationFinished, &SpawnPointsGenerator);
+						Generator->GenerateSpawnPoints(*this, SpawnPointCount, Delegate);
+						SpawnPointCountRemaining -= SpawnPointCount;
+						ProportionRemaining -= SpawnPointsGenerator.Proportion;
+					}
+				}
+			}
+		}
+	}
+	else
+	{
+		if (EQSRequest.IsValid() == false)
+		{
+			EQSRequest.InitForOwnerAndBlackboard(*this, /*BBAsset=*/nullptr);
+			if (!ensureMsgf(EQSRequest.IsValid(), TEXT("Query request initialization can fail due to missing parameters. See the runtime log for details")))
+			{
+				return;
+			}
+		}
+
+		FQueryFinishedSignature Delegate = FQueryFinishedSignature::CreateUObject(this, &AMassSpawner::OnEQSQueryFinished);
+		EQSRequest.Execute(*this, /*BlackboardComponent=*/nullptr, Delegate);
+	}
+}
+
+void AMassSpawner::OnSpawnPointGenerationFinished(const TArray<FVector>& Locations, FMassSpawnPointGenerator* FinishedGenerator)
+{
+	AllGeneratedLocations.Append(Locations);
+
+	bool bAllSpawnPointsGenerated = true;
+	bool bFoundFinishedGenerator = false;
+	for (FMassSpawnPointGenerator& SpawnPointsGenerator : SpawnPointsGenerators)
+	{
+		if (&SpawnPointsGenerator == FinishedGenerator)
+		{
+			SpawnPointsGenerator.bPointsGenerated = true;
+			bFoundFinishedGenerator = true;
+		}
+
+		bAllSpawnPointsGenerated &= SpawnPointsGenerator.bPointsGenerated;
+	}
+
+	checkf(bFoundFinishedGenerator, TEXT("Something went wrong, we are receiving a callback on an unknow spawn point generator"))
+	
+	if(bAllSpawnPointsGenerated)
+	{
+		SpawnAtLocations(AllGeneratedLocations);
+	}
+}
+
+int32 AMassSpawner::GetSpawnCount() const
+{
+	const float FinalSpawningCountScale = SpawningCountScale * UE::MassSpawner::ScalabilitySpawnDensityMultiplier;
+	
+	if (bUseEntityConfig)
+	{
+		return int32(FinalSpawningCountScale * Count);
+	}
+	else
+	{
+		int32 OldCount = 0;
+		for (const FInstancedStruct& Entry : SpawnSets)
+		{
+			const FMassSpawnConfigBase& Config = Entry.Get<FMassSpawnConfigBase>();
+			OldCount += Config.MaxNumber;
+		}
+		return int32(FinalSpawningCountScale * OldCount);
+	}
+}
+
+void AMassSpawner::OnEQSQueryFinished(TSharedPtr<FEnvQueryResult> Result)
+{
+	if (Result.IsValid() == false || Result->IsSuccessful() == false)
+	{
+		UE_VLOG_UELOG(this, LogMassSpawner, Error, TEXT("EQS query failed or result is invalid"));
+		return;
+	}
+	
+	TArray<FVector> Locations;
+	Result->GetAllAsLocations(Locations);
+
+	SpawnAtLocations(Locations);
+}
+
+void AMassSpawner::SpawnAtLocations(const TArray<FVector>& Locations)
+{
+	UMassSpawnerSubsystem* SpawnerSystem = UWorld::GetSubsystem<UMassSpawnerSubsystem>(GetWorld());
+	if (SpawnerSystem == nullptr)
+	{
+		UE_VLOG_UELOG(this, LogMassSpawner, Error, TEXT("UMassSpawnerSubsystem missing while trying to spawn entities"));
+		return;
+	}
+
+	FMassSpawnAuxData MassSpawnAuxData;
+	MassSpawnAuxData.Transforms.Reserve(Locations.Num());
+	for (const FVector& Location : Locations)
+	{
+		FTransform& Transform = MassSpawnAuxData.Transforms.AddDefaulted_GetRef();
+
+		Transform.SetLocation(Location);
+	}
+
+	if (bUseEntityConfig)
+	{
+		const int32 SpawnCount = GetSpawnCount();
+		if(SpawnCount > 0)
+		{
+			float TotalProportion = 0.0f;
+			for (const FMassSpawnedEntityType& EntityType : EntityTypes)
+			{
+				TotalProportion += EntityType.Proportion;
+			}
+
+			if (TotalProportion <= 0)
+			{
+				UE_VLOG_UELOG(this, LogMassSpawner, Error, TEXT("The total combined porportion of all the entity types needs to be greater than 0.0f."));
+				return;
+			}
+
+			AllSpawnedEntities.Reserve(AllSpawnedEntities.Num() + EntityTypes.Num());
+
+			for (const FMassSpawnedEntityType& EntityType : EntityTypes)
+			{
+				if (const UMassEntityConfigAsset* EntityConfig = EntityType.GetEntityConfig())
+				{
+					const int32 EntityCount = int32(SpawnCount * EntityType.Proportion / TotalProportion);
+					if(EntityCount > 0)
+					{
+						FMassSpawnConfigBase SpawnConfig;
+						SpawnConfig.MinNumber = EntityCount;
+						SpawnConfig.MaxNumber = EntityCount;
+						const FMassEntityTemplate* EntityTemplate = EntityConfig->GetConfig().GetOrCreateEntityTemplate(*this, *EntityConfig);
+						if (EntityTemplate && EntityTemplate->IsValid())
+						{
+							FSpawnedEntities& SpawnedEntities = AllSpawnedEntities.AddDefaulted_GetRef();
+							SpawnedEntities.TemplateID = EntityTemplate->GetTemplateID();
+							SpawnerSystem->SpawnEntities(EntityTemplate->GetTemplateID(), SpawnConfig, FStructView::Make(MassSpawnAuxData), SpawnedEntities.Entities);
+						}
+					}
+				}
+			}
+		}
+	}
+	else
+	{
+		SpawnerSystem->SpawnCollection(MakeArrayView(SpawnSets), /*Count=*/-1, FStructView::Make(MassSpawnAuxData));
+	}
+}
+
+void AMassSpawner::DoDespawning()
+{
+	if (AllSpawnedEntities.IsEmpty())
+	{
+		return;
+	}
+	
+	UMassSpawnerSubsystem* SpawnerSystem = UWorld::GetSubsystem<UMassSpawnerSubsystem>(GetWorld());
+	if (SpawnerSystem == nullptr)
+	{
+		UE_LOG(LogMassSpawner, Error, TEXT("UMassSpawnerSubsystem missing while trying to despawn entities"));
+		return;
+	}
+
+	for (const FSpawnedEntities& SpawnedEntities : AllSpawnedEntities)
+	{
+		SpawnerSystem->DestroyEntities(SpawnedEntities.TemplateID, SpawnedEntities.Entities);
+	}
+	AllSpawnedEntities.Reset();
+}
+
+bool AMassSpawner::DespawnEntity(const FLWEntity Entity)
+{
+	UMassSpawnerSubsystem* SpawnerSystem = UWorld::GetSubsystem<UMassSpawnerSubsystem>(GetWorld());
+	if (SpawnerSystem == nullptr)
+	{
+		UE_LOG(LogMassSpawner, Error, TEXT("UMassSpawnerSubsystem missing while trying to despawn a single entity"));
+		return false;
+	}
+
+	for (FSpawnedEntities& SpawnedEntities : AllSpawnedEntities)
+	{
+		const int32 Index = SpawnedEntities.Entities.Find(Entity);
+		if (Index != INDEX_NONE)
+		{
+			SpawnerSystem->DestroyEntities(SpawnedEntities.TemplateID, MakeArrayView(&Entity, 1));
+			// not using "RemoveAtSwap" on purpose to retain the order which will help creation of sparse chunks if need be
+			SpawnedEntities.Entities.RemoveAt(Index, 1, /*bAllowShrinking=*/false);
+			return true;
+		}
+	}
+	return false;
+}
+
+void AMassSpawner::ClearTemplates()
+{
+	for (FMassSpawnedEntityType& EntityType : EntityTypes)
+	{
+		if (const UMassEntityConfigAsset* EntityConfig = EntityType.GetEntityConfig())
+		{
+			EntityConfig->GetConfig().DestroyEntityTemplate(*this, *EntityConfig);
+		}
+	}
+}
+
+void AMassSpawner::UnloadConfig()
+{
+	// Clear all templates that were created by the config
+	ClearTemplates();
+	for (FMassSpawnedEntityType& EntityType : EntityTypes)
+	{
+		EntityType.UnloadEntityConfig();
+	}
+}
+
+void AMassSpawner::OnAdjustTickSchematics(UWorld* World,TArray<TSoftObjectPtr<UPipeSchematic>>& InOutTickSchematics)
+{
+	// this function gets called from a static multicast delegate, so we need to differentiate based on the World it's being called for.
+	if (World != GetWorld())
+	{
+		return;
+	}
+
+	if (bOverrideSchematics)
+	{
+		UE_VLOG_UELOG(World, LogMassSpawner, Log, TEXT("Overriding TickSchematics with %s"), *GetName());
+		InOutTickSchematics = TickSchematics;
+	}
+	else
+	{
+		UE_VLOG_UELOG(World, LogMassSpawner, Log, TEXT("Appending to TickSchematics with %s"), *GetName());
+		InOutTickSchematics.Append(TickSchematics);
+	}
+}
+
diff --git a/Engine/Plugins/Runtime/MassGameplay/Source/MassSpawner/Private/MassSpawnerSubsystem.cpp b/Engine/Plugins/Runtime/MassGameplay/Source/MassSpawner/Private/MassSpawnerSubsystem.cpp
new file mode 100644
--- /dev/null
+++ b/Engine/Plugins/Runtime/MassGameplay/Source/MassSpawner/Private/MassSpawnerSubsystem.cpp
@@ -0,0 +1,330 @@
+// Copyright Epic Games, Inc. All Rights Reserved.
+
+#include "MassSpawnerSubsystem.h"
+#include "MassSpawnerTypes.h"
+#include "MassEntityTemplate.h"
+#include "MassEntitySystem.h"
+#include "MassEntityTemplateRegistry.h"
+#include "Engine/World.h"
+#include "MassExecutor.h"
+#include "InstancedStruct.h"
+#include "VisualLogger/VisualLogger.h"
+#include "MassSpawner.h"
+#include "MassTranslator.h"
+#include "MassTranslatorRegistry.h"
+#include "ProfilingDebugging/CpuProfilerTrace.h"
+#include "MassSimulationSubsystem.h"
+
+
+namespace UE::Mass::SpawnerSubsystem 
+{
+
+	bool RunFragmentDestructors(const FArchetypeChunkCollection& Chunks, UMassSpawnerSubsystem& SpawnerSubSystem, UPipeEntitySubsystem& EntitySystem)
+	{
+		const FArchetypeHandle ArchetypeHandle = Chunks.GetArchetype();
+		// @todo this is a temporary measure to skip entities we've destroyed without the
+		// UMassSpawnerSubsystem::DestroyEntities' caller knowledge. Should be removed once that's addressed.
+		if (ArchetypeHandle.IsValid() == false)
+		{
+			return false;
+		}
+
+		// @Todo: should queue up the entity destruction and apply in batch
+		const UMassTranslatorRegistry& Registry = UMassTranslatorRegistry::Get();
+		TArray<const UPipeProcessor*> Destructors;
+
+		EntitySystem.ForEachArchetypeComponentType(ArchetypeHandle, [&Registry, &Destructors](const UScriptStruct* FragmentType)
+		{
+			check(FragmentType);
+			const UMassFragmentDestructor* Destructor = Registry.GetFragmentDestructor(*FragmentType);
+			if (Destructor)
+			{
+				Destructors.AddUnique(Destructor);
+			}
+		});
+
+		if (Destructors.Num())
+		{
+			FRuntimePipeline DestructionPipeline;
+			DestructionPipeline.InitializeFromArray(Destructors, SpawnerSubSystem);
+
+			FPipeContext PipeContext(EntitySystem, /*TimeDelta=*/0.0f);
+
+			UE::Pipe::Executor::RunSparse(DestructionPipeline, PipeContext, Chunks);
+		}
+
+		return Destructors.Num() > 0;
+	}
+	
+	void CreateSparseChunks(const UPipeEntitySubsystem& EntitySystem, const TConstArrayView<FLWEntity> Entities, TArray<FArchetypeChunkCollection>& OutChunkCollections)
+	{
+		TRACE_CPUPROFILER_EVENT_SCOPE_STR("SpawnerSubsystem_CreateSparseChunks");
+
+		TMap<const FArchetypeHandle, TArray<FLWEntity>> ArchetypeToEntities;
+
+		for (const FLWEntity& Entity : Entities)
+		{
+			FArchetypeHandle Archetype = EntitySystem.GetArchetypeForEntity(Entity);
+			TArray<FLWEntity>& PerArchetypeEntities = ArchetypeToEntities.FindOrAdd(Archetype);
+			PerArchetypeEntities.Add(Entity);
+		}
+
+		for (auto& Pair : ArchetypeToEntities)
+		{
+			// @todo this is a temporary measure to skip entities we've destroyed without the
+			// UMassSpawnerSubsystem::DestroyEntities' caller knowledge. Should be removed once that's addressed.
+			if (Pair.Key.IsValid())
+			{
+				OutChunkCollections.Add(FArchetypeChunkCollection(Pair.Key, Pair.Value));
+			}
+		}
+	}
+	
+} // namespace UE::Mass::SpawnerSubsystem
+
+//----------------------------------------------------------------------//
+//  UMassSpawnerSubsystem
+//----------------------------------------------------------------------//
+void UMassSpawnerSubsystem::Initialize(FSubsystemCollectionBase& Collection) 
+{	
+	Super::Initialize(Collection);
+
+	// making sure UMassSimulationSubsystem gets created before the MassActorManager
+	Collection.InitializeDependency<UMassSimulationSubsystem>();
+
+	SimulationSystem = UWorld::GetSubsystem<UMassSimulationSubsystem>(GetWorld());
+	EntitySystem = UPipeEntitySubsystem::GetCurrent(GetWorld());
+	check(EntitySystem);
+}
+
+void UMassSpawnerSubsystem::PostInitialize()
+{
+	TemplateRegistryInstance = NewObject<UMassEntityTemplateRegistry>(this);
+}
+
+bool UMassSpawnerSubsystem::SpawnEntities(const AActor& ActorInstance, const uint32 NumberToSpawn, TArray<FLWEntity>& OutEntities) const
+{
+	if (NumberToSpawn == 0)
+	{
+		UE_VLOG(this, LogMassSpawner, Warning, TEXT("Trying to spawn 0 entities. This would cause inefficiency. Bailing out with result FALSE."));
+		return false;
+	}
+
+	check(TemplateRegistryInstance);
+	if (const FMassEntityTemplate* EntityTemplate = TemplateRegistryInstance->FindOrBuildInstanceTemplate(ActorInstance))
+	{
+		return SpawnEntities(*EntityTemplate, NumberToSpawn, OutEntities);
+	}
+	return false;
+}
+
+bool UMassSpawnerSubsystem::SpawnEntities(const FMassEntityTemplate& EntityTemplate, const uint32 NumberToSpawn, TArray<FLWEntity>& OutEntities) const
+{
+	check(EntitySystem);
+	check(EntityTemplate.IsValid());
+
+	if (NumberToSpawn == 0)
+	{
+		UE_VLOG(this, LogMassSpawner, Warning, TEXT("Trying to spawn 0 entities. This would cause inefficiency. Bailing out with result FALSE."));
+		return false;
+	}
+
+	//TRACE_CPUPROFILER_EVENT_SCOPE_STR("MassSpawnerSubsystem SpawnEntities");
+	
+	OutEntities.Reset();
+	if (EntitySystem->BatchCreateEntities(EntityTemplate.GetArchetype(), NumberToSpawn, OutEntities) > 0)
+	{ 
+		UE_CVLOG(OutEntities.Num() != NumberToSpawn, this, LogMassSpawner, Warning, TEXT("Tried to batch-create %d entities but created %d."), NumberToSpawn, OutEntities.Num());
+
+		const TArrayView<const FInstancedStruct> FragmentInstances = EntityTemplate.GetFragments();
+		const FArchetypeChunkCollection ChunkCollection(EntityTemplate.GetArchetype(), OutEntities);
+		EntitySystem->BatchSetEntityComponentsValues(ChunkCollection, FragmentInstances);
+
+		if (EntityTemplate.GetInitializationPipeline().Processors.Num())
+		{
+			FMassEntityTemplate& MutableEntityTemplate = const_cast<FMassEntityTemplate&>(EntityTemplate);
+
+			FPipeContext PipeContext(*EntitySystem, /*TimeDelta=*/0.0f);
+			UE::Pipe::Executor::RunSparse(MutableEntityTemplate.GetMutableInitializationPipeline(), PipeContext, ChunkCollection);
+		}
+		return true;
+	}
+	return false;
+}
+
+void UMassSpawnerSubsystem::SpawnEntities(FMassEntityTemplateID TemplateID, const FMassSpawnConfigBase& SpawnConfig, const FStructView& AuxData, TArray<FLWEntity>& OutEntities)
+{
+	check(TemplateID.IsValid());
+
+	if (!ensureMsgf(TemplateRegistryInstance, TEXT("UMassSpawnerSubsystem didn\'t get its OnPostWorldInit call yet!")))
+	{
+		return;
+	}
+
+	const FMassEntityTemplate* EntityTemplate = TemplateRegistryInstance->FindTemplateFromTemplateID(TemplateID);
+	checkf(EntityTemplate != nullptr, TEXT("TemplateID must have been registered!"));
+
+	DoSpawning(*EntityTemplate, SpawnConfig, AuxData, OutEntities);
+}
+
+
+void UMassSpawnerSubsystem::SpawnCollection(TArrayView<FInstancedStruct> Collection, const int32 Count, const FStructView& AuxData /* = FStructView() */)
+{
+	if (!ensureMsgf(TemplateRegistryInstance, TEXT("UMassSpawnerSubsystem didn\'t get its OnPostWorldInit call yet!")))
+	{
+		return;
+	}
+
+	// @todo totally ignoring "Count" parameter for now. Respecting it will required building spawning model properly
+	// right now the function will just spawn Max allowed by given spawn data
+
+	for (FInstancedStruct& Entry : Collection)
+	{
+		if (Entry.IsValid())
+		{
+			const FMassEntityTemplate* EntityTemplate = TemplateRegistryInstance->FindOrBuildStructTemplate(Entry);
+			if (EntityTemplate && EntityTemplate->IsValid())
+			{
+				TArray<FLWEntity> Entities;
+
+				DoSpawning(*EntityTemplate, Entry.GetMutable<FMassSpawnConfigBase>(), AuxData, Entities);
+			}
+		}
+	}
+}
+
+void UMassSpawnerSubsystem::RegisterCollection(TArrayView<FInstancedStruct> Collection)
+{
+	if (!ensureMsgf(TemplateRegistryInstance, TEXT("UMassSpawnerSubsystem didn't get its OnPostWorldInit called yet!")))
+	{
+		return;
+	}
+
+	for (FInstancedStruct& Entry : Collection)
+	{
+		if (Entry.IsValid())
+		{
+			TemplateRegistryInstance->FindOrBuildStructTemplate(Entry);
+		}
+	}
+}
+
+void UMassSpawnerSubsystem::DestroyEntities(const FMassEntityTemplateID TemplateID, TConstArrayView<FLWEntity> Entities)
+{
+	TRACE_CPUPROFILER_EVENT_SCOPE_STR("MassSpawnerSubsystem_DestroyEntities")
+
+	if (!ensureMsgf(TemplateRegistryInstance, TEXT("UMassSpawnerSubsystem didn\'t get its OnPostWorldInit call yet!")))
+	{
+		return;
+	}
+
+	check(EntitySystem);
+	check(SimulationSystem);
+	checkf(!SimulationSystem->GetPhaseManager().IsDuringPipeProcessing()
+		, TEXT("%s called while MassEntity processing in progress. This is unsupported and dangerous!"), ANSI_TO_TCHAR(__FUNCTION__));
+
+	UWorld* World = GetWorld();
+	check(World);
+
+	// Run deinitializers and destructor processors on all entities.
+	const FMassEntityTemplate* EntityTemplate = TemplateRegistryInstance->FindTemplateFromTemplateID(TemplateID);
+	if (EntityTemplate)
+	{
+		// There's no certainty all Entities still belong to the original archetype, the MutableEntityTemplate.GetArchetype().
+		// We need to figure out the separate, per-archetype sparse chunks
+		TArray<FArchetypeChunkCollection> ChunkCollections;
+		UE::Mass::SpawnerSubsystem::CreateSparseChunks(*EntitySystem, Entities, ChunkCollections);
+
+		FPipeContext PipeContext(*EntitySystem, /*TimeDelta=*/0.0f);
+
+		// Run destructors
+		bool bDestructorsRun = false;
+		for (const FArchetypeChunkCollection& Chunks : ChunkCollections)
+		{
+			bDestructorsRun = UE::Mass::SpawnerSubsystem::RunFragmentDestructors(Chunks, *this, *EntitySystem) || bDestructorsRun;
+		}
+
+		// if any destructors were called it's possible the entities changed archetypes. Need to 
+		// recalculate the chunks
+		if (bDestructorsRun)
+		{
+			ChunkCollections.Reset();
+			UE::Mass::SpawnerSubsystem::CreateSparseChunks(*EntitySystem, Entities, ChunkCollections);
+		}
+		
+		// Run template deinitializers
+		FMassEntityTemplate& MutableEntityTemplate = const_cast<FMassEntityTemplate&>(*EntityTemplate);
+		if (MutableEntityTemplate.GetDeinitializationPipeline().Processors.Num())
+		{
+			for (const FArchetypeChunkCollection& Chunks : ChunkCollections)
+			{
+				if (Chunks.GetArchetype().IsValid())
+				{
+					UE::Pipe::Executor::RunSparse(MutableEntityTemplate.GetMutableDeinitializationPipeline(), PipeContext, Chunks);
+				}
+			}
+		}
+
+		// if any deinitializer were called it's possible the entities changed archetypes. Need to 
+		// recalculate the chunks
+		if (MutableEntityTemplate.GetDeinitializationPipeline().Processors.Num())
+		{
+			ChunkCollections.Reset();
+			UE::Mass::SpawnerSubsystem::CreateSparseChunks(*EntitySystem, Entities, ChunkCollections);
+		}
+
+		for (const FArchetypeChunkCollection& Chunks : ChunkCollections)
+		{
+			EntitySystem->BatchDestroyEntityChunks(Chunks);
+		}
+	}
+	else
+	{
+		EntitySystem->BatchDestroyEntities(Entities);
+	}
+}
+
+void UMassSpawnerSubsystem::DoSpawning(const FMassEntityTemplate& EntityTemplate, const FMassSpawnConfigBase& Data, const FStructView& AuxData, TArray<FLWEntity>& OutEntities) const
+{
+	check(EntitySystem);
+	check(EntityTemplate.GetArchetype().IsValid());
+	UE_VLOG(this, LogMassSpawner, Log, TEXT("Spawning with EntityTemplate:\n%s"), *EntityTemplate.DebugGetDescription(EntitySystem));
+
+	//TRACE_CPUPROFILER_EVENT_SCOPE_STR("MassSpawnerSubsystem DoSpawning");
+
+	// 1. Create required number of entities with EntityTemplate.Archetype
+	// 2. Copy data from FMassEntityTemplate.Fragments.
+	//		a. @todo, could be done as part of creation?
+	//		b. @todo could be done via an always-first translator in EntityTemplate.Translators
+	// 3. Run all EntityTemplate.Translators to configure remaining values (i.e. not set via FMassEntityTemplate.Fragments).
+
+	TArray<FLWEntity> SpawnedEntities;
+	EntitySystem->BatchCreateEntities(EntityTemplate.GetArchetype(), Data.MaxNumber, SpawnedEntities);
+
+	const TArrayView<const FInstancedStruct> FragmentInstances = EntityTemplate.GetFragments();
+	const FArchetypeChunkCollection ChunkCollection(EntityTemplate.GetArchetype(), SpawnedEntities);
+	EntitySystem->BatchSetEntityComponentsValues(ChunkCollection, FragmentInstances);
+	
+	if (EntityTemplate.GetInitializationPipeline().Processors.Num())
+	{
+		FMassEntityTemplate& MutableEntityTemplate = const_cast<FMassEntityTemplate&>(EntityTemplate);
+
+		FPipeContext PipeContext(*EntitySystem, /*TimeDelta=*/0.0f);
+		PipeContext.AuxData = AuxData;
+		UE::Pipe::Executor::RunSparse(MutableEntityTemplate.GetMutableInitializationPipeline(), PipeContext, ChunkCollection);
+	}
+
+	OutEntities.Append(MoveTemp(SpawnedEntities));
+}
+
+const FMassEntityTemplate* UMassSpawnerSubsystem::GetMassEntityTemplate(FMassEntityTemplateID TemplateID) const
+{
+	check(TemplateID.IsValid());
+
+	if (!ensureMsgf(TemplateRegistryInstance, TEXT("UMassSpawnerSubsystem didn\'t get its OnPostWorldInit call yet!")))
+	{
+		return nullptr;
+	}
+
+	return TemplateRegistryInstance->FindTemplateFromTemplateID(TemplateID);
+}
diff --git a/Engine/Plugins/Runtime/MassGameplay/Source/MassSpawner/Private/MassTranslator.cpp b/Engine/Plugins/Runtime/MassGameplay/Source/MassSpawner/Private/MassTranslator.cpp
new file mode 100644
--- /dev/null
+++ b/Engine/Plugins/Runtime/MassGameplay/Source/MassSpawner/Private/MassTranslator.cpp
@@ -0,0 +1,62 @@
+// Copyright Epic Games, Inc. All Rights Reserved.
+
+#include "MassTranslator.h"
+#include "MassCommonTypes.h"
+#include "MassTranslatorRegistry.h"
+
+//----------------------------------------------------------------------//
+//  UMassTranslator
+//----------------------------------------------------------------------//
+UMassTranslator::UMassTranslator()
+{
+	bAutoRegisterWithProcessingPhases = true;
+	ExecutionFlags = (int32)EProcessorExecutionFlags::All;
+}
+
+void UMassTranslator::AddRequiredTagsToQuery(FLWComponentQuery& EntityQuery)
+{
+	EntityQuery.AddTagRequirements<ELWComponentPresence::All>(RequiredTags);
+}
+
+//----------------------------------------------------------------------//
+// UMassFragmentBuilder
+//----------------------------------------------------------------------//
+UMassFragmentBuilder::UMassFragmentBuilder()
+{
+	bAutoRegisterWithProcessingPhases = false;
+#if WITH_EDITORONLY_DATA
+	bCanShowUpInSettings = false;
+#endif // WITH_EDITORONLY_DATA
+}
+
+void UMassFragmentBuilder::PostInitProperties()
+{
+	Super::PostInitProperties();
+
+	if (GetClass()->HasAnyClassFlags(CLASS_Abstract) == false && FragmentType != nullptr)
+	{
+		Register();
+	}
+}
+
+//----------------------------------------------------------------------//
+//  UMassFragmentInitializer
+//----------------------------------------------------------------------//
+void UMassFragmentInitializer::Register()
+{
+	if (FragmentType)
+	{
+		UMassTranslatorRegistry::GetMutable().RegisterFragmentInitializer(*FragmentType, GetClass());
+	}
+}
+
+//----------------------------------------------------------------------//
+//  UMassFragmentDestructor
+//----------------------------------------------------------------------//
+void UMassFragmentDestructor::Register()
+{
+	if (FragmentType)
+	{
+		UMassTranslatorRegistry::GetMutable().RegisterFragmentDestructor(*FragmentType, GetClass());
+	}
+}
\ No newline at end of file
diff --git a/Engine/Plugins/Runtime/MassGameplay/Source/MassSpawner/Private/MassTranslatorRegistry.cpp b/Engine/Plugins/Runtime/MassGameplay/Source/MassSpawner/Private/MassTranslatorRegistry.cpp
new file mode 100644
--- /dev/null
+++ b/Engine/Plugins/Runtime/MassGameplay/Source/MassSpawner/Private/MassTranslatorRegistry.cpp
@@ -0,0 +1,38 @@
+// Copyright Epic Games, Inc. All Rights Reserved.
+
+#include "MassTranslatorRegistry.h"
+#include "MassCommonTypes.h"
+#include "MassTranslator.h"
+
+//----------------------------------------------------------------------//
+// UMassTranslatorRegistry
+//----------------------------------------------------------------------//
+UMassTranslatorRegistry::UMassTranslatorRegistry()
+{
+	// there can be only one!
+	check(HasAnyFlags(RF_ClassDefaultObject));
+}
+
+void UMassTranslatorRegistry::RegisterFragmentInitializer(const UScriptStruct& FragmentType, TSubclassOf<UMassFragmentInitializer> FragmentInitializerClass)
+{
+	check(FragmentInitializerClass);
+	FragmentInitializersMap.FindOrAdd(&FragmentType) = FragmentInitializerClass;
+}
+
+void UMassTranslatorRegistry::RegisterFragmentDestructor(const UScriptStruct& FragmentType, TSubclassOf<UMassFragmentDestructor> FragmentDestructorClass)
+{
+	check(FragmentDestructorClass);
+	FragmentDestructorsMap.FindOrAdd(&FragmentType) = FragmentDestructorClass;
+}
+
+const UMassFragmentInitializer* UMassTranslatorRegistry::GetFragmentInitializer(const UScriptStruct& FragmentType) const
+{
+	const TSubclassOf<UMassFragmentInitializer>* InitializerClass = FragmentInitializersMap.Find(&FragmentType);
+	return (InitializerClass && *InitializerClass) ? GetMutableDefault<UMassFragmentInitializer>(*InitializerClass) : nullptr;
+	
+}
+const UMassFragmentDestructor* UMassTranslatorRegistry::GetFragmentDestructor(const UScriptStruct& FragmentType) const
+{
+	const TSubclassOf<UMassFragmentDestructor>* DestructorClass = FragmentDestructorsMap.Find(&FragmentType);
+	return (DestructorClass && *DestructorClass) ? GetMutableDefault<UMassFragmentDestructor>(*DestructorClass) : nullptr;
+}
diff --git a/Engine/Plugins/Runtime/MassGameplay/Source/MassSpawner/Public/MassAssortedFragmentsTrait.h b/Engine/Plugins/Runtime/MassGameplay/Source/MassSpawner/Public/MassAssortedFragmentsTrait.h
new file mode 100644
--- /dev/null
+++ b/Engine/Plugins/Runtime/MassGameplay/Source/MassSpawner/Public/MassAssortedFragmentsTrait.h
@@ -0,0 +1,24 @@
+// Copyright Epic Games, Inc. All Rights Reserved.
+
+#pragma once 
+
+#include "CoreMinimal.h"
+#include "MassEntityTraitBase.h"
+#include "InstancedStruct.h"
+#include "MassAssortedFragmentsTrait.generated.h"
+
+/**
+* Mass Agent Feature which appends a list of specified fragments.  
+*/
+UCLASS(BlueprintType, EditInlineNew, CollapseCategories, meta=(DisplayName="Assorted Fragments"))
+class MASSSPAWNER_API UMassAssortedFragmentsTrait : public UMassEntityTraitBase
+{
+	GENERATED_BODY()
+
+protected:
+
+	virtual void BuildTemplate(FMassEntityTemplateBuildContext& BuildContext, UWorld& World) const override;
+
+	UPROPERTY(Category="Fragments", EditAnywhere, meta = (BaseStruct = "LWComponentData", ExcludeBaseStruct))
+	TArray<FInstancedStruct> Fragments;
+};
diff --git a/Engine/Plugins/Runtime/MassGameplay/Source/MassSpawner/Public/MassEntityConfigAsset.h b/Engine/Plugins/Runtime/MassGameplay/Source/MassSpawner/Public/MassEntityConfigAsset.h
new file mode 100644
--- /dev/null
+++ b/Engine/Plugins/Runtime/MassGameplay/Source/MassSpawner/Public/MassEntityConfigAsset.h
@@ -0,0 +1,87 @@
+// Copyright Epic Games, Inc. All Rights Reserved.
+
+#pragma once 
+
+#include "CoreMinimal.h"
+#include "Engine/DataAsset.h"
+#include "MassEntityTemplate.h"
+#include "MassEntityTemplateRegistry.h"
+#include "MassEntityConfigAsset.generated.h"
+
+class UMassEntityConfigAsset;
+class UMassEntityTraitBase;
+
+/**
+ * Describes a Mass agent to spawn. The struct can be embedded to allow last minute changes to the agent (i.e. for debugging).
+ * The agent config describes a unique list of features which are used to create an entity template.
+ * Derived configs can override parent features.
+ */
+USTRUCT()
+struct MASSSPAWNER_API FMassEntityConfig
+{
+	friend class UMassTemplateRegistry;
+
+	GENERATED_BODY()
+
+	FMassEntityConfig() = default;
+	FMassEntityConfig(UMassEntityConfigAsset& InParent);
+
+	/** Create entity template based on the features included in this config.
+	 *  @param OwnerActor Actor who is creating the template.
+	 *  @param ConfigOwner Owner of the FMassEntityConfig used for error reporting.
+	 */
+	const FMassEntityTemplate* GetOrCreateEntityTemplate(AActor& OwnerActor, const UObject& ConfigOwner) const;
+
+	void DestroyEntityTemplate(const AActor& OwnerActor, const UObject& ConfigOwner) const;
+
+	/** 
+	 * Fetches the EntityTemplate for given OwnerActor, fails a check if one cannot be found.
+	 */
+	const FMassEntityTemplate& GetEntityTemplateChecked(const AActor& OwnerActor, const UObject& ConfigOwner) const;
+
+	/** @return Parent config */
+	const UMassEntityConfigAsset* GetParent() const { return Parent; }
+
+	/** @return View to the array of features defined on this config */
+	TConstArrayView<UMassEntityTraitBase*> GetTraits() const { return Traits; }
+
+	/** Adds Trait to the collection of traits hosted by this FMassEntityConfig instance */
+	void AddTrait(UMassEntityTraitBase& Trait);
+
+protected:
+	/** Combines traits based on the config hierarchy and returns list of unique traits */
+	void GetCombinedTraits(TArray<UMassEntityTraitBase*>& OutTraits, TArray<const UObject*>& Visited, const UObject& ConfigOwner) const;
+
+	/** Reference to parent config asset */
+	UPROPERTY(Category = "Derived Traits", EditAnywhere)
+	UMassEntityConfigAsset* Parent = nullptr;
+
+	/** Array of unique traits of this config */
+	UPROPERTY(Category = "Traits", EditAnywhere, Instanced)
+	TArray<UMassEntityTraitBase*> Traits;
+
+private:
+	const FMassEntityTemplate* GetEntityTemplateInternal(const AActor& OwnerActor, const UObject& ConfigOwner, uint32& HashOut, FMassEntityTemplateID& TemplateIDOut, TArray<UMassEntityTraitBase*>& CombinedTraitsOut) const;
+};
+
+/**
+ * Agent Config asset allows to create shared configs that can be used as base for derived configs.
+ * The asset can be used as is i.e. on a spawner, or you can use FMassEntityConfig to allow last minute changes at use site.
+ */
+UCLASS(BlueprintType)
+class MASSSPAWNER_API UMassEntityConfigAsset : public UDataAsset
+{
+	GENERATED_BODY()
+
+public:
+	/** @return Agent config stored in this asset */
+	const FMassEntityConfig& GetConfig() const { return Config; }
+
+	/** @return Mutable agent config stored in this asset */
+	FMassEntityConfig& GetMutableConfig() { return Config; }
+
+protected:
+	/** The config described in this asset. */
+	UPROPERTY(Category = "Entity Config", EditAnywhere)
+	FMassEntityConfig Config;
+};
\ No newline at end of file
diff --git a/Engine/Plugins/Runtime/MassGameplay/Source/MassSpawner/Public/MassEntityEQSSpawnPointsGenerator.h b/Engine/Plugins/Runtime/MassGameplay/Source/MassSpawner/Public/MassEntityEQSSpawnPointsGenerator.h
new file mode 100644
--- /dev/null
+++ b/Engine/Plugins/Runtime/MassGameplay/Source/MassSpawner/Public/MassEntityEQSSpawnPointsGenerator.h
@@ -0,0 +1,26 @@
+// Copyright Epic Games, Inc. All Rights Reserved.
+
+#pragma once 
+
+#include "CoreMinimal.h"
+#include "MassEntitySpawnPointsGeneratorBase.h"
+#include "EnvironmentQuery/EnvQueryTypes.h"
+#include "MassEntityEQSSpawnPointsGenerator.generated.h"
+
+/**
+ * Describes the SpawnPoints Generator when we want to leverage the points given by an EQS Query
+ */
+UCLASS(BlueprintType, meta=(DisplayName="EQS SpawnPoints Generator"))
+class MASSSPAWNER_API UMassEntityEQSSpawnPointsGenerator : public UMassEntitySpawnPointsGeneratorBase
+{
+	GENERATED_BODY()
+
+public:
+	virtual void GenerateSpawnPoints(UObject& QueryOwner, int32 Count, FFinishedGeneratingSpawnPointsSignature& FinishedGeneratingSpawnPointsDelegate) const override;
+
+protected:
+	void OnEQSQueryFinished(TSharedPtr<FEnvQueryResult> Result, int32 Count, FFinishedGeneratingSpawnPointsSignature FinishedGeneratingSpawnPointsDelegate) const;
+
+	UPROPERTY(Category = "Query", EditAnywhere)
+	FEQSParametrizedQueryExecutionRequest EQSRequest;
+};
diff --git a/Engine/Plugins/Runtime/MassGameplay/Source/MassSpawner/Public/MassEntitySpawnPointsGeneratorBase.h b/Engine/Plugins/Runtime/MassGameplay/Source/MassSpawner/Public/MassEntitySpawnPointsGeneratorBase.h
new file mode 100644
--- /dev/null
+++ b/Engine/Plugins/Runtime/MassGameplay/Source/MassSpawner/Public/MassEntitySpawnPointsGeneratorBase.h
@@ -0,0 +1,29 @@
+// Copyright Epic Games, Inc. All Rights Reserved.
+
+#pragma once 
+
+#include "CoreMinimal.h"
+#include "MassEntitySpawnPointsGeneratorBase.generated.h"
+
+/** @Todo: Ideally we would return something else then just a FVector, we could return as an example the lane index, smartobject id, navmesh location, etc...*/
+DECLARE_DELEGATE_OneParam(FFinishedGeneratingSpawnPointsSignature, const TArray<FVector>&);
+
+/**
+ * Base class for Mass Entity Spawn Points Generator.
+ * A Mass Spawn Points Generator can be of several type (EQS, ZoneGraph, Volume, Area, etc.)
+ * The concept is to override the GenerateSpawnPoints() method and requesting a certain number of Spawn Point Locations to the method.
+ */
+UCLASS(Abstract, Blueprintable)
+class MASSSPAWNER_API UMassEntitySpawnPointsGeneratorBase : public UObject
+{
+	GENERATED_BODY()
+
+public:
+
+	/** Generate "Count" number of SpawnPoints and return as a list of position
+	 * @param Count of point to generate
+	 * @param FinishedGeneratingSpawnPointsDelegate is the callback to call once the generation is done
+	 */
+	 virtual void GenerateSpawnPoints(UObject& QueryOwner, int32 Count, FFinishedGeneratingSpawnPointsSignature& FinishedGeneratingSpawnPointsDelegate) const PURE_VIRTUAL(UMassEntitySpawnPointsGeneratorBase::GenerateSpawnPoints, );
+
+};
diff --git a/Engine/Plugins/Runtime/MassGameplay/Source/MassSpawner/Public/MassEntityTemplate.h b/Engine/Plugins/Runtime/MassGameplay/Source/MassSpawner/Public/MassEntityTemplate.h
new file mode 100644
--- /dev/null
+++ b/Engine/Plugins/Runtime/MassGameplay/Source/MassSpawner/Public/MassEntityTemplate.h
@@ -0,0 +1,172 @@
+// Copyright Epic Games, Inc. All Rights Reserved.
+
+#pragma once
+
+#include "CoreMinimal.h"
+#include "MassEntityTypes.h"
+#include "EntitySubsystem.h"
+#include "MassCommonTypes.h"
+#include "MassTranslator.h"
+#include "MassEntityTemplate.generated.h"
+
+
+struct FEntityView;
+
+
+UENUM()
+enum class EMassEntityTemplateIDType : uint8
+{
+	None,
+	ScriptStruct,
+	Class,
+	Instance
+};
+
+USTRUCT()
+struct MASSSPAWNER_API FMassObjectHandlers
+{
+	GENERATED_BODY()
+	UPROPERTY()
+	TArray<const UPipeProcessor*> Initializers;
+	UPROPERTY()
+	TArray<const UPipeProcessor*> Deinitializers;
+};
+
+//ID of the template an entity is using
+USTRUCT()
+struct MASSSPAWNER_API FMassEntityTemplateID
+{
+	GENERATED_BODY()
+
+	FMassEntityTemplateID(uint32 InHash, EMassEntityTemplateIDType InType)
+	: Hash(InHash)
+	, Type(InType)
+	{}
+
+	FMassEntityTemplateID() = default;
+
+	uint32 GetHash() const { return Hash; }
+	void SetHash(uint32 InHash) { Hash = InHash; }
+
+	EMassEntityTemplateIDType GetType() const { return Type; }
+	void SetType(EMassEntityTemplateIDType InType) { Type = InType; }
+
+	bool operator==(const FMassEntityTemplateID& Other) const
+	{
+		return (Hash == Other.Hash) && (Type == Other.Type);
+	}
+
+	friend uint32 GetTypeHash(const FMassEntityTemplateID& TemplateID)
+	{
+		return HashCombine(TemplateID.Hash, (uint32)TemplateID.Type);
+	}
+
+	bool IsValid() { return Type != EMassEntityTemplateIDType::None; }
+
+	FString ToString() const;
+
+protected:
+	UPROPERTY()
+	uint32 Hash = 0;
+
+	UPROPERTY()
+	EMassEntityTemplateIDType Type = EMassEntityTemplateIDType::None;
+};
+
+/** @todo document	*/
+USTRUCT()
+struct MASSSPAWNER_API FMassEntityTemplate
+{
+	GENERATED_BODY()
+
+	typedef TFunction<void(UObject& /*Owner*/, FEntityView& /*EntityView*/, const EMassTranslationDirection /*CurrentDirection*/)> FObjectFragmentInitializerFunction;
+
+	FMassEntityTemplate() = default;
+	/** InArchetype is expected to be valid. The function will crash-check it. */
+	void SetArchetype(const FArchetypeHandle& InArchetype);
+	const FArchetypeHandle& GetArchetype() const { return Archetype; }
+	/** Configures initialization and deinitialization pipeline from ObjectHandlers */
+	void SetUpProcessors(const FMassObjectHandlers& ObjectHandlers, UObject& PipelineOwner);
+	
+	const FRuntimePipeline& GetInitializationPipeline() const { return InitializationPipeline; }
+	FRuntimePipeline& GetMutableInitializationPipeline() { return InitializationPipeline; }
+
+	const FRuntimePipeline& GetDeinitializationPipeline() const { return DeinitializationPipeline; }
+	FRuntimePipeline& GetMutableDeinitializationPipeline() { return DeinitializationPipeline; }
+
+	TConstArrayView<FObjectFragmentInitializerFunction> GetObjectFragmentInitializers() const { return ObjectInitializers; }
+	TArray<FObjectFragmentInitializerFunction>& GetMutableObjectFragmentInitializers() { return ObjectInitializers; }
+
+	bool IsValid() const { return Archetype.IsValid() && (FragmentCollection.IsEmpty() == false); }
+
+	void SetTemplateID(FMassEntityTemplateID InTemplateID) { TemplateID = InTemplateID; }
+	FMassEntityTemplateID GetTemplateID() const { return TemplateID; }
+
+	int32 GetFragmentsNum() const { return FragmentCollection.GetFragmentsNum(); }
+	TArrayView<const FInstancedStruct> GetFragments() const { return FragmentCollection.GetFragments(); }
+	TArray<const UScriptStruct*>& GetMutableFragmentTypes() { return FragmentTypes; }
+	TConstArrayView<const UScriptStruct*> GetFragmentTypes() const { return FragmentTypes; }
+	FMassUniqueFragmentCollection& GetMutableFragmentCollection() { return FragmentCollection; }
+	TConstArrayView<FInstancedStruct> GetChunkFragments() const { return ChunkFragments; }
+
+	FLWCompositionDescriptor GetCompositionDescriptor() const { return FLWCompositionDescriptor(FragmentCollection.GetFragmentBitSet(), TagBitSet, ChunkFragmentBitSet); }
+
+	template<typename T>
+	void AddTag()
+	{
+		TagBitSet.Add<T>();
+	}
+	
+	void AddTag(const UScriptStruct& TagType)
+	{
+		TagBitSet.Add(TagType);
+	}
+
+	const FLWTagBitSet& GetTags() const { return TagBitSet; }	
+	FLWTagBitSet& GetMutableTags() { return TagBitSet; }
+
+	template<typename T>
+	void AddChunkFragment()
+	{
+		static_assert(TIsDerivedFrom<T, FLWChunkComponent>::IsDerived, "Given struct doesn't represent a valid chunk fragment type. Make sure to inherit from FLWChunkComonent or one of its child-types.");
+		ChunkFragments.AddUnique(FInstancedStruct(T::StaticStruct()));
+		ChunkFragmentBitSet.Add<T>();
+	}
+
+	template<typename T>
+	T& AddChunkFragment_GetRef()
+	{
+		static_assert(TIsDerivedFrom<T, FLWChunkComponent>::IsDerived, "Given struct doesn't represent a valid chunk fragment type. Make sure to inherit from FLWChunkComonent or one of its child-types.");
+		ChunkFragmentBitSet.Add<T>();
+		const int32 Index = ChunkFragments.AddUnique(FInstancedStruct(T::StaticStruct()));
+		return ChunkFragments[Index].template GetMutable<T>();
+	}
+
+	FString DebugGetDescription(UEntitySubsystem* EntitySubsystem = nullptr) const;
+	FString DebugGetArchetypeDescription(UEntitySubsystem& EntitySubsystem) const;
+
+private:
+	FArchetypeHandle Archetype;
+
+	FMassUniqueFragmentCollection FragmentCollection;
+	FLWTagBitSet TagBitSet;
+	FLWChunkComponentBitSet ChunkFragmentBitSet;
+	
+	UPROPERTY()
+	FRuntimePipeline InitializationPipeline;
+	
+	UPROPERTY()
+	FRuntimePipeline DeinitializationPipeline;
+
+	/** An mirror of FMassUniqueFragmentCollection.Fragments' types, used for adding fragment types instead of fragment instances */
+	UPROPERTY()
+	TArray<const UScriptStruct*> FragmentTypes;
+
+	UPROPERTY()
+	TArray<FInstancedStruct> ChunkFragments;
+
+	// These functions will be called to initialize entity's UObject-based fragments
+	TArray<FObjectFragmentInitializerFunction> ObjectInitializers;
+
+	FMassEntityTemplateID TemplateID;
+};
diff --git a/Engine/Plugins/Runtime/MassGameplay/Source/MassSpawner/Public/MassEntityTemplateRegistry.h b/Engine/Plugins/Runtime/MassGameplay/Source/MassSpawner/Public/MassEntityTemplateRegistry.h
new file mode 100644
--- /dev/null
+++ b/Engine/Plugins/Runtime/MassGameplay/Source/MassSpawner/Public/MassEntityTemplateRegistry.h
@@ -0,0 +1,252 @@
+// Copyright Epic Games, Inc. All Rights Reserved.
+
+#pragma once
+
+#include "CoreMinimal.h"
+#include "UObject/Object.h"
+#include "MassCommonTypes.h"
+#include "MassEntityTemplate.h"
+#include "MassTranslatorRegistry.h"
+#include "MassTranslator.h"
+#include "MassEntityTemplateRegistry.generated.h"
+
+
+class UWorld;
+class UEntitySubsystem;
+class UMassTranslator;
+
+enum class EFragmentInitialization : uint8
+{
+	DefaultInitializer,
+	NoInitializer
+};
+
+struct FMassEntityTemplateBuildContext
+{
+	explicit FMassEntityTemplateBuildContext(FMassEntityTemplate& InTemplate)
+		: Template(InTemplate)
+	{}
+
+	//----------------------------------------------------------------------//
+	// Fragments 
+	//----------------------------------------------------------------------//
+	template<typename T>
+	T& AddFragment_GetRef()
+	{
+		return Template.GetMutableFragmentCollection().Add_GetRef<T>();
+	}
+
+	template<typename T>
+	T& AddFragmentWithDefaultInitializer_GetRef()
+	{
+		AddDefaultInitializer<T>();
+		return AddFragment_GetRef<T>();
+	}
+
+	template<typename T>
+	T& AddFragmentWithInitializer_GetRef(const UPipeProcessor& Initializer)
+	{
+		AddInitializer(Initializer);
+		return AddFragment_GetRef<T>();
+	}
+
+	template<typename T>
+	void AddFragment()
+	{
+		static_assert(TIsDerivedFrom<T, FLWComponentData>::IsDerived, "Given struct doesn't represent a valid fragment type. Make sure to inherit from FLWComponentData or one of its child-types.");
+		Template.GetMutableFragmentCollection().Add(T::StaticStruct());
+	}
+
+	template<typename T>
+	void AddFragmentWithDefaultInitializer()
+	{
+		static_assert(TIsDerivedFrom<T, FLWComponentData>::IsDerived, "Given struct doesn't represent a valid fragment type. Make sure to inherit from FLWComponentData or one of its child-types.");
+		AddFragment<T>();
+		AddDefaultInitializer<T>();
+	}
+
+	template<typename T>
+	void AddFragmentWithInitializer(const UPipeProcessor& Initializer)
+	{
+		static_assert(TIsDerivedFrom<T, FLWComponentData>::IsDerived, "Given struct doesn't represent a valid fragment type. Make sure to inherit from FLWComponentData or one of its child-types.");
+		AddFragment<T>();
+		AddInitializer(Initializer);
+	}
+
+	void AddFragment(const FStructView InFragment)
+	{ 
+		Template.GetMutableFragmentCollection().Add(InFragment);
+	}
+
+	void AddFragmentWithDefaultInitializer(FStructView InFragment)
+	{
+		AddFragment(InFragment);
+		if (InFragment.GetScriptStruct())
+		{
+			AddDefaultInitializer(*InFragment.GetScriptStruct());
+		}
+	}
+
+	void AddFragmentWithInitializer(FStructView InFragment, const UPipeProcessor& Initializer)
+	{
+		AddFragment(InFragment);
+		AddInitializer(Initializer);
+	}
+
+	template<typename T>
+	void AddTag()
+	{
+		Template.AddTag<T>();
+	}
+
+	void AddTag(const UScriptStruct& TagType)
+	{
+		Template.AddTag(TagType);
+	}
+
+	template<typename T>
+	void AddChunkFragment()
+	{
+		Template.AddChunkFragment<T>();
+	}
+
+	template<typename T>
+	T& AddChunkFragment_GetRef()
+	{
+		return Template.AddChunkFragment_GetRef<T>();
+	}
+
+	//----------------------------------------------------------------------//
+	// Initializers
+	//----------------------------------------------------------------------//
+	void AddInitializer(const UPipeProcessor& Initializer) { Handlers.Initializers.AddUnique(&Initializer); }
+	
+	void AddDefaultInitializer(const UScriptStruct& FragmentType)
+	{
+		const UMassTranslatorRegistry& Registry = UMassTranslatorRegistry::Get();
+		if (const UMassFragmentInitializer* Initializer = Registry.GetFragmentInitializer(FragmentType))
+		{
+			Handlers.Initializers.AddUnique(Initializer);
+		}
+	}
+
+	template<typename T>
+	void AddDefaultInitializer()
+	{
+		const UMassTranslatorRegistry& Registry = UMassTranslatorRegistry::Get();
+		if (const UMassFragmentInitializer* Initializer = Registry.GetFragmentInitializer(*T::StaticStruct()))
+		{
+			Handlers.Initializers.AddUnique(Initializer);
+		}
+	}
+
+	//----------------------------------------------------------------------//
+	// Deinitializers
+	//----------------------------------------------------------------------//
+	void AddDeinitializer(const UPipeProcessor& Deinitializer) { Handlers.Deinitializers.AddUnique(&Deinitializer); }
+
+	void AddDefaultDeinitializer(const UScriptStruct& FragmentType)
+	{
+		const UMassTranslatorRegistry& Registry = UMassTranslatorRegistry::Get();
+		if (const UMassFragmentDestructor* Deinitializer = Registry.GetFragmentDestructor(FragmentType))
+		{
+			Handlers.Deinitializers.AddUnique(Deinitializer);
+		}
+	}
+
+	template<typename T>
+	void AddDefaultDeinitializer()
+	{
+		const UMassTranslatorRegistry& Registry = UMassTranslatorRegistry::Get();
+		if (const UMassFragmentDestructor* Deinitializer = Registry.GetFragmentDestructor(*T::StaticStruct()))
+		{
+			Handlers.Deinitializers.AddUnique(Deinitializer);
+		}
+	}
+
+	//----------------------------------------------------------------------//
+	// Translators
+	//----------------------------------------------------------------------//
+	template<typename T>
+	void AddTranslator()
+	{
+		TArray<UScriptStruct*> TagTypes;
+		GetDefault<T>()->AppendRequiredTags(Template.GetMutableTags());
+	}
+
+	//----------------------------------------------------------------------//
+	// Template access
+	//----------------------------------------------------------------------//
+	FMassEntityTemplateID GetTemplateID() const { return Template.GetTemplateID(); }
+	FMassUniqueFragmentCollection& GetMutableFragments() { return Template.GetMutableFragmentCollection(); }
+	TArray<FMassEntityTemplate::FObjectFragmentInitializerFunction>& GetMutableObjectFragmentInitializers() { return Template.GetMutableObjectFragmentInitializers(); }
+
+	FMassObjectHandlers Handlers;
+protected:
+	FMassEntityTemplate& Template;
+};
+
+/** @todo document 
+ *  Note: the Object-Instance entity templates do not support having multiple instances of the same class that differ in
+ *  composition (like with dynamically added components). The first registered instance will determine the composition 
+ *  of the entity template.
+ */
+UCLASS()
+class MASSSPAWNER_API UMassEntityTemplateRegistry : public UObject
+{
+	GENERATED_BODY()
+public:
+	// @todo consider TFunction instead
+	DECLARE_DELEGATE_ThreeParams(FStructToTemplateBuilderDelegate, const UWorld* /*World*/, const FInstancedStruct& /*InStructInstance*/, FMassEntityTemplateBuildContext& /*BuildContext*/);
+	DECLARE_DELEGATE_FourParams(FClassToTemplateBuilderDelegate, const UWorld* /*World*/, const UClass& /*InClass*/, const UObject* /*ObjectInstance*/, FMassEntityTemplateBuildContext& /*BuildContext*/);
+
+	UMassEntityTemplateRegistry() = default;
+	virtual void BeginDestroy() override;
+	virtual UWorld* GetWorld() const override;
+
+	static FStructToTemplateBuilderDelegate& FindOrAdd(const UScriptStruct& DataType);
+	static FClassToTemplateBuilderDelegate& FindOrAdd(const UClass& Class);
+
+	const FMassEntityTemplate* FindOrBuildStructTemplate(const FInstancedStruct& StructInstance);
+	const FMassEntityTemplate* FindOrBuildInstanceTemplate(const UObject& Instance);
+	const FMassEntityTemplate& FindInstanceTemplateChecked(const UObject& Instance) const;
+	const FMassEntityTemplate* FindOrBuildClassTemplate(const UClass& Class);
+
+	/** Removes all the cached template instances */
+	void DebugReset();
+
+	const FMassEntityTemplate* FindTemplateFromTemplateID(FMassEntityTemplateID TemplateID) const;
+	FMassEntityTemplate* FindMutableTemplateFromTemplateID(FMassEntityTemplateID TemplateID);
+
+	FMassEntityTemplate& CreateTemplate(const uint32 HashLookup, FMassEntityTemplateID TemplateID);
+	void DestroyTemplate(const uint32 HashLookup, FMassEntityTemplateID TemplateID);
+	void InitializeEntityTemplate(FMassEntityTemplate& OutTemplate) const;
+
+protected:
+	/** A function implementing common path for UClass and UObject template building 
+	 *	@param Class type of class for which we need to build an entity template
+	 *	@param Instance if given means we're building a template for a runtime UObject instance. It will be passed down to template builder
+	 */
+	const FMassEntityTemplate* BuildClassTemplate(const UClass& Class, const UObject* Instance = nullptr);
+
+	/** @return true if a template has been built, false otherwise */
+	bool BuildTemplateImpl(const FStructToTemplateBuilderDelegate& Builder, const FInstancedStruct& StructInstance, FMassEntityTemplate& OutTemplate);
+
+	static FClassToTemplateBuilderDelegate* GetBuilderForClass(const UClass& Class);
+
+protected:
+	static TMap<const UScriptStruct*, FStructToTemplateBuilderDelegate> StructBasedBuilders;
+	static TMap<const UClass*, FClassToTemplateBuilderDelegate> ClassBasedBuilders;
+	
+	/** 
+	 *  Map from a hash to a FMassEntityTemplateID
+	 *  For build from FInstancedStruct it will be the Combined hash of the UScriptStruct FName and FTemplateRegistryHelpers::CalcHash()
+	 *  For build from UClass or UObject it will be the hash of the UClass FName
+	 *  This hash will not be deterministic between server and clients
+	 */
+	TMap<uint32, FMassEntityTemplateID> LookupTemplateIDMap;
+
+	UPROPERTY(Transient)
+	TMap<FMassEntityTemplateID, FMassEntityTemplate> TemplateIDToTemplateMap;
+};
+
diff --git a/Engine/Plugins/Runtime/MassGameplay/Source/MassSpawner/Public/MassEntityTraitBase.h b/Engine/Plugins/Runtime/MassGameplay/Source/MassSpawner/Public/MassEntityTraitBase.h
new file mode 100644
--- /dev/null
+++ b/Engine/Plugins/Runtime/MassGameplay/Source/MassSpawner/Public/MassEntityTraitBase.h
@@ -0,0 +1,31 @@
+// Copyright Epic Games, Inc. All Rights Reserved.
+
+#pragma once 
+
+#include "CoreMinimal.h"
+#include "MassEntityTraitBase.generated.h"
+
+struct FMassEntityTemplateBuildContext;
+
+/**
+ * Base class for Mass Entity Traits.
+ * An entity trait is a set of fragments that create a logical trait tha makes sense to end use (i.e. replication, visualization).
+ * The template building method allows to configure some fragments based on properties or cached values.
+ * For example, a fragment can be added based on a referenced asset, or some memory hungry settings can be
+ * cached and just and index stored on a fragment.
+ */
+UCLASS(Abstract)
+class MASSSPAWNER_API UMassEntityTraitBase : public UObject
+{
+	GENERATED_BODY()
+
+public:
+
+	/** Appends items into the entity template required for the trait. */
+	virtual void BuildTemplate(FMassEntityTemplateBuildContext& BuildContext, UWorld& World) const PURE_VIRTUAL(UMassEntityTraitBase::BuildTemplate, return; );
+
+	virtual void DestroyTemplate() const {}
+
+	/** Called when all Traits have had BuildTemplate() called. */
+	virtual void ValidateTemplate(FMassEntityTemplateBuildContext& BuildContext, UWorld& World) const {};
+};
diff --git a/Engine/Plugins/Runtime/MassGameplay/Source/MassSpawner/Public/MassEntityZoneGraphSpawnPointsGenerator.h b/Engine/Plugins/Runtime/MassGameplay/Source/MassSpawner/Public/MassEntityZoneGraphSpawnPointsGenerator.h
new file mode 100644
--- /dev/null
+++ b/Engine/Plugins/Runtime/MassGameplay/Source/MassSpawner/Public/MassEntityZoneGraphSpawnPointsGenerator.h
@@ -0,0 +1,38 @@
+// Copyright Epic Games, Inc. All Rights Reserved.
+
+#pragma once 
+
+#include "CoreMinimal.h"
+#include "MassEntitySpawnPointsGeneratorBase.h"
+#include "ZoneGraphTypes.h"
+
+#include "MassEntityZoneGraphSpawnPointsGenerator.generated.h"
+
+class AZoneGraphData;
+
+/**
+ * Describes the SpawnPoints Generator when we want to spawn directly on Zone Graph
+ */
+UCLASS(BlueprintType, meta=(DisplayName="ZoneGraph SpawnPoints Generator"))
+class MASSSPAWNER_API UMassEntityZoneGraphSpawnPointsGenerator : public UMassEntitySpawnPointsGeneratorBase
+{
+	GENERATED_BODY()
+
+public:
+	virtual void GenerateSpawnPoints(UObject& QueryOwner, int32 Count, FFinishedGeneratingSpawnPointsSignature& FinishedGeneratingSpawnPointsDelegate) const override;
+
+protected:
+	void GeneratePointsForZoneGraphData(const ::AZoneGraphData& ZoneGraphData, TArray<FVector>& Locations, const FRandomStream& RandomStream) const;
+
+	/** Tags to filter which lane to use to generate points on */
+	UPROPERTY(EditAnywhere, Category = "ZoneGraph Generator Config")
+	FZoneGraphTagFilter TagFilter;
+
+	/** Minimum gap for spawning entities on a given lanes */
+	UPROPERTY(EditAnywhere, Category = "ZoneGraph Generator Config")
+	float MinGap = 100;
+
+	/** Minimum gap for spawning entities on a given lanes */
+	UPROPERTY(EditAnywhere, Category = "ZoneGraph Generator Config")
+	float MaxGap = 300;
+};
\ No newline at end of file
diff --git a/Engine/Plugins/Runtime/MassGameplay/Source/MassSpawner/Public/MassSpawner.h b/Engine/Plugins/Runtime/MassGameplay/Source/MassSpawner/Public/MassSpawner.h
new file mode 100644
--- /dev/null
+++ b/Engine/Plugins/Runtime/MassGameplay/Source/MassSpawner/Public/MassSpawner.h
@@ -0,0 +1,158 @@
+// Copyright Epic Games, Inc. All Rights Reserved.
+
+#pragma once 
+
+#include "GameFramework/Actor.h"
+#include "EnvironmentQuery/EnvQueryTypes.h"
+#include "InstancedStruct.h"
+#include "Engine/World.h"
+#include "MassSpawnerTypes.h"
+#include "MassSpawner.generated.h"
+
+class UMassEntitySpawnPointsGeneratorConfigAsset;
+class UPipeSchematic;
+
+/** A spawner you can put on a map and configure it to spawn different things */
+UCLASS(hidecategories = (Object, Actor, Input, Rendering, LOD, Cooking, Collision, HLOD, Partition))
+class MASSSPAWNER_API AMassSpawner : public AActor
+{
+	GENERATED_BODY()
+public:
+	AMassSpawner();
+
+protected:
+	virtual void BeginPlay() override;
+	virtual void EndPlay(const EEndPlayReason::Type EndPlayReason) override;
+	virtual void PostInitProperties() override;
+	virtual void PostRegisterAllComponents() override;
+	virtual void BeginDestroy() override;
+
+public:
+	TArray<FInstancedStruct>& GetSpawnSets() { return SpawnSets; }
+
+#if WITH_EDITOR
+	UFUNCTION(BlueprintCallable, CallInEditor, Category = "Debug")
+	void DEBUG_Spawn();
+
+	/** Nuke all the entities */
+	UFUNCTION(BlueprintCallable, CallInEditor, Category = "Debug")
+	void DEBUG_Clear();
+#endif // WITH_EDITOR
+
+protected:
+#if WITH_EDITOR
+	virtual void PostEditChangeProperty(FPropertyChangedEvent& PropertyChangedEvent) override;
+	void ValidateSpawnSets();
+#endif // WITH_EDITOR
+
+	void RegisterEntityTemplates();
+
+public:
+	/**
+	 * Starts the spawning of all the agent types of this spawner
+	 */
+	UFUNCTION(BlueprintCallable, Category = "Spawning")
+	void DoSpawning();
+
+	/**
+	 * Despawn all mass agent that was spawned by this spawner
+	 */
+	UFUNCTION(BlueprintCallable, Category = "Spawning")
+	void DoDespawning();
+
+	UFUNCTION(BlueprintCallable, Category = "Spawning")
+	void ClearTemplates();
+
+	UFUNCTION(BlueprintCallable, Category = "Spawning")
+	void UnloadConfig();
+
+	/**
+	 * If given entity has been spawned by this MassSpawner instance then it will get destroyed and all the book keeping 
+	 * updated. Otherwise the call has no effect.
+	 * @return true if the entity got removed. False otherwise.
+	 */
+	bool DespawnEntity(const FLWEntity Entity);
+
+	/**
+	 * Scales the spawning counts (this only works when bUseEntityConfig is set)
+	 * @param Scale is the number to multiply the all counts of each agent types 
+	 */
+	UFUNCTION(BlueprintCallable, Category = "Spawning")
+	void ScaleSpawningCount(float Scale) { SpawningCountScale = Scale; }
+
+protected:
+	void OnEQSQueryFinished(TSharedPtr<FEnvQueryResult> Result);
+	void OnAdjustTickSchematics(UWorld* World, TArray<TSoftObjectPtr<UPipeSchematic>>& InOutTickSchematics);
+	void OnPostWorldInit(UWorld* World, const UWorld::InitializationValues);
+	void SpawnAtLocations(const TArray<FVector>& Locations);
+	void OnSpawnPointGenerationFinished(const TArray<FVector>& Locations, FMassSpawnPointGenerator* FinishedGenerator);
+
+	int32 GetSpawnCount() const;
+
+protected:
+
+	struct FSpawnedEntities
+	{
+		FMassEntityTemplateID TemplateID;
+		TArray<FLWEntity> Entities;
+	};
+
+	// @TODO CONSIDER HAVING ONE PER ENTRY IN SpawnSets
+	UPROPERTY(Category = "Mass|Spawn", EditAnywhere, meta = (BaseStruct = "MassSpawnConfigBase", EditCondition = "!bUseEntityConfig"))
+	TArray<FInstancedStruct> SpawnSets;
+
+	UPROPERTY(EditAnywhere, Category = "Mass|Spawn")
+	bool bUseEntityConfig = false;
+
+	UPROPERTY(EditAnywhere, Category = "Mass|Spawn", meta = (EditCondition = "bUseEntityConfig"))
+	int32 Count;
+
+	UPROPERTY(BlueprintReadWrite, EditAnywhere, Category = "Mass|Spawn", meta = (EditCondition = "bUseEntityConfig"))
+	TArray<FMassSpawnedEntityType> EntityTypes;
+
+	UPROPERTY(EditAnywhere, Category = "Mass|Spawn")
+	bool bUseSpawnPointsGenerators = false;
+
+	/** Asset that describes the way we want to generate SpawnPoints of the entity */
+	UPROPERTY(EditAnywhere, Category = "Mass|Spawn", meta = (EditCondition = "bUseSpawnPointsGenerators"))
+	TArray<FMassSpawnPointGenerator> SpawnPointsGenerators;
+
+	UPROPERTY(Category = "Mass|Spawn", EditAnywhere, meta = (EditCondition = "!bUseSpawnPointsGenerators"))
+	FEQSParametrizedQueryExecutionRequest EQSRequest;
+
+	UPROPERTY(Category = "Mass|Spawn", EditAnywhere)
+	uint32 bAutoSpawnOnBeginPlay : 1;
+
+	/** By default TickSchematics will be appended to the simulation's schematics. If this property is set to true the
+	 *  TickSchematics will override the original simulation schematics */
+	UPROPERTY(Category = "Mass|Simulation", EditAnywhere)
+	uint32 bOverrideSchematics : 1;
+
+	UPROPERTY(Category = "Mass|Simulation", EditAnywhere)
+	TArray<TSoftObjectPtr<UPipeSchematic>> TickSchematics;
+
+	/** Scale of the spawning count */
+	float SpawningCountScale = 1.0f;
+
+	FDelegateHandle TickSchematicsAdjustHandle;
+
+	FDelegateHandle SimulationStartedHandle;
+
+	FDelegateHandle OnPostWorldInitDelegateHandle;
+
+	TArray<FSpawnedEntities> AllSpawnedEntities;
+
+	TArray<FVector> AllGeneratedLocations;
+
+#if WITH_EDITORONLY_DATA
+private:
+	UPROPERTY()
+	UBillboardComponent* SpriteComponent;
+#endif // WITH_EDITORONLY_DATA
+ };
+
+namespace UE::MassSpawner
+{
+	MASSSPAWNER_API extern float ScalabilitySpawnDensityMultiplier;
+}
+
diff --git a/Engine/Plugins/Runtime/MassGameplay/Source/MassSpawner/Public/MassSpawnerSubsystem.h b/Engine/Plugins/Runtime/MassGameplay/Source/MassSpawner/Public/MassSpawnerSubsystem.h
new file mode 100644
--- /dev/null
+++ b/Engine/Plugins/Runtime/MassGameplay/Source/MassSpawner/Public/MassSpawnerSubsystem.h
@@ -0,0 +1,67 @@
+// Copyright Epic Games, Inc. All Rights Reserved.
+
+#pragma once
+
+#include "EntitySubsystem.h"
+#include "Subsystems/WorldSubsystem.h"
+#include "MassSpawnerSubsystem.generated.h"
+
+class UPipeEntitySubsystem;
+struct FMassSpawnConfigBase;
+struct FMassEntityTemplate;
+class UMassEntityTemplateRegistry;
+struct FInstancedStruct;
+struct FStructView;
+struct FMassEntityTemplateID;
+class UMassSimulationSubsystem;
+
+UCLASS()
+class MASSSPAWNER_API UMassSpawnerSubsystem : public UWorldSubsystem
+{
+	GENERATED_BODY()
+public:
+
+	void SpawnEntities(FMassEntityTemplateID TemplateID, const FMassSpawnConfigBase& SpawnConfig, const FStructView& AuxData, TArray<FLWEntity>& OutEntities);
+
+	/** Finds the appropriate entity template for ActorInstance and calls the other SpawnEntities implementation 
+	 *  @return true if spawning was successful, false otherwise. The failure might come from there being no EntityTemplate 
+	 *  @param ActorInstance instance of an actor for which we want to spawn entities for
+	 *  @param NumberToSpawn number of entities to spawn
+	 *  @param OutEntities where the IDs of created entities get added. Note that the contents of OutEntities get overridden by the function.
+	 *  for ActorInstance. See logs for more details. */
+	bool SpawnEntities(const AActor& ActorInstance, const uint32 NumberToSpawn, TArray<FLWEntity>& OutEntities) const;
+
+	/** Spawns entities of the kind described by the given EntityTemplate. The spawned entities are fully initialized
+	 *  meaning the EntityTemplate.InitializationPipeline gets run for all spawned entities.
+	 *  @param EntityTemplate template to use for spawning entities
+	 *  @param NumberToSpawn number of entities to spawn
+	 *  @param OutEntities where the IDs of created entities get added. Note that the contents of OutEntities get overridden by the function.
+	 *  @return true if spawning was successful, false otherwise. In case of failure see logs for more details. */
+	bool SpawnEntities(const FMassEntityTemplate& EntityTemplate, const uint32 NumberToSpawn, TArray<FLWEntity>& OutEntities) const;
+
+	void SpawnCollection(TArrayView<FInstancedStruct> Collection, const int32 Count, const FStructView& AuxData = FStructView());
+
+	void DestroyEntities(const FMassEntityTemplateID TemplateID, TConstArrayView<FLWEntity> Entities);
+
+	UMassEntityTemplateRegistry& GetTemplateRegistryInstance() const { check(TemplateRegistryInstance); return *TemplateRegistryInstance; }
+
+	void RegisterCollection(TArrayView<FInstancedStruct> Collection);
+
+	const FMassEntityTemplate* GetMassEntityTemplate(FMassEntityTemplateID TemplateID) const;
+
+protected:
+	virtual void Initialize(FSubsystemCollectionBase& Collection) override;
+	virtual void PostInitialize() override;
+
+	void DoSpawning(const FMassEntityTemplate& EntityTemplate, const FMassSpawnConfigBase& Data, const FStructView& AuxData, TArray<FLWEntity>& OutEntities) const;
+
+	UPROPERTY()
+	UPipeEntitySubsystem* EntitySystem;
+
+	UPROPERTY()
+	UMassEntityTemplateRegistry* TemplateRegistryInstance;
+
+	UPROPERTY()
+	UMassSimulationSubsystem* SimulationSystem;
+};
+
