AggregateTicking and Pipe merged into Engine/Plugins/Runtime/MassEntity

diff --git a/Engine/Plugins/Runtime/MassGameplay/Source/MassSpawner/Public/MassSpawnerTypes.h b/Engine/Plugins/Runtime/MassGameplay/Source/MassSpawner/Public/MassSpawnerTypes.h
new file mode 100644
--- /dev/null
+++ b/Engine/Plugins/Runtime/MassGameplay/Source/MassSpawner/Public/MassSpawnerTypes.h
@@ -0,0 +1,106 @@
+// Copyright Epic Games, Inc. All Rights Reserved.
+
+#pragma once 
+
+#include "Engine/DataTable.h"
+#include "Logging/LogMacros.h"
+#include "MassCommonTypes.h"
+#include "InstancedStruct.h"
+#include "MassEntityTemplate.h"
+#include "MassEntitySpawnPointsGeneratorBase.h"
+
+#include "MassSpawnerTypes.generated.h"
+
+
+DECLARE_LOG_CATEGORY_EXTERN(LogMassSpawner, Warning, All);
+
+class APawn;
+class UCurveFloat;
+class UMassEntityConfigAsset;
+
+USTRUCT()
+struct FMassSpawnAuxData
+{
+	GENERATED_BODY()
+
+	TArray<FTransform> Transforms;
+};
+
+USTRUCT()
+struct FMassSpawnConfigBase
+{
+	GENERATED_BODY()
+
+	UPROPERTY(EditAnywhere, Category = "Mass|Spawn", meta = (BaseStruct = "LWComponentData"))
+	TArray<FInstancedStruct> AdditionalDataFragments;
+
+	UPROPERTY(EditAnywhere, Category = "Mass|Spawn", meta = (ClampMin = 0, UIMin = 0))
+	int32 MinNumber = 0;
+
+	UPROPERTY(EditAnywhere, Category = "Mass|Spawn", meta = (ClampMin = 1, UIMin = 1))
+	int32 MaxNumber = 1;
+
+	UPROPERTY(EditAnywhere, Category = "Mass|Spawn", meta = (ClampMin = 0.0, UIMin = 0.0))
+	float Radius = 0.f;
+
+	friend uint32 GetTypeHash(const FMassSpawnConfigBase& Instance)
+	{
+		return HashCombine(HashCombine(GetTypeHash(Instance.MinNumber), GetTypeHash(Instance.MaxNumber)), GetTypeHash(Instance.Radius));
+	}
+};
+
+/**
+ * Describes an entity type to spawn.
+ */
+USTRUCT(BlueprintType)
+struct FMassSpawnedEntityType
+{
+	GENERATED_BODY()
+
+	/** Asset that describes the entity */
+	UPROPERTY(BlueprintReadWrite, EditAnywhere, Category = "Mass|Spawn")
+	TSoftObjectPtr<UMassEntityConfigAsset> EntityConfig;
+
+	/** Proportion of the count that should be this agent type, (the proportions will be normalized with other sibling agent types) */
+	UPROPERTY(BlueprintReadWrite, EditAnywhere, Category = "Mass|Spawn", meta = (ClampMin = 0.0f, UIMin = 0.0f))
+	float Proportion = 1.0f;
+
+	void UnloadEntityConfig()
+	{
+		EntityConfigPtr = nullptr;
+		EntityConfig.ResetWeakPtr();
+	}
+
+	const UMassEntityConfigAsset* GetEntityConfig() const;
+	UMassEntityConfigAsset* GetEntityConfig();
+
+private:
+	UPROPERTY()
+	mutable UMassEntityConfigAsset* EntityConfigPtr = nullptr;
+};
+
+USTRUCT(BlueprintType)
+struct FMassSpawnPointGenerator
+{
+	GENERATED_BODY()
+
+	/** The Generator to use to generate the spawn the points */
+	UPROPERTY(BlueprintReadWrite, EditAnywhere, Category = "Mass|Generator")
+	TSubclassOf<UMassEntitySpawnPointsGeneratorBase> GeneratorClass;
+
+	/** Proportion of the spawn points that should be generated by this generator, (the proportions will be normalized with other sibling generators) */
+	UPROPERTY(BlueprintReadWrite, EditAnywhere, Category = "Mass|Generator", meta = (ClampMin = 0.0f, UIMin = 0.0f))
+	float Proportion = 1.0f;
+
+	/** Runtime value to know if we received the generated points for this generator */
+	bool bPointsGenerated = false;
+};
+
+USTRUCT()
+struct MASSSPAWNER_API FDataFragment_ReplicationTemplateID : public FLWComponentData
+{
+	GENERATED_BODY()
+
+	UPROPERTY(Transient)
+	FMassEntityTemplateID ID;
+};
diff --git a/Engine/Plugins/Runtime/MassGameplay/Source/MassSpawner/Public/MassTranslator.h b/Engine/Plugins/Runtime/MassGameplay/Source/MassSpawner/Public/MassTranslator.h
new file mode 100644
--- /dev/null
+++ b/Engine/Plugins/Runtime/MassGameplay/Source/MassSpawner/Public/MassTranslator.h
@@ -0,0 +1,95 @@
+// Copyright Epic Games, Inc. All Rights Reserved.
+
+#pragma once
+
+#include "MassProcessor.h"
+#include "MassCommonTypes.h"
+#include "LWComponentTypes.h"
+#include "MassTranslator.generated.h"
+
+
+class UPipeEntitySubsystem;
+
+UENUM()
+enum class EMassTranslationDirection : uint8
+{
+	None = 0,
+	InitializationOnly = 1,
+	ActorToMass = 1 << 1,
+	MassToActor = 1 << 2,
+	BothWays = ActorToMass | MassToActor
+};
+ENUM_CLASS_FLAGS(EMassTranslationDirection);
+
+/** 
+ *  A class that's responsible for translation between UObjects and Mass. A translator knows how to initialize 
+ *  fragments related to the UClass that the given translator cares about. It can also be used at runtime to 
+ *  copy values from UObjects to fragments and back.
+ */
+UCLASS(abstract)
+class MASSSPAWNER_API UMassTranslator : public UPipeProcessor
+{
+	GENERATED_BODY()
+
+protected:
+	UMassTranslator();
+public:
+	/** Fetches the FComponentTag-derived types required by this Translator. And entity needs these tags to be 
+	 *  processed by this Translator instance. 
+	 * 
+	 *  @todo might want this function on the PipeProcessor level. TBD
+	 * 
+	 *  @param OutTagTypes tag types will be appended to this array. No uniqueness checks are performed. */
+	void AppendRequiredTags(FLWTagBitSet& InOutTags) const { InOutTags += RequiredTags; }
+
+protected:
+	void AddRequiredTagsToQuery(FLWComponentQuery& EntityQuery);
+
+protected:
+	/** These are the tag fragments expected by this translator that let other code (like entity traits) hint to 
+	 *  the system which translators they'd want their entity to be processed by. */
+	FLWTagBitSet RequiredTags;
+};
+
+UCLASS(abstract)
+class MASSSPAWNER_API UMassFragmentBuilder : public UPipeProcessor
+{
+	GENERATED_BODY()
+
+public:
+	UMassFragmentBuilder();
+
+	/** Override to register this fragment builder to the registry */
+	virtual void Register() PURE_VIRTUAL(UMassFragmentBuilder::Register, );
+
+protected:
+	virtual void PostInitProperties() override;
+
+protected:
+	/** Set in class' constructor and determines for which Fragment type this given UMassFragmentBuilder will be used 
+	 *  as the default initializer/deinitializer (as returned by the UMassTranslatorRegistry). If set to null will
+	 *  require the user to manually add this UMassFragmentBuilder as initializer/deinitializer 
+	 *  (@see FMassEntityTemplateBuildContext.AddInitializer and FMassEntityTemplateBuildContext.AddDeinitializer) */
+	UPROPERTY()
+	UScriptStruct* FragmentType;
+};
+
+UCLASS(abstract)
+class MASSSPAWNER_API UMassFragmentInitializer : public UMassFragmentBuilder
+{
+	GENERATED_BODY()
+
+public:
+	/** adds this FragmentInitializer to the MassTranslatorRegistry to be used to initialize fragments of FragmentType */
+	virtual void Register();
+};
+
+UCLASS(abstract)
+class MASSSPAWNER_API UMassFragmentDestructor : public UMassFragmentBuilder
+{
+	GENERATED_BODY()
+
+public:
+	/** adds this FragmentDestructor to the MassTranslatorRegistry to be used to initialize fragments of FragmentType */
+	virtual void Register();
+};
\ No newline at end of file
diff --git a/Engine/Plugins/Runtime/MassGameplay/Source/MassSpawner/Public/MassTranslatorRegistry.h b/Engine/Plugins/Runtime/MassGameplay/Source/MassSpawner/Public/MassTranslatorRegistry.h
new file mode 100644
--- /dev/null
+++ b/Engine/Plugins/Runtime/MassGameplay/Source/MassSpawner/Public/MassTranslatorRegistry.h
@@ -0,0 +1,50 @@
+// Copyright Epic Games, Inc. All Rights Reserved.
+
+#pragma once
+
+#include "CoreMinimal.h"
+#include "UObject/Object.h"
+#include "Templates/SubclassOf.h"
+#include "InstancedStruct.h"
+#include "MassTranslatorRegistry.generated.h"
+
+class UScriptStruct;
+class AActor;
+class UActorComponent;
+class UPipeProcessor;
+class UMassTranslator;
+class UMassObjectFragmentCreator;
+class UMassFragmentInitializer;
+class UMassFragmentDestructor;
+struct FMassUniqueFragmentCollection;
+
+/** @todo rename the class, no longer hosting Translators, just Initializers and Deinitializers/Destructors
+ *  @todo brief description of the class
+ */
+UCLASS()
+class MASSSPAWNER_API UMassTranslatorRegistry : public UObject
+{
+	GENERATED_BODY()
+public:
+	UMassTranslatorRegistry();
+
+	static UMassTranslatorRegistry& GetMutable() { return *GetMutableDefault<UMassTranslatorRegistry>(); }
+	static const UMassTranslatorRegistry& Get() { return *GetDefault<UMassTranslatorRegistry>(); }
+
+	/** UMassFragmentInitializer - given a UScriptStruct sets its initial values to some defaults or randomized 
+	 *  values(implementation - specific). It's a kind of a PipeProcessor */
+	void RegisterFragmentInitializer(const UScriptStruct& FragmentType, TSubclassOf<UMassFragmentInitializer> FragmentInitializerClass);
+	const UMassFragmentInitializer* GetFragmentInitializer(const UScriptStruct& FragmentType) const;
+
+	/** UMassFragmentDestructor - given a UScriptStruct destroys values and do cleanup */
+	void RegisterFragmentDestructor(const UScriptStruct& FragmentType, TSubclassOf<UMassFragmentDestructor> FragmentDestructorClass);
+	const UMassFragmentDestructor* GetFragmentDestructor(const UScriptStruct& FragmentType) const;
+
+protected:
+
+	UPROPERTY()
+	TMap<const UScriptStruct*, TSubclassOf<UMassFragmentInitializer>> FragmentInitializersMap;
+
+	UPROPERTY()
+	TMap<const UScriptStruct*, TSubclassOf<UMassFragmentDestructor>> FragmentDestructorsMap;
+};
