Cleaned up MassEntitySubsystem instance fetching patterns. It's supposed to be only accessed directly from the world as a world's subsystem.

diff --git a/Engine/Plugins/AI/MassAI/Source/MassAIDebug/Private/GameplayDebuggerCategory_Mass.cpp b/Engine/Plugins/AI/MassAI/Source/MassAIDebug/Private/GameplayDebuggerCategory_Mass.cpp
--- a/Engine/Plugins/AI/MassAI/Source/MassAIDebug/Private/GameplayDebuggerCategory_Mass.cpp
+++ b/Engine/Plugins/AI/MassAI/Source/MassAIDebug/Private/GameplayDebuggerCategory_Mass.cpp
@@ -165,7 +165,7 @@ void FGameplayDebuggerCategory_Mass::CollectData(APlayerController* OwnerPC, AAc
 		}
 	}
 
-	UMassEntitySubsystem* EntitySystem = UMassEntitySubsystem::GetCurrent(World);
+	UMassEntitySubsystem* EntitySystem = UWorld::GetSubsystem<UMassEntitySubsystem>(World);
 	if (EntitySystem)
 	{
 		AddTextLine(FString::Printf(TEXT("{Green}Entities count active{grey}/all: {white}%d{grey}/%d"), EntitySystem->DebugGetEntityCount(), EntitySystem->DebugGetEntityCount()));
diff --git a/Engine/Plugins/AI/MassAI/Source/MassAIMovement/Private/MassMovementSubsystem.cpp b/Engine/Plugins/AI/MassAI/Source/MassAIMovement/Private/MassMovementSubsystem.cpp
--- a/Engine/Plugins/AI/MassAI/Source/MassAIMovement/Private/MassMovementSubsystem.cpp
+++ b/Engine/Plugins/AI/MassAI/Source/MassAIMovement/Private/MassMovementSubsystem.cpp
@@ -19,6 +19,6 @@ void UMassMovementSubsystem::Initialize(FSubsystemCollectionBase& Collection)
 	Collection.InitializeDependency<UMassSimulationSubsystem>();
 	ZoneGraphSubsystem = Collection.InitializeDependency<UZoneGraphSubsystem>();
 
-	EntitySubsystem = UMassEntitySubsystem::GetCurrent(GetWorld());
+	EntitySubsystem = UWorld::GetSubsystem<UMassEntitySubsystem>(GetWorld());
 }
 
diff --git a/Engine/Plugins/Runtime/MassEntity/Source/MassEntity/Private/MassEntityDebug.cpp b/Engine/Plugins/Runtime/MassEntity/Source/MassEntity/Private/MassEntityDebug.cpp
--- a/Engine/Plugins/Runtime/MassEntity/Source/MassEntity/Private/MassEntityDebug.cpp
+++ b/Engine/Plugins/Runtime/MassEntity/Source/MassEntity/Private/MassEntityDebug.cpp
@@ -3,7 +3,8 @@
 #include "MassEntityDebug.h"
 #include "Misc/OutputDevice.h"
 #include "MassProcessor.h"
-
+#include "Engine/World.h"
+#include "Engine/Engine.h"
 
 DEFINE_ENUM_TO_STRING(EPipeProcessingPhase);
 
@@ -38,7 +39,8 @@ FAutoConsoleCommandWithWorldArgsAndOutputDevice PrintEntityFragmentsCmd(
 	FConsoleCommandWithWorldArgsAndOutputDeviceDelegate::CreateLambda(
 		[](const TArray<FString>& Params, UWorld* World, FOutputDevice& Ar)
 		{
-			if (UMassEntitySubsystem* EntitySystem = UMassEntitySubsystem::GetCurrent(World))
+			check(World);
+			if (UMassEntitySubsystem* EntitySystem = World->GetSubsystem<UMassEntitySubsystem>())
 			{
 				int32 Index = INDEX_NONE;
 				if (LexTryParseString<int32>(Index, *Params[0]))
@@ -76,7 +78,7 @@ FAutoConsoleCommandWithWorldArgsAndOutputDevice LogArchetypesCmd(
 				LexToString(World->WorldType),
 				*ToString(World->GetNetMode()));
 
-			if (UMassEntitySubsystem* EntitySystem = UMassEntitySubsystem::GetCurrent(World))
+			if (UMassEntitySubsystem* EntitySystem = World->GetSubsystem<UMassEntitySubsystem>())
 			{
 				bool bIncludeEmpty = true;
 				if (Params.Num())
@@ -103,7 +105,8 @@ FAutoConsoleCommandWithWorld RecacheQueries(
 	TEXT("Forces EntityQueries to recache their valid archetypes"),
 	FConsoleCommandWithWorldDelegate::CreateLambda([](UWorld* InWorld)
 		{
-			if (UMassEntitySubsystem* System = UMassEntitySubsystem::GetCurrent(InWorld))
+			check(InWorld);
+			if (UMassEntitySubsystem* System = InWorld->GetSubsystem<UMassEntitySubsystem>())
 			{
 				System->DebugForceArchetypeDataVersionBump();
 			}
diff --git a/Engine/Plugins/Runtime/MassEntity/Source/MassEntity/Private/MassProcessingPhase.cpp b/Engine/Plugins/Runtime/MassEntity/Source/MassEntity/Private/MassProcessingPhase.cpp
--- a/Engine/Plugins/Runtime/MassEntity/Source/MassEntity/Private/MassProcessingPhase.cpp
+++ b/Engine/Plugins/Runtime/MassEntity/Source/MassEntity/Private/MassProcessingPhase.cpp
@@ -172,7 +172,7 @@ void UPipeProcessingPhaseManager::InitializePhases(UObject& InProcessorOwner)
 
 void UPipeProcessingPhaseManager::Start(UWorld& World)
 {
-	EntitySubsystem = UMassEntitySubsystem::GetCurrent(&World);
+	EntitySubsystem = UWorld::GetSubsystem<UMassEntitySubsystem>(&World);
 
 	if (ensure(EntitySubsystem))
 	{
diff --git a/Engine/Plugins/Runtime/MassEntity/Source/MassEntity/Public/MassEntitySubsystem.h b/Engine/Plugins/Runtime/MassEntity/Source/MassEntity/Public/MassEntitySubsystem.h
--- a/Engine/Plugins/Runtime/MassEntity/Source/MassEntity/Public/MassEntitySubsystem.h
+++ b/Engine/Plugins/Runtime/MassEntity/Source/MassEntity/Public/MassEntitySubsystem.h
@@ -79,11 +79,6 @@ public:
 	virtual void Initialize(FSubsystemCollectionBase& Collection) override;
 	virtual void Deinitialize() override;
 	//~End of USubsystem interface
-
-	static FORCEINLINE UMassEntitySubsystem* GetCurrent(UWorld* World)
-	{
-		return World ? World->GetSubsystem<UMassEntitySubsystem>() : nullptr;
-	}
 	
 	/** 
 	 * A special, relaxed but slower version of CreateArchetype functions that allows ComponentAngTagsList to contain 
diff --git a/Engine/Plugins/Runtime/MassGameplay/Source/MassActors/Private/MassActorHelper.cpp b/Engine/Plugins/Runtime/MassGameplay/Source/MassActors/Private/MassActorHelper.cpp
--- a/Engine/Plugins/Runtime/MassGameplay/Source/MassActors/Private/MassActorHelper.cpp
+++ b/Engine/Plugins/Runtime/MassGameplay/Source/MassActors/Private/MassActorHelper.cpp
@@ -11,7 +11,7 @@ namespace UE::MassActor
 	bool AddEntityTagToActor(const AActor& Actor, const UScriptStruct& TagType)
 	{
 		UWorld* World = Actor.GetWorld();
-		UMassEntitySubsystem* EntitySystem = UMassEntitySubsystem::GetCurrent(World);
+		UMassEntitySubsystem* EntitySystem = UWorld::GetSubsystem<UMassEntitySubsystem>(World);
 		UMassActorSubsystem* MassActorSubsystem = UWorld::GetSubsystem<UMassActorSubsystem>(World);
 		if (EntitySystem && MassActorSubsystem)
 		{
@@ -36,7 +36,7 @@ namespace UE::MassActor
 	bool RemoveEntityTagFromActor(const AActor& Actor, const UScriptStruct& TagType)
 	{
 		UWorld* World = Actor.GetWorld();
-		UMassEntitySubsystem* EntitySystem = UMassEntitySubsystem::GetCurrent(World);
+		UMassEntitySubsystem* EntitySystem = UWorld::GetSubsystem<UMassEntitySubsystem>(World);
 		UMassActorSubsystem* MassActorSubsystem = UWorld::GetSubsystem<UMassActorSubsystem>(World);
 		if (EntitySystem && MassActorSubsystem)
 		{
diff --git a/Engine/Plugins/Runtime/MassGameplay/Source/MassActors/Private/MassActorSubsystem.cpp b/Engine/Plugins/Runtime/MassGameplay/Source/MassActors/Private/MassActorSubsystem.cpp
--- a/Engine/Plugins/Runtime/MassGameplay/Source/MassActors/Private/MassActorSubsystem.cpp
+++ b/Engine/Plugins/Runtime/MassGameplay/Source/MassActors/Private/MassActorSubsystem.cpp
@@ -62,7 +62,7 @@ void UMassActorSubsystem::Initialize(FSubsystemCollectionBase& Collection)
 	// making sure UMassSimulationSubsystem gets created before the MassActorManager
 	Collection.InitializeDependency<UMassSimulationSubsystem>();
 	
-	EntitySystem = UMassEntitySubsystem::GetCurrent(GetWorld());
+	EntitySystem = UWorld::GetSubsystem<UMassEntitySubsystem>(GetWorld());
 }
 
 FMassHandle UMassActorSubsystem::GetHandleFromActor(const TObjectKey<const AActor> Actor)
diff --git a/Engine/Plugins/Runtime/MassGameplay/Source/MassActors/Private/MassAgentComponent.cpp b/Engine/Plugins/Runtime/MassGameplay/Source/MassActors/Private/MassAgentComponent.cpp
--- a/Engine/Plugins/Runtime/MassGameplay/Source/MassActors/Private/MassAgentComponent.cpp
+++ b/Engine/Plugins/Runtime/MassGameplay/Source/MassActors/Private/MassAgentComponent.cpp
@@ -162,7 +162,7 @@ void UMassAgentComponent::SetEntityHandleInternal(FMassHandle NewHandle)
 
 #if	UE_REPLICATION_COMPILE_SERVER_CODE
 	// Fetch NetID if it exist
-	if (const UMassEntitySubsystem* EntitySubsystem = UMassEntitySubsystem::GetCurrent(GetWorld()))
+	if (const UMassEntitySubsystem* EntitySubsystem = UWorld::GetSubsystem<UMassEntitySubsystem>(GetWorld()))
 	{
 		if (const FMassNetworkIDFragment* NetIDFragment = EntitySubsystem->GetComponentDataPtr<FMassNetworkIDFragment>(AgentHandle.GetLWEntity()))
 		{
@@ -188,7 +188,7 @@ void UMassAgentComponent::SetEntityHandleInternal(FMassHandle NewHandle)
 	}
 
 	// Sync up with mass
-	if (const UMassEntitySubsystem* EntitySubsystem = UMassEntitySubsystem::GetCurrent(GetWorld()))
+	if (const UMassEntitySubsystem* EntitySubsystem = UWorld::GetSubsystem<UMassEntitySubsystem>(GetWorld()))
 	{
 		if (IsNetSimulating())
 		{
@@ -277,7 +277,7 @@ void UMassAgentComponent::ClearEntityHandleInternal()
 	}
 
 	// Sync up with mass
-	if (const UMassEntitySubsystem* EntitySubsystem = UMassEntitySubsystem::GetCurrent(GetWorld()))
+	if (const UMassEntitySubsystem* EntitySubsystem = UWorld::GetSubsystem<UMassEntitySubsystem>(GetWorld()))
 	{
 		if (IsNetSimulating())
 		{
@@ -352,7 +352,7 @@ void UMassAgentComponent::DebugCheckStateConsistency()
 			MASSAGENT_CHECK(bValidAgentHandle, TEXT("Expecting a valid mass agent handle in state %s"), *UEnum::GetValueAsString(State));
 			if (bValidAgentHandle)
 			{
-				if (UMassEntitySubsystem* EntitySubsystem = UMassEntitySubsystem::GetCurrent(GetWorld()))
+				if (UMassEntitySubsystem* EntitySubsystem = UWorld::GetSubsystem<UMassEntitySubsystem>(GetWorld()))
 				{
 					const bool bIsValidEntity = EntitySubsystem->IsEntityValid(AgentHandle.GetLWEntity());
 					MASSAGENT_CHECK(bIsValidEntity, TEXT("Exepecting a valid entity in state"), *UEnum::GetValueAsString(State))
@@ -418,7 +418,7 @@ void UMassAgentComponent::KillEntity(const bool bDestroyActor)
 
 	AActor* Owner = GetOwner();
 	UWorld* World = GetWorld();
-	UMassEntitySubsystem* EntitySubsystem = UMassEntitySubsystem::GetCurrent(World);
+	UMassEntitySubsystem* EntitySubsystem = UWorld::GetSubsystem<UMassEntitySubsystem>(World);
 	if (Owner == nullptr || EntitySubsystem == nullptr)
 	{
 		return;
diff --git a/Engine/Plugins/Runtime/MassGameplay/Source/MassActors/Private/MassAgentSubsystem.cpp b/Engine/Plugins/Runtime/MassGameplay/Source/MassActors/Private/MassAgentSubsystem.cpp
--- a/Engine/Plugins/Runtime/MassGameplay/Source/MassActors/Private/MassAgentSubsystem.cpp
+++ b/Engine/Plugins/Runtime/MassGameplay/Source/MassActors/Private/MassAgentSubsystem.cpp
@@ -46,7 +46,7 @@ void UMassAgentSubsystem::Initialize(FSubsystemCollectionBase& Collection)
 #endif // UE_REPLICATION_COMPILE_CLIENT_CODE
 	
 	UWorld* World = GetWorld();
-	EntitySystem = UMassEntitySubsystem::GetCurrent(World);
+	EntitySystem = UWorld::GetSubsystem<UMassEntitySubsystem>(World);
 	SpawnerSystem = UWorld::GetSubsystem<UMassSpawnerSubsystem>(World);
 
 	SimulationSystem = UWorld::GetSubsystem<UMassSimulationSubsystem>(World);
diff --git a/Engine/Plugins/Runtime/MassGameplay/Source/MassReplication/Private/MassClientBubbleSerializerBase.cpp b/Engine/Plugins/Runtime/MassGameplay/Source/MassReplication/Private/MassClientBubbleSerializerBase.cpp
--- a/Engine/Plugins/Runtime/MassGameplay/Source/MassReplication/Private/MassClientBubbleSerializerBase.cpp
+++ b/Engine/Plugins/Runtime/MassGameplay/Source/MassReplication/Private/MassClientBubbleSerializerBase.cpp
@@ -12,7 +12,7 @@ void FMassClientBubbleSerializerBase::InitializeForWorld(UWorld& InWorld)
 
 	ReplicationManager = InWorld.GetSubsystem<UMassReplicationManager>();
 
-	EntitySystem = UMassEntitySubsystem::GetCurrent(World);
+	EntitySystem = UWorld::GetSubsystem<UMassEntitySubsystem>(World);
 	check(EntitySystem);
 
 	checkf(ClientHandler, TEXT("ClientHandler must be setup! See SetClientHandler()"));
diff --git a/Engine/Plugins/Runtime/MassGameplay/Source/MassReplication/Private/MassReplicationManager.cpp b/Engine/Plugins/Runtime/MassGameplay/Source/MassReplication/Private/MassReplicationManager.cpp
--- a/Engine/Plugins/Runtime/MassGameplay/Source/MassReplication/Private/MassReplicationManager.cpp
+++ b/Engine/Plugins/Runtime/MassGameplay/Source/MassReplication/Private/MassReplicationManager.cpp
@@ -24,7 +24,7 @@ void UMassReplicationManager::Initialize(FSubsystemCollectionBase& Collection)
 
 	check(MassLODManager);
 
-	EntitySystem = UMassEntitySubsystem::GetCurrent(GetWorld());
+	EntitySystem = UWorld::GetSubsystem<UMassEntitySubsystem>(GetWorld());
 }
 
 void UMassReplicationManager::Deinitialize()
diff --git a/Engine/Plugins/Runtime/MassGameplay/Source/MassRepresentation/Private/MassRepresentationProcessor.cpp b/Engine/Plugins/Runtime/MassGameplay/Source/MassRepresentation/Private/MassRepresentationProcessor.cpp
--- a/Engine/Plugins/Runtime/MassGameplay/Source/MassRepresentation/Private/MassRepresentationProcessor.cpp
+++ b/Engine/Plugins/Runtime/MassGameplay/Source/MassRepresentation/Private/MassRepresentationProcessor.cpp
@@ -39,7 +39,7 @@ void UMassRepresentationProcessor::Initialize(UObject& Owner)
 	Super::Initialize(Owner);
 
 	World = Owner.GetWorld();
-	PipeEntitySubsystem = UMassEntitySubsystem::GetCurrent(World);
+	PipeEntitySubsystem = UWorld::GetSubsystem<UMassEntitySubsystem>(World);
 	RepresentationSubsystem = UWorld::GetSubsystem<UMassRepresentationSubsystem>(World);
 
 	// Calculate the default representation when actor isn't spawned yet.
@@ -288,7 +288,7 @@ void UMassRepresentationProcessor::TeleportActor(const FTransform& Transform, AA
 
 void UMassRepresentationProcessor::ReleaseAnyActorOrCancelAnySpawning(UMassRepresentationSubsystem& RepresentationSubsystem, const FMassHandle MassAgent)
 {
-	if (const UMassEntitySubsystem* EntitySubsystem = UMassEntitySubsystem::GetCurrent(RepresentationSubsystem.GetWorld()))
+	if (const UMassEntitySubsystem* EntitySubsystem = UWorld::GetSubsystem<UMassEntitySubsystem>(RepresentationSubsystem.GetWorld()))
 	{
 		if (FDataFragment_Actor* ActorInfo = EntitySubsystem->GetComponentDataPtr<FDataFragment_Actor>(MassAgent.GetLWEntity()))
 		{
diff --git a/Engine/Plugins/Runtime/MassGameplay/Source/MassRepresentation/Private/MassRepresentationSubsystem.cpp b/Engine/Plugins/Runtime/MassGameplay/Source/MassRepresentation/Private/MassRepresentationSubsystem.cpp
--- a/Engine/Plugins/Runtime/MassGameplay/Source/MassRepresentation/Private/MassRepresentationSubsystem.cpp
+++ b/Engine/Plugins/Runtime/MassGameplay/Source/MassRepresentation/Private/MassRepresentationSubsystem.cpp
@@ -302,7 +302,7 @@ void UMassRepresentationSubsystem::Initialize(FSubsystemCollectionBase& Collecti
 
 	if (UWorld* World = GetWorld())
 	{
-		EntitySubsystem = UMassEntitySubsystem::GetCurrent(World);
+		EntitySubsystem = UWorld::GetSubsystem<UMassEntitySubsystem>(World);
 		ActorSpawnerSubsystem = World->GetSubsystem<UMassActorSpawnerSubsystem>();
 		WorldPartitionSubsystem = World->GetSubsystem<UWorldPartitionSubsystem>();
 
diff --git a/Engine/Plugins/Runtime/MassGameplay/Source/MassSimulation/Private/MassSimulationSubsystem.cpp b/Engine/Plugins/Runtime/MassGameplay/Source/MassSimulation/Private/MassSimulationSubsystem.cpp
--- a/Engine/Plugins/Runtime/MassGameplay/Source/MassSimulation/Private/MassSimulationSubsystem.cpp
+++ b/Engine/Plugins/Runtime/MassGameplay/Source/MassSimulation/Private/MassSimulationSubsystem.cpp
@@ -72,8 +72,8 @@ void UMassSimulationSubsystem::Initialize(FSubsystemCollectionBase& Collection)
 {
 	Super::Initialize(Collection);
 
-	CachedEntitySubsystem = Collection.InitializeDependency<UMassEntitySubsystem>();
-	check(CachedEntitySubsystem);
+	EntitySubsystem = Collection.InitializeDependency<UMassEntitySubsystem>();
+	check(EntitySubsystem);
 	
 	GetOnProcessingPhaseStarted(EPipeProcessingPhase::PrePhysics).AddUObject(this, &UMassSimulationSubsystem::OnProcessingPhaseStarted, EPipeProcessingPhase::PrePhysics);
 }
@@ -148,10 +148,10 @@ void UMassSimulationSubsystem::OnProcessingPhaseStarted(const float DeltaSeconds
 		case EPipeProcessingPhase::PrePhysics:
 			{
 				TRACE_CPUPROFILER_EVENT_SCOPE(DoEntityCompation);
-				check(CachedEntitySubsystem);
+				check(EntitySubsystem);
 				if (UE::MassSimulation::bDoEntityCompaction)
 				{
-					CachedEntitySubsystem->DoEntityCompaction(GET_MASS_CONFIG_VALUE(DesiredEntityCompactionTimeSlicePerTick));
+					EntitySubsystem->DoEntityCompaction(GET_MASS_CONFIG_VALUE(DesiredEntityCompactionTimeSlicePerTick));
 				}
 			}
 			break;
diff --git a/Engine/Plugins/Runtime/MassGameplay/Source/MassSimulation/Public/MassSimulationSubsystem.h b/Engine/Plugins/Runtime/MassGameplay/Source/MassSimulation/Public/MassSimulationSubsystem.h
--- a/Engine/Plugins/Runtime/MassGameplay/Source/MassSimulation/Public/MassSimulationSubsystem.h
+++ b/Engine/Plugins/Runtime/MassGameplay/Source/MassSimulation/Public/MassSimulationSubsystem.h
@@ -25,7 +25,7 @@ public:
 	
 	UMassSimulationSubsystem(const FObjectInitializer& ObjectInitializer = FObjectInitializer::Get());
 
-	UMassEntitySubsystem* GetEntitySubsystem() const { return CachedEntitySubsystem; }
+	//UMassEntitySubsystem* GetEntitySubsystem() const { return CachedEntitySubsystem; }
 	const UPipeProcessingPhaseManager& GetPhaseManager() const { check(PhaseManager); return *PhaseManager; }
 
 	FPipeProcessingPhase::FOnPhaseEvent& GetOnProcessingPhaseStarted(const EPipeProcessingPhase Phase) const;
@@ -58,7 +58,7 @@ protected:
 protected:
 
 	UPROPERTY()
-	UMassEntitySubsystem* CachedEntitySubsystem;
+	UMassEntitySubsystem* EntitySubsystem;
 
 	UPROPERTY()
 	UPipeProcessingPhaseManager* PhaseManager;
diff --git a/Engine/Plugins/Runtime/MassGameplay/Source/MassSpawner/Private/MassEntityTemplateRegistry.cpp b/Engine/Plugins/Runtime/MassGameplay/Source/MassSpawner/Private/MassEntityTemplateRegistry.cpp
--- a/Engine/Plugins/Runtime/MassGameplay/Source/MassSpawner/Private/MassEntityTemplateRegistry.cpp
+++ b/Engine/Plugins/Runtime/MassGameplay/Source/MassSpawner/Private/MassEntityTemplateRegistry.cpp
@@ -138,7 +138,7 @@ bool UMassEntityTemplateRegistry::BuildTemplateImpl(const FStructToTemplateBuild
 		InitializeEntityTemplate(OutTemplate);
 
 		UE_VLOG(this, LogMassSpawner, Log, TEXT("Created entity template for %s:\n%s"), *GetNameSafe(StructInstance.GetScriptStruct())
-			, *OutTemplate.DebugGetDescription(UMassEntitySubsystem::GetCurrent(World)));
+			, *OutTemplate.DebugGetDescription(UWorld::GetSubsystem<UMassEntitySubsystem>(World)));
 
 		return true;
 	}
@@ -252,7 +252,7 @@ void UMassEntityTemplateRegistry::InitializeEntityTemplate(FMassEntityTemplate&
 
 	UWorld* World = GetWorld();
 	// find or create template
-	UMassEntitySubsystem* EntitySys = UMassEntitySubsystem::GetCurrent(World);
+	UMassEntitySubsystem* EntitySys = UWorld::GetSubsystem<UMassEntitySubsystem>(World);
 	check(EntitySys);
 
 	FTemplateRegistryHelpers::FragmentInstancesToTypes(OutTemplate.GetFragments(), OutTemplate.GetMutableFragmentTypes());
diff --git a/Engine/Plugins/Runtime/MassGameplay/Source/MassSpawner/Private/MassSpawnerSubsystem.cpp b/Engine/Plugins/Runtime/MassGameplay/Source/MassSpawner/Private/MassSpawnerSubsystem.cpp
--- a/Engine/Plugins/Runtime/MassGameplay/Source/MassSpawner/Private/MassSpawnerSubsystem.cpp
+++ b/Engine/Plugins/Runtime/MassGameplay/Source/MassSpawner/Private/MassSpawnerSubsystem.cpp
@@ -93,7 +93,7 @@ void UMassSpawnerSubsystem::Initialize(FSubsystemCollectionBase& Collection)
 	Collection.InitializeDependency<UMassSimulationSubsystem>();
 
 	SimulationSystem = UWorld::GetSubsystem<UMassSimulationSubsystem>(GetWorld());
-	EntitySystem = UMassEntitySubsystem::GetCurrent(GetWorld());
+	EntitySystem = UWorld::GetSubsystem<UMassEntitySubsystem>(GetWorld());
 	check(EntitySystem);
 }
 
