AggregateTicking and Pipe type renames:

diff --git a/Engine/Plugins/Runtime/MassGameplay/Source/MassActors/Private/Translators/MassCapsuleComponentTranslators.cpp b/Engine/Plugins/Runtime/MassGameplay/Source/MassActors/Private/Translators/MassCapsuleComponentTranslators.cpp
--- a/Engine/Plugins/Runtime/MassGameplay/Source/MassActors/Private/Translators/MassCapsuleComponentTranslators.cpp
+++ b/Engine/Plugins/Runtime/MassGameplay/Source/MassActors/Private/Translators/MassCapsuleComponentTranslators.cpp
@@ -20,13 +20,13 @@ UMassCapsuleTransformToMassTranslator::UMassCapsuleTransformToMassTranslator()
 void UMassCapsuleTransformToMassTranslator::ConfigureQueries()
 {
 	AddRequiredTagsToQuery(EntityQuery);
-	EntityQuery.AddRequirement<FDataFragment_CapsuleComponentWrapper>(ELWComponentAccess::ReadOnly);
-	EntityQuery.AddRequirement<FDataFragment_Transform>(ELWComponentAccess::ReadWrite);
+	EntityQuery.AddRequirement<FDataFragment_CapsuleComponentWrapper>(EMassFragmentAccess::ReadOnly);
+	EntityQuery.AddRequirement<FDataFragment_Transform>(EMassFragmentAccess::ReadWrite);
 }
 
-void UMassCapsuleTransformToMassTranslator::Execute(UMassEntitySubsystem& EntitySubsystem, FLWComponentSystemExecutionContext& Context)
+void UMassCapsuleTransformToMassTranslator::Execute(UMassEntitySubsystem& EntitySubsystem, FMassExecutionContext& Context)
 {
-	EntityQuery.ForEachEntityChunk(EntitySubsystem, Context, [this](FLWComponentSystemExecutionContext& Context)
+	EntityQuery.ForEachEntityChunk(EntitySubsystem, Context, [this](FMassExecutionContext& Context)
 		{
 			const TConstArrayView<FDataFragment_CapsuleComponentWrapper> CapsuleComponentList = Context.GetComponentView<FDataFragment_CapsuleComponentWrapper>();
 			const TArrayView<FDataFragment_Transform> LocationList = Context.GetMutableComponentView<FDataFragment_Transform>();
@@ -54,13 +54,13 @@ UMassTransformToActorCapsuleTranslator::UMassTransformToActorCapsuleTranslator()
 void UMassTransformToActorCapsuleTranslator::ConfigureQueries()
 {
 	AddRequiredTagsToQuery(EntityQuery);
-	EntityQuery.AddRequirement<FDataFragment_CapsuleComponentWrapper>(ELWComponentAccess::ReadWrite);
-	EntityQuery.AddRequirement<FDataFragment_Transform>(ELWComponentAccess::ReadOnly);
+	EntityQuery.AddRequirement<FDataFragment_CapsuleComponentWrapper>(EMassFragmentAccess::ReadWrite);
+	EntityQuery.AddRequirement<FDataFragment_Transform>(EMassFragmentAccess::ReadOnly);
 }
 
-void UMassTransformToActorCapsuleTranslator::Execute(UMassEntitySubsystem& EntitySubsystem, FLWComponentSystemExecutionContext& Context)
+void UMassTransformToActorCapsuleTranslator::Execute(UMassEntitySubsystem& EntitySubsystem, FMassExecutionContext& Context)
 {
-	EntityQuery.ForEachEntityChunk(EntitySubsystem, Context, [this](FLWComponentSystemExecutionContext& Context)
+	EntityQuery.ForEachEntityChunk(EntitySubsystem, Context, [this](FMassExecutionContext& Context)
 		{
 			const TArrayView<FDataFragment_CapsuleComponentWrapper> CapsuleComponentList = Context.GetMutableComponentView<FDataFragment_CapsuleComponentWrapper>();
 			const TConstArrayView<FDataFragment_Transform> LocationList = Context.GetComponentView<FDataFragment_Transform>();
diff --git a/Engine/Plugins/Runtime/MassGameplay/Source/MassActors/Private/Translators/MassCharacterMovementTranslators.cpp b/Engine/Plugins/Runtime/MassGameplay/Source/MassActors/Private/Translators/MassCharacterMovementTranslators.cpp
--- a/Engine/Plugins/Runtime/MassGameplay/Source/MassActors/Private/Translators/MassCharacterMovementTranslators.cpp
+++ b/Engine/Plugins/Runtime/MassGameplay/Source/MassActors/Private/Translators/MassCharacterMovementTranslators.cpp
@@ -23,14 +23,14 @@ UMassCharacterMovementToMassTranslator::UMassCharacterMovementToMassTranslator()
 void UMassCharacterMovementToMassTranslator::ConfigureQueries()
 {
 	AddRequiredTagsToQuery(EntityQuery);
-	EntityQuery.AddRequirement<FDataFragment_CharacterMovementComponentWrapper>(ELWComponentAccess::ReadOnly);
-	EntityQuery.AddRequirement<FDataFragment_Transform>(ELWComponentAccess::ReadWrite);
-	EntityQuery.AddRequirement<FMassVelocityFragment>(ELWComponentAccess::ReadWrite);
+	EntityQuery.AddRequirement<FDataFragment_CharacterMovementComponentWrapper>(EMassFragmentAccess::ReadOnly);
+	EntityQuery.AddRequirement<FDataFragment_Transform>(EMassFragmentAccess::ReadWrite);
+	EntityQuery.AddRequirement<FMassVelocityFragment>(EMassFragmentAccess::ReadWrite);
 }
 
-void UMassCharacterMovementToMassTranslator::Execute(UMassEntitySubsystem& EntitySubsystem, FLWComponentSystemExecutionContext& Context)
+void UMassCharacterMovementToMassTranslator::Execute(UMassEntitySubsystem& EntitySubsystem, FMassExecutionContext& Context)
 {
-	EntityQuery.ForEachEntityChunk(EntitySubsystem, Context, [this](FLWComponentSystemExecutionContext& Context)
+	EntityQuery.ForEachEntityChunk(EntitySubsystem, Context, [this](FMassExecutionContext& Context)
 	{
 		const TConstArrayView<FDataFragment_CharacterMovementComponentWrapper> ComponentList = Context.GetComponentView<FDataFragment_CharacterMovementComponentWrapper>();
 		const TArrayView<FDataFragment_Transform> LocationList = Context.GetMutableComponentView<FDataFragment_Transform>();
@@ -64,13 +64,13 @@ UMassCharacterMovementToActorTranslator::UMassCharacterMovementToActorTranslator
 void UMassCharacterMovementToActorTranslator::ConfigureQueries()
 {
 	AddRequiredTagsToQuery(EntityQuery);
-	EntityQuery.AddRequirement<FDataFragment_CharacterMovementComponentWrapper>(ELWComponentAccess::ReadWrite);
-	EntityQuery.AddRequirement<FMassVelocityFragment>(ELWComponentAccess::ReadOnly);
+	EntityQuery.AddRequirement<FDataFragment_CharacterMovementComponentWrapper>(EMassFragmentAccess::ReadWrite);
+	EntityQuery.AddRequirement<FMassVelocityFragment>(EMassFragmentAccess::ReadOnly);
 }
 
-void UMassCharacterMovementToActorTranslator::Execute(UMassEntitySubsystem& EntitySubsystem, FLWComponentSystemExecutionContext& Context)
+void UMassCharacterMovementToActorTranslator::Execute(UMassEntitySubsystem& EntitySubsystem, FMassExecutionContext& Context)
 {
-	EntityQuery.ForEachEntityChunk(EntitySubsystem, Context, [this](FLWComponentSystemExecutionContext& Context)
+	EntityQuery.ForEachEntityChunk(EntitySubsystem, Context, [this](FMassExecutionContext& Context)
 	{
 		const TArrayView<FDataFragment_CharacterMovementComponentWrapper> ComponentList = Context.GetMutableComponentView<FDataFragment_CharacterMovementComponentWrapper>();
 		const TConstArrayView<FMassVelocityFragment> VelocityList = Context.GetComponentView<FMassVelocityFragment>();
@@ -100,13 +100,13 @@ UMassCharacterOrientationToMassTranslator::UMassCharacterOrientationToMassTransl
 void UMassCharacterOrientationToMassTranslator::ConfigureQueries()
 {
 	AddRequiredTagsToQuery(EntityQuery);
-	EntityQuery.AddRequirement<FDataFragment_CharacterMovementComponentWrapper>(ELWComponentAccess::ReadOnly);
-	EntityQuery.AddRequirement<FDataFragment_Transform>(ELWComponentAccess::ReadWrite);
+	EntityQuery.AddRequirement<FDataFragment_CharacterMovementComponentWrapper>(EMassFragmentAccess::ReadOnly);
+	EntityQuery.AddRequirement<FDataFragment_Transform>(EMassFragmentAccess::ReadWrite);
 }
 
-void UMassCharacterOrientationToMassTranslator::Execute(UMassEntitySubsystem& EntitySubsystem, FLWComponentSystemExecutionContext& Context)
+void UMassCharacterOrientationToMassTranslator::Execute(UMassEntitySubsystem& EntitySubsystem, FMassExecutionContext& Context)
 {
-	EntityQuery.ForEachEntityChunk(EntitySubsystem, Context, [this](FLWComponentSystemExecutionContext& Context)
+	EntityQuery.ForEachEntityChunk(EntitySubsystem, Context, [this](FMassExecutionContext& Context)
 	{
 		const TConstArrayView<FDataFragment_CharacterMovementComponentWrapper> ComponentList = Context.GetComponentView<FDataFragment_CharacterMovementComponentWrapper>();
 		const TArrayView<FDataFragment_Transform> LocationList = Context.GetMutableComponentView<FDataFragment_Transform>();
@@ -140,13 +140,13 @@ UMassCharacterOrientationToActorTranslator::UMassCharacterOrientationToActorTran
 void UMassCharacterOrientationToActorTranslator::ConfigureQueries()
 {
 	AddRequiredTagsToQuery(EntityQuery);
-	EntityQuery.AddRequirement<FDataFragment_CharacterMovementComponentWrapper>(ELWComponentAccess::ReadWrite);
-	EntityQuery.AddRequirement<FDataFragment_Transform>(ELWComponentAccess::ReadOnly);
+	EntityQuery.AddRequirement<FDataFragment_CharacterMovementComponentWrapper>(EMassFragmentAccess::ReadWrite);
+	EntityQuery.AddRequirement<FDataFragment_Transform>(EMassFragmentAccess::ReadOnly);
 }
 
-void UMassCharacterOrientationToActorTranslator::Execute(UMassEntitySubsystem& EntitySubsystem, FLWComponentSystemExecutionContext& Context)
+void UMassCharacterOrientationToActorTranslator::Execute(UMassEntitySubsystem& EntitySubsystem, FMassExecutionContext& Context)
 {
-	EntityQuery.ForEachEntityChunk(EntitySubsystem, Context, [this](FLWComponentSystemExecutionContext& Context)
+	EntityQuery.ForEachEntityChunk(EntitySubsystem, Context, [this](FMassExecutionContext& Context)
 	{
 		const TArrayView<FDataFragment_CharacterMovementComponentWrapper> ComponentList = Context.GetMutableComponentView<FDataFragment_CharacterMovementComponentWrapper>();
 		const TConstArrayView<FDataFragment_Transform> TransformList = Context.GetComponentView<FDataFragment_Transform>();
@@ -178,12 +178,12 @@ UMassFragmentInitializer_NavLocation::UMassFragmentInitializer_NavLocation()
 
 void UMassFragmentInitializer_NavLocation::ConfigureQueries() 
 {
-	EntityQuery.AddRequirement<FDataFragment_NavLocation>(ELWComponentAccess::ReadWrite);
+	EntityQuery.AddRequirement<FDataFragment_NavLocation>(EMassFragmentAccess::ReadWrite);
 }
 
-void UMassFragmentInitializer_NavLocation::Execute(UMassEntitySubsystem& EntitySubsystem, FLWComponentSystemExecutionContext& Context)
+void UMassFragmentInitializer_NavLocation::Execute(UMassEntitySubsystem& EntitySubsystem, FMassExecutionContext& Context)
 {
-	EntityQuery.ForEachEntityChunk(EntitySubsystem, Context, [this](FLWComponentSystemExecutionContext& Context)
+	EntityQuery.ForEachEntityChunk(EntitySubsystem, Context, [this](FMassExecutionContext& Context)
 		{
 			const TArrayView<FDataFragment_NavLocation> NavLocationList = Context.GetMutableComponentView<FDataFragment_NavLocation>();
 			const int32 NumEntities = Context.GetEntitiesNum();
@@ -205,10 +205,10 @@ UMassFragmentInitializer_Transform::UMassFragmentInitializer_Transform()
 
 void UMassFragmentInitializer_Transform::ConfigureQueries()
 {
-	EntityQuery.AddRequirement<FDataFragment_Transform>(ELWComponentAccess::ReadWrite);
+	EntityQuery.AddRequirement<FDataFragment_Transform>(EMassFragmentAccess::ReadWrite);
 }
 
-void UMassFragmentInitializer_Transform::Execute(UMassEntitySubsystem& EntitySubsystem, FLWComponentSystemExecutionContext& Context)
+void UMassFragmentInitializer_Transform::Execute(UMassEntitySubsystem& EntitySubsystem, FMassExecutionContext& Context)
 {
 	const UWorld* World = EntitySubsystem.GetWorld();
 
@@ -235,7 +235,7 @@ void UMassFragmentInitializer_Transform::Execute(UMassEntitySubsystem& EntitySub
 		}
 
 		int32 NumRequiredSpawnTransforms = 0;
-		EntityQuery.ForEachEntityChunk(EntitySubsystem, Context, [&NumRequiredSpawnTransforms](const FLWComponentSystemExecutionContext& Context)
+		EntityQuery.ForEachEntityChunk(EntitySubsystem, Context, [&NumRequiredSpawnTransforms](const FMassExecutionContext& Context)
 			{
 				NumRequiredSpawnTransforms += Context.GetEntitiesNum();
 			});
@@ -254,7 +254,7 @@ void UMassFragmentInitializer_Transform::Execute(UMassEntitySubsystem& EntitySub
 			}
 		}
 
-		EntityQuery.ForEachEntityChunk(EntitySubsystem, Context, [&Transforms, this](FLWComponentSystemExecutionContext& Context)
+		EntityQuery.ForEachEntityChunk(EntitySubsystem, Context, [&Transforms, this](FMassExecutionContext& Context)
 			{
 				const TArrayView<FDataFragment_Transform> LocationList = Context.GetMutableComponentView<FDataFragment_Transform>();
 				const int32 NumEntities = Context.GetEntitiesNum();
diff --git a/Engine/Plugins/Runtime/MassGameplay/Source/MassActors/Private/Translators/MassSceneComponentLocationTranslator.cpp b/Engine/Plugins/Runtime/MassGameplay/Source/MassActors/Private/Translators/MassSceneComponentLocationTranslator.cpp
--- a/Engine/Plugins/Runtime/MassGameplay/Source/MassActors/Private/Translators/MassSceneComponentLocationTranslator.cpp
+++ b/Engine/Plugins/Runtime/MassGameplay/Source/MassActors/Private/Translators/MassSceneComponentLocationTranslator.cpp
@@ -18,13 +18,13 @@ UMassSceneComponentLocationToMassTranslator::UMassSceneComponentLocationToMassTr
 void UMassSceneComponentLocationToMassTranslator::ConfigureQueries()
 {
 	AddRequiredTagsToQuery(EntityQuery);
-	EntityQuery.AddRequirement<FMassSceneComponentWrapperFragment>(ELWComponentAccess::ReadOnly);
-	EntityQuery.AddRequirement<FDataFragment_Transform>(ELWComponentAccess::ReadWrite);
+	EntityQuery.AddRequirement<FMassSceneComponentWrapperFragment>(EMassFragmentAccess::ReadOnly);
+	EntityQuery.AddRequirement<FDataFragment_Transform>(EMassFragmentAccess::ReadWrite);
 }
 
-void UMassSceneComponentLocationToMassTranslator::Execute(UMassEntitySubsystem& EntitySubsystem, FLWComponentSystemExecutionContext& Context)
+void UMassSceneComponentLocationToMassTranslator::Execute(UMassEntitySubsystem& EntitySubsystem, FMassExecutionContext& Context)
 {
-	EntityQuery.ForEachEntityChunk(EntitySubsystem, Context, [this](FLWComponentSystemExecutionContext& Context)
+	EntityQuery.ForEachEntityChunk(EntitySubsystem, Context, [this](FMassExecutionContext& Context)
 	{
 		const TConstArrayView<FMassSceneComponentWrapperFragment> ComponentList = Context.GetComponentView<FMassSceneComponentWrapperFragment>();
 		const TArrayView<FDataFragment_Transform> LocationList = Context.GetMutableComponentView<FDataFragment_Transform>();
@@ -54,13 +54,13 @@ UMassSceneComponentLocationToActorTranslator::UMassSceneComponentLocationToActor
 void UMassSceneComponentLocationToActorTranslator::ConfigureQueries()
 {
 	AddRequiredTagsToQuery(EntityQuery);
-	EntityQuery.AddRequirement<FMassSceneComponentWrapperFragment>(ELWComponentAccess::ReadOnly);
-	EntityQuery.AddRequirement<FDataFragment_Transform>(ELWComponentAccess::ReadWrite);
+	EntityQuery.AddRequirement<FMassSceneComponentWrapperFragment>(EMassFragmentAccess::ReadOnly);
+	EntityQuery.AddRequirement<FDataFragment_Transform>(EMassFragmentAccess::ReadWrite);
 }
 
-void UMassSceneComponentLocationToActorTranslator::Execute(UMassEntitySubsystem& EntitySubsystem, FLWComponentSystemExecutionContext& Context)
+void UMassSceneComponentLocationToActorTranslator::Execute(UMassEntitySubsystem& EntitySubsystem, FMassExecutionContext& Context)
 {
-	EntityQuery.ForEachEntityChunk(EntitySubsystem, Context, [this](FLWComponentSystemExecutionContext& Context)
+	EntityQuery.ForEachEntityChunk(EntitySubsystem, Context, [this](FMassExecutionContext& Context)
 		{
 			const TConstArrayView<FMassSceneComponentWrapperFragment> ComponentList = Context.GetComponentView<FMassSceneComponentWrapperFragment>();
 			const TArrayView<FDataFragment_Transform> LocationList = Context.GetMutableComponentView<FDataFragment_Transform>();
diff --git a/Engine/Plugins/Runtime/MassGameplay/Source/MassActors/Private/Translators/MassTranslators_BehaviorTree.cpp b/Engine/Plugins/Runtime/MassGameplay/Source/MassActors/Private/Translators/MassTranslators_BehaviorTree.cpp
--- a/Engine/Plugins/Runtime/MassGameplay/Source/MassActors/Private/Translators/MassTranslators_BehaviorTree.cpp
+++ b/Engine/Plugins/Runtime/MassGameplay/Source/MassActors/Private/Translators/MassTranslators_BehaviorTree.cpp
@@ -11,5 +11,5 @@ UMassTranslator_BehaviorTree::UMassTranslator_BehaviorTree()
 
 void UMassTranslator_BehaviorTree::ConfigureQueries()
 {
-	EntityQuery.AddRequirement<FDataFragment_BehaviorTreeComponentWrapper>(ELWComponentAccess::ReadWrite);
+	EntityQuery.AddRequirement<FDataFragment_BehaviorTreeComponentWrapper>(EMassFragmentAccess::ReadWrite);
 }
diff --git a/Engine/Plugins/Runtime/MassGameplay/Source/MassActors/Public/MassActorHelper.h b/Engine/Plugins/Runtime/MassGameplay/Source/MassActors/Public/MassActorHelper.h
--- a/Engine/Plugins/Runtime/MassGameplay/Source/MassActors/Public/MassActorHelper.h
+++ b/Engine/Plugins/Runtime/MassGameplay/Source/MassActors/Public/MassActorHelper.h
@@ -22,7 +22,7 @@ namespace UE::MassActor
 	template<typename TagType>
 	bool AddEntityTagToActor(const AActor& Actor)
 	{
-		static_assert(TIsDerivedFrom<TagType, FComponentTag>::IsDerived, "Given struct doesn't represent a valid tag type.");
+		static_assert(TIsDerivedFrom<TagType, FMassTag>::IsDerived, "Given struct doesn't represent a valid tag type.");
 		return AddEntityTagToActor(Actor, *TagType::StaticStruct());
 	}
 
@@ -39,7 +39,7 @@ namespace UE::MassActor
 	template<typename TagType>
 	bool RemoveEntityTagFromActor(const AActor& Actor)
 	{
-		static_assert(TIsDerivedFrom<TagType, FComponentTag>::IsDerived, "Given struct doesn't represent a valid tag type.");
+		static_assert(TIsDerivedFrom<TagType, FMassTag>::IsDerived, "Given struct doesn't represent a valid tag type.");
 		return RemoveEntityTagFromActor(Actor, *TagType::StaticStruct());
 	}
 } // namespace UE::MassActor
diff --git a/Engine/Plugins/Runtime/MassGameplay/Source/MassActors/Public/MassActorSubsystem.h b/Engine/Plugins/Runtime/MassGameplay/Source/MassActors/Public/MassActorSubsystem.h
--- a/Engine/Plugins/Runtime/MassGameplay/Source/MassActors/Public/MassActorSubsystem.h
+++ b/Engine/Plugins/Runtime/MassGameplay/Source/MassActors/Public/MassActorSubsystem.h
@@ -104,7 +104,7 @@ protected:
 	virtual void Initialize(FSubsystemCollectionBase& Collection) override;
 	// USubsystem END
 	
-	TMap<TObjectKey<const AActor>, FLWEntity> ActorHandleMap;
+	TMap<TObjectKey<const AActor>, FMassEntityHandle> ActorHandleMap;
 	UE_MT_DECLARE_RW_ACCESS_DETECTOR(ActorHandleMapDetector);
 	
 	UPROPERTY()
diff --git a/Engine/Plugins/Runtime/MassGameplay/Source/MassActors/Public/MassAgentComponent.h b/Engine/Plugins/Runtime/MassGameplay/Source/MassActors/Public/MassAgentComponent.h
--- a/Engine/Plugins/Runtime/MassGameplay/Source/MassActors/Public/MassAgentComponent.h
+++ b/Engine/Plugins/Runtime/MassGameplay/Source/MassActors/Public/MassAgentComponent.h
@@ -100,8 +100,8 @@ public:
 	const FMassEntityConfig& GetEntityConfig() const { return EntityConfig; }
 	void SetEntityConfig(const FMassEntityConfig& InEntityConfig);
 
-	const FLWCompositionDescriptor& GetPuppetSpecificAddition() const { return PuppetSpecificAddition; }
-	FLWCompositionDescriptor& GetMutablePuppetSpecificAddition() { return PuppetSpecificAddition; }
+	const FMassCompositionDescriptor& GetPuppetSpecificAddition() const { return PuppetSpecificAddition; }
+	FMassCompositionDescriptor& GetMutablePuppetSpecificAddition() { return PuppetSpecificAddition; }
 
 	/** Registers the component with the owner effectively turning it on. Calling it multiple times won't break anything  */
 	UFUNCTION(BlueprintCallable, Category = Mass)
@@ -137,7 +137,7 @@ protected:
 	 *  Contains all the fragments added to the entity during puppet's initialization. Required for clean up when
 	 *  despawning puppet while the entity remains alive.
 	 */
-	FLWCompositionDescriptor PuppetSpecificAddition;
+	FMassCompositionDescriptor PuppetSpecificAddition;
 
 	UPROPERTY(EditAnywhere, Category = "Mass")
 	FMassEntityConfig EntityConfig;
@@ -145,7 +145,7 @@ protected:
 #if WITH_EDITORONLY_DATA
 	/** for debugging purposes only. The Entity used at runtime is stored in AgentHandle */
 	UPROPERTY(VisibleAnywhere, Category = "Mass")
-	FLWEntity DebugEntity;
+	FMassEntityHandle DebugEntity;
 #endif // WITH_EDITORONLY_DATA
 
 	FMassHandle AgentHandle;
diff --git a/Engine/Plugins/Runtime/MassGameplay/Source/MassActors/Public/MassAgentSubsystem.h b/Engine/Plugins/Runtime/MassGameplay/Source/MassActors/Public/MassAgentSubsystem.h
--- a/Engine/Plugins/Runtime/MassGameplay/Source/MassActors/Public/MassAgentSubsystem.h
+++ b/Engine/Plugins/Runtime/MassGameplay/Source/MassActors/Public/MassAgentSubsystem.h
@@ -113,13 +113,13 @@ protected:
 	void HandlePendingInitialization();
 
 	/** Bound to UMassSimulationSubsystem.OnProcessingPhaseStartedDelegate and called before every processing phase start */
-	void OnProcessingPhaseStarted(const float DeltaSeconds, const EPipeProcessingPhase Phase);
+	void OnProcessingPhaseStarted(const float DeltaSeconds, const EMassProcessingPhase Phase);
 
 	/** Callback registered to the replication manager when a mass agent is added to the replication (client only) */
-	void OnMassAgentAddedToReplication(FMassNetworkID NetID, FLWEntity Entity);
+	void OnMassAgentAddedToReplication(FMassNetworkID NetID, FMassEntityHandle Entity);
 
 	/** Callback registered to the replication manager when a mass agent is removed from the replication (client only) */
-	void OnMassAgentRemovedFromReplication(FMassNetworkID NetID, FLWEntity Entity);
+	void OnMassAgentRemovedFromReplication(FMassNetworkID NetID, FMassEntityHandle Entity);
 
 protected:
 	UPROPERTY()
diff --git a/Engine/Plugins/Runtime/MassGameplay/Source/MassActors/Public/Translators/MassCapsuleComponentTranslators.h b/Engine/Plugins/Runtime/MassGameplay/Source/MassActors/Public/Translators/MassCapsuleComponentTranslators.h
--- a/Engine/Plugins/Runtime/MassGameplay/Source/MassActors/Public/Translators/MassCapsuleComponentTranslators.h
+++ b/Engine/Plugins/Runtime/MassGameplay/Source/MassActors/Public/Translators/MassCapsuleComponentTranslators.h
@@ -26,7 +26,7 @@ struct FDataFragment_CapsuleComponentWrapper : public FDataFragment_ObjectWrappe
  * iterating over fragments just to check that specific value.
  */
 USTRUCT()
-struct FMassCapsuleTransformCopyToMassTag : public FComponentTag
+struct FMassCapsuleTransformCopyToMassTag : public FMassTag
 {
 	GENERATED_BODY()
 };
@@ -40,13 +40,13 @@ public:
 
 protected:
 	virtual void ConfigureQueries() override;
-	virtual void Execute(UMassEntitySubsystem& EntitySubsystem, FLWComponentSystemExecutionContext& Context) override;	
+	virtual void Execute(UMassEntitySubsystem& EntitySubsystem, FMassExecutionContext& Context) override;	
 
-	FLWComponentQuery EntityQuery;
+	FMassEntityQuery EntityQuery;
 };
 
 USTRUCT()
-struct FMassCapsuleTransformCopyToActorTag : public FComponentTag
+struct FMassCapsuleTransformCopyToActorTag : public FMassTag
 {
 	GENERATED_BODY()
 };
@@ -60,7 +60,7 @@ public:
 
 protected:
 	virtual void ConfigureQueries() override;
-	virtual void Execute(UMassEntitySubsystem& EntitySubsystem, FLWComponentSystemExecutionContext& Context) override;
+	virtual void Execute(UMassEntitySubsystem& EntitySubsystem, FMassExecutionContext& Context) override;
 
-	FLWComponentQuery EntityQuery;
+	FMassEntityQuery EntityQuery;
 };
diff --git a/Engine/Plugins/Runtime/MassGameplay/Source/MassActors/Public/Translators/MassCharacterMovementTranslators.h b/Engine/Plugins/Runtime/MassGameplay/Source/MassActors/Public/Translators/MassCharacterMovementTranslators.h
--- a/Engine/Plugins/Runtime/MassGameplay/Source/MassActors/Public/Translators/MassCharacterMovementTranslators.h
+++ b/Engine/Plugins/Runtime/MassGameplay/Source/MassActors/Public/Translators/MassCharacterMovementTranslators.h
@@ -20,7 +20,7 @@ struct FDataFragment_CharacterMovementComponentWrapper : public FDataFragment_Ob
 };
 
 USTRUCT()
-struct FMassCharacterMovementCopyToMassTag : public FComponentTag
+struct FMassCharacterMovementCopyToMassTag : public FMassTag
 {
 	GENERATED_BODY()
 };
@@ -35,13 +35,13 @@ public:
 
 protected:
 	virtual void ConfigureQueries() override;
-	virtual void Execute(UMassEntitySubsystem& EntitySubsystem, FLWComponentSystemExecutionContext& Context) override;
+	virtual void Execute(UMassEntitySubsystem& EntitySubsystem, FMassExecutionContext& Context) override;
 
-	FLWComponentQuery EntityQuery;
+	FMassEntityQuery EntityQuery;
 };
 
 USTRUCT()
-struct FMassCharacterMovementCopyToActorTag : public FComponentTag
+struct FMassCharacterMovementCopyToActorTag : public FMassTag
 {
 	GENERATED_BODY()
 };
@@ -57,14 +57,14 @@ public:
 
 protected:
 	virtual void ConfigureQueries() override;
-	virtual void Execute(UMassEntitySubsystem& EntitySubsystem, FLWComponentSystemExecutionContext& Context) override;
+	virtual void Execute(UMassEntitySubsystem& EntitySubsystem, FMassExecutionContext& Context) override;
 
-	FLWComponentQuery EntityQuery;
+	FMassEntityQuery EntityQuery;
 };
 
 
 USTRUCT()
-struct FMassCharacterOrientationCopyToMassTag : public FComponentTag
+struct FMassCharacterOrientationCopyToMassTag : public FMassTag
 {
 	GENERATED_BODY()
 };
@@ -79,13 +79,13 @@ public:
 
 protected:
 	virtual void ConfigureQueries() override;
-	virtual void Execute(UMassEntitySubsystem& EntitySubsystem, FLWComponentSystemExecutionContext& Context) override;
+	virtual void Execute(UMassEntitySubsystem& EntitySubsystem, FMassExecutionContext& Context) override;
 
-	FLWComponentQuery EntityQuery;
+	FMassEntityQuery EntityQuery;
 };
 
 USTRUCT()
-struct FMassCharacterOrientationCopyToActorTag : public FComponentTag
+struct FMassCharacterOrientationCopyToActorTag : public FMassTag
 {
 	GENERATED_BODY()
 };
@@ -101,9 +101,9 @@ public:
 
 protected:
 	virtual void ConfigureQueries() override;
-	virtual void Execute(UMassEntitySubsystem& EntitySubsystem, FLWComponentSystemExecutionContext& Context) override;
+	virtual void Execute(UMassEntitySubsystem& EntitySubsystem, FMassExecutionContext& Context) override;
 
-	FLWComponentQuery EntityQuery;
+	FMassEntityQuery EntityQuery;
 };
 
 UCLASS()
@@ -115,9 +115,9 @@ public:
 
 protected:
 	virtual void ConfigureQueries() override;
-	virtual void Execute(UMassEntitySubsystem& EntitySubsystem, FLWComponentSystemExecutionContext& Context) override;
+	virtual void Execute(UMassEntitySubsystem& EntitySubsystem, FMassExecutionContext& Context) override;
 
-	FLWComponentQuery EntityQuery;
+	FMassEntityQuery EntityQuery;
 };
 
 UCLASS()
@@ -130,7 +130,7 @@ public:
 
 protected:
 	virtual void ConfigureQueries() override;
-	virtual void Execute(UMassEntitySubsystem& EntitySubsystem, FLWComponentSystemExecutionContext& Context) override;
+	virtual void Execute(UMassEntitySubsystem& EntitySubsystem, FMassExecutionContext& Context) override;
 
-	FLWComponentQuery EntityQuery;
+	FMassEntityQuery EntityQuery;
 };
\ No newline at end of file
diff --git a/Engine/Plugins/Runtime/MassGameplay/Source/MassActors/Public/Translators/MassSceneComponentLocationTranslator.h b/Engine/Plugins/Runtime/MassGameplay/Source/MassActors/Public/Translators/MassSceneComponentLocationTranslator.h
--- a/Engine/Plugins/Runtime/MassGameplay/Source/MassActors/Public/Translators/MassSceneComponentLocationTranslator.h
+++ b/Engine/Plugins/Runtime/MassGameplay/Source/MassActors/Public/Translators/MassSceneComponentLocationTranslator.h
@@ -15,7 +15,7 @@ struct FMassSceneComponentWrapperFragment : public FDataFragment_ObjectWrapper
 };
 
 USTRUCT()
-struct FMassSceneComponentLocationCopyToMassTag : public FComponentTag
+struct FMassSceneComponentLocationCopyToMassTag : public FMassTag
 {
 	GENERATED_BODY()
 };
@@ -30,14 +30,14 @@ public:
 
 protected:
 	virtual void ConfigureQueries() override;
-	virtual void Execute(UMassEntitySubsystem& EntitySubsystem, FLWComponentSystemExecutionContext& Context) override;
+	virtual void Execute(UMassEntitySubsystem& EntitySubsystem, FMassExecutionContext& Context) override;
 
-	FLWComponentQuery EntityQuery;
+	FMassEntityQuery EntityQuery;
 };
 
 
 USTRUCT()
-struct FMassSceneComponentLocationCopyToActorTag : public FComponentTag
+struct FMassSceneComponentLocationCopyToActorTag : public FMassTag
 {
 	GENERATED_BODY()
 };
@@ -52,7 +52,7 @@ public:
 
 protected:
 	virtual void ConfigureQueries() override;
-	virtual void Execute(UMassEntitySubsystem& EntitySubsystem, FLWComponentSystemExecutionContext& Context) override;
+	virtual void Execute(UMassEntitySubsystem& EntitySubsystem, FMassExecutionContext& Context) override;
 
-	FLWComponentQuery EntityQuery;
+	FMassEntityQuery EntityQuery;
 };
\ No newline at end of file
diff --git a/Engine/Plugins/Runtime/MassGameplay/Source/MassActors/Public/Translators/MassTranslators_BehaviorTree.h b/Engine/Plugins/Runtime/MassGameplay/Source/MassActors/Public/Translators/MassTranslators_BehaviorTree.h
--- a/Engine/Plugins/Runtime/MassGameplay/Source/MassActors/Public/Translators/MassTranslators_BehaviorTree.h
+++ b/Engine/Plugins/Runtime/MassGameplay/Source/MassActors/Public/Translators/MassTranslators_BehaviorTree.h
@@ -26,7 +26,7 @@ public:
 
 protected:
 	virtual void ConfigureQueries() override;
-	virtual void Execute(UMassEntitySubsystem& EntitySubsystem, FLWComponentSystemExecutionContext& Context) override {}
+	virtual void Execute(UMassEntitySubsystem& EntitySubsystem, FMassExecutionContext& Context) override {}
 
-	FLWComponentQuery EntityQuery;
+	FMassEntityQuery EntityQuery;
 };
diff --git a/Engine/Plugins/Runtime/MassGameplay/Source/MassCommon/Private/MassCommonTypes.cpp b/Engine/Plugins/Runtime/MassGameplay/Source/MassCommon/Private/MassCommonTypes.cpp
--- a/Engine/Plugins/Runtime/MassGameplay/Source/MassCommon/Private/MassCommonTypes.cpp
+++ b/Engine/Plugins/Runtime/MassGameplay/Source/MassCommon/Private/MassCommonTypes.cpp
@@ -81,7 +81,7 @@ bool HasDebugEntities()
 	return DebugEntityBegin != INDEX_NONE && DebugEntityEnd != INDEX_NONE;
 }
 	
-bool IsDebuggingEntity(FLWEntity Entity, FColor* OutEntityColor)
+bool IsDebuggingEntity(FMassEntityHandle Entity, FColor* OutEntityColor)
 {
 	const int32 EntityIdx = Entity.Index;
 	const bool bIsDebuggingEntity = (DebugEntityBegin != INDEX_NONE && DebugEntityEnd != INDEX_NONE && DebugEntityBegin <= EntityIdx && EntityIdx <= DebugEntityEnd);
@@ -94,7 +94,7 @@ bool IsDebuggingEntity(FLWEntity Entity, FColor* OutEntityColor)
 	return bIsDebuggingEntity;
 }
 
-FColor GetEntityDebugColor(FLWEntity Entity)
+FColor GetEntityDebugColor(FMassEntityHandle Entity)
 {
 	const int32 EntityIdx = Entity.Index;
 	return EntityIdx != INDEX_NONE ? GColorList.GetFColorByIndex(EntityIdx % GColorList.GetColorsNum()) : FColor::Black;
diff --git a/Engine/Plugins/Runtime/MassGameplay/Source/MassCommon/Private/MassCommonUtils.cpp b/Engine/Plugins/Runtime/MassGameplay/Source/MassCommon/Private/MassCommonUtils.cpp
--- a/Engine/Plugins/Runtime/MassGameplay/Source/MassCommon/Private/MassCommonUtils.cpp
+++ b/Engine/Plugins/Runtime/MassGameplay/Source/MassCommon/Private/MassCommonUtils.cpp
@@ -4,13 +4,13 @@
 
 namespace UE::Mass::Utils
 {
-	TArray<FLWEntity> EntityQueueToArray(TQueue<FLWEntity, EQueueMode::Mpsc>& EntitiesQueue, const int32 EntitiesCount)
+	TArray<FMassEntityHandle> EntityQueueToArray(TQueue<FMassEntityHandle, EQueueMode::Mpsc>& EntitiesQueue, const int32 EntitiesCount)
 	{
 		check(EntitiesCount > 0);
-		TArray<FLWEntity> EntitiesArray;
+		TArray<FMassEntityHandle> EntitiesArray;
 		EntitiesArray.AddUninitialized(EntitiesCount);
 
-		FLWEntity TempEntity;
+		FMassEntityHandle TempEntity;
 		uint32 CurrentIndex = 0;
 		while (EntitiesQueue.Dequeue(TempEntity))
 		{
diff --git a/Engine/Plugins/Runtime/MassGameplay/Source/MassCommon/Private/MassGenericProcessors.cpp b/Engine/Plugins/Runtime/MassGameplay/Source/MassCommon/Private/MassGenericProcessors.cpp
--- a/Engine/Plugins/Runtime/MassGameplay/Source/MassCommon/Private/MassGenericProcessors.cpp
+++ b/Engine/Plugins/Runtime/MassGameplay/Source/MassCommon/Private/MassGenericProcessors.cpp
@@ -12,11 +12,11 @@ URandomizeVectorProcessor::URandomizeVectorProcessor()
 	bAutoRegisterWithProcessingPhases = false;
 }
 
-void URandomizeVectorProcessor::Execute(UMassEntitySubsystem& EntitySubsystem, FLWComponentSystemExecutionContext& Context)
+void URandomizeVectorProcessor::Execute(UMassEntitySubsystem& EntitySubsystem, FMassExecutionContext& Context)
 {
 	QUICK_SCOPE_CYCLE_COUNTER(RandomizeVectorProcessor_Run);
 	
-	EntityQuery.ForEachEntityChunk(EntitySubsystem, Context, [this](FLWComponentSystemExecutionContext& Context)
+	EntityQuery.ForEachEntityChunk(EntitySubsystem, Context, [this](FMassExecutionContext& Context)
 		{
 			for (FVector& Vector : ComponentList)
 			{
@@ -40,7 +40,7 @@ void URandomizeVectorProcessor::PostEditChangeProperty(FPropertyChangedEvent& Pr
 		if (PropName == ComponentTypeName)
 		{
 			if (ComponentType == nullptr
-				|| ComponentType->IsChildOf(FLWComponentData::StaticStruct()) == false
+				|| ComponentType->IsChildOf(FMassFragment::StaticStruct()) == false
 				|| ComponentType->GetStructureSize() != SupportedSize)
 			{
 				ComponentType = FVectorComponent::StaticStruct();
diff --git a/Engine/Plugins/Runtime/MassGameplay/Source/MassCommon/Public/MassCommonFragments.h b/Engine/Plugins/Runtime/MassGameplay/Source/MassCommon/Public/MassCommonFragments.h
--- a/Engine/Plugins/Runtime/MassGameplay/Source/MassCommon/Public/MassCommonFragments.h
+++ b/Engine/Plugins/Runtime/MassGameplay/Source/MassCommon/Public/MassCommonFragments.h
@@ -8,7 +8,7 @@
 
 
 USTRUCT()
-struct MASSCOMMON_API FDataFragment_Transform : public FLWComponentData
+struct MASSCOMMON_API FDataFragment_Transform : public FMassFragment
 {
 	GENERATED_BODY()
 
@@ -22,7 +22,7 @@ protected:
 };
 
 USTRUCT()
-struct MASSCOMMON_API FDataFragment_NavLocation : public FLWComponentData
+struct MASSCOMMON_API FDataFragment_NavLocation : public FMassFragment
 {
 	GENERATED_BODY()
 	NavNodeRef NodeRef = INVALID_NAVNODEREF;
@@ -32,7 +32,7 @@ struct MASSCOMMON_API FDataFragment_NavLocation : public FLWComponentData
  *  use URandomizeVectorProcessor to initialize it to random variables at spawn time.
  */
 USTRUCT()
-struct MASSCOMMON_API FVectorComponent : public FLWComponentData
+struct MASSCOMMON_API FVectorComponent : public FMassFragment
 {
 	GENERATED_BODY()
 	UPROPERTY(EditAnywhere, Category = "")
@@ -40,7 +40,7 @@ struct MASSCOMMON_API FVectorComponent : public FLWComponentData
 };
 
 USTRUCT()
-struct MASSCOMMON_API FDataFragment_AgentRadius : public FLWComponentData
+struct MASSCOMMON_API FDataFragment_AgentRadius : public FMassFragment
 {
 	GENERATED_BODY()
 	UPROPERTY(EditAnywhere, Category = "")
@@ -51,7 +51,7 @@ struct MASSCOMMON_API FDataFragment_AgentRadius : public FLWComponentData
  *	Mass simulation..
  */
 USTRUCT()
-struct FDataFragment_ObjectWrapper : public FLWComponentData
+struct FDataFragment_ObjectWrapper : public FMassFragment
 {
 	GENERATED_BODY()
 };
diff --git a/Engine/Plugins/Runtime/MassGameplay/Source/MassCommon/Public/MassCommonTypes.h b/Engine/Plugins/Runtime/MassGameplay/Source/MassCommon/Public/MassCommonTypes.h
--- a/Engine/Plugins/Runtime/MassGameplay/Source/MassCommon/Public/MassCommonTypes.h
+++ b/Engine/Plugins/Runtime/MassGameplay/Source/MassCommon/Public/MassCommonTypes.h
@@ -19,8 +19,8 @@ MASSCOMMON_API DECLARE_LOG_CATEGORY_EXTERN(LogMass, Warning, All);
 namespace UE::MassDebug
 {
 	MASSCOMMON_API extern bool HasDebugEntities();
-	MASSCOMMON_API extern bool IsDebuggingEntity(FLWEntity Entity, FColor* OutEntityColor = nullptr);
-	MASSCOMMON_API extern FColor GetEntityDebugColor(FLWEntity Entity);
+	MASSCOMMON_API extern bool IsDebuggingEntity(FMassEntityHandle Entity, FColor* OutEntityColor = nullptr);
+	MASSCOMMON_API extern FColor GetEntityDebugColor(FMassEntityHandle Entity);
 } // namespace UE::MassDebug
 #endif // WITH_MASSGAMEPLAY_DEBUG
 
@@ -40,11 +40,11 @@ struct MASSCOMMON_API FMassHandle
 public:
 	FMassHandle() = default;
 
-	explicit FMassHandle(FLWEntity InLWEntity)
+	explicit FMassHandle(FMassEntityHandle InLWEntity)
 		: LWEntity{ InLWEntity }
 	{}
 
-	FLWEntity GetLWEntity() const
+	FMassEntityHandle GetLWEntity() const
 	{
 		return LWEntity;
 	}
@@ -81,7 +81,7 @@ public:
 	}
 
 private:
-	FLWEntity LWEntity;
+	FMassEntityHandle LWEntity;
 };
 
 USTRUCT()
@@ -129,7 +129,7 @@ struct MASSCOMMON_API FMassUniqueFragmentCollection
 	template<typename T>
 	T& Add_GetRef()
 	{
-		static_assert(TIsDerivedFrom<T, FLWComponentData>::IsDerived, "Given struct is not of a fragment type. If you use this FMassUniqueFragmentCollection with the EntitySubsystem you'll hit some checks");
+		static_assert(TIsDerivedFrom<T, FMassFragment>::IsDerived, "Given struct is not of a fragment type. If you use this FMassUniqueFragmentCollection with the EntitySubsystem you'll hit some checks");
 		if (const FInstancedStruct* ElementPtr = FragmentInstances.FindByPredicate(FSameTypeScriptStructPredicate(T::StaticStruct())))
 		{
 			return ElementPtr->GetMutable<T>();
@@ -144,7 +144,7 @@ struct MASSCOMMON_API FMassUniqueFragmentCollection
 	template<typename T>
 	T& GetRefChecked()
 	{
-		static_assert(TIsDerivedFrom<T, FLWComponentData>::IsDerived, "Given struct is not of a fragment type. Make sure your struct extends FLWComponentData or its derivatives");
+		static_assert(TIsDerivedFrom<T, FMassFragment>::IsDerived, "Given struct is not of a fragment type. Make sure your struct extends FMassFragment or its derivatives");
 		FInstancedStruct* ElementPtr = FragmentInstances.FindByPredicate(FSameTypeScriptStructPredicate(T::StaticStruct()));
 		check(ElementPtr);
 		return ElementPtr->GetMutable<T>();
@@ -152,7 +152,7 @@ struct MASSCOMMON_API FMassUniqueFragmentCollection
 
 	int32 GetFragmentsNum() const { return FragmentInstances.Num(); }
 	TArrayView<const FInstancedStruct> GetFragments() const { return MakeArrayView(FragmentInstances); }
-	const FLWComponentBitSet& GetFragmentBitSet() const { return FragmentBitSet; }
+	const FMassFragmentBitSet& GetFragmentBitSet() const { return FragmentBitSet; }
 
 	bool IsEmpty() const { return FragmentInstances.IsEmpty(); }
 
@@ -162,7 +162,7 @@ private:
 	UPROPERTY()
 	TArray<FInstancedStruct> FragmentInstances;
 
-	FLWComponentBitSet FragmentBitSet;
+	FMassFragmentBitSet FragmentBitSet;
 };
 
 /** Float encoded in int16, 1cm accuracy. */
diff --git a/Engine/Plugins/Runtime/MassGameplay/Source/MassCommon/Public/MassCommonUtils.h b/Engine/Plugins/Runtime/MassGameplay/Source/MassCommon/Public/MassCommonUtils.h
--- a/Engine/Plugins/Runtime/MassGameplay/Source/MassCommon/Public/MassCommonUtils.h
+++ b/Engine/Plugins/Runtime/MassGameplay/Source/MassCommon/Public/MassCommonUtils.h
@@ -12,5 +12,5 @@ namespace UE::Mass::Utils
 	 * Creates a TArray of entities based on the given EntitiesQueue. Note that it's the caller's responsibility to 
 	 * ensure EntitiesCount > 0, otherwise the function will fail a check (with explosive results).
 	 */
-	MASSCOMMON_API TArray<FLWEntity> EntityQueueToArray(TQueue<FLWEntity, EQueueMode::Mpsc>& EntitiesQueue, const int32 EntitiesCount);
+	MASSCOMMON_API TArray<FMassEntityHandle> EntityQueueToArray(TQueue<FMassEntityHandle, EQueueMode::Mpsc>& EntitiesQueue, const int32 EntitiesCount);
 } // namespace UE::Mass::Utils
\ No newline at end of file
diff --git a/Engine/Plugins/Runtime/MassGameplay/Source/MassCommon/Public/MassGenericProcessors.h b/Engine/Plugins/Runtime/MassGameplay/Source/MassCommon/Public/MassGenericProcessors.h
--- a/Engine/Plugins/Runtime/MassGameplay/Source/MassCommon/Public/MassGenericProcessors.h
+++ b/Engine/Plugins/Runtime/MassGameplay/Source/MassCommon/Public/MassGenericProcessors.h
@@ -8,13 +8,13 @@
 
 
 /** 
- *  An experimental processor that can be configured to work on any LWComponent type as long as it's size matches
+ *  An experimental processor that can be configured to work on any Fragment type as long as it's size matches
  *  FVectorComponent (effectively the size of FVector). Meant mostly for prototyping.
  *  @todo currently the processor is randomizing it as if it was a fixed-max-length 2d velocity vector. If we keep this
  *  idea around we can add more properties to the processor that would configure the behavior.
  */
 UCLASS()
-class MASSCOMMON_API URandomizeVectorProcessor : public UPipeProcessor
+class MASSCOMMON_API URandomizeVectorProcessor : public UMassProcessor
 {
 	GENERATED_BODY()
 
@@ -31,11 +31,11 @@ public:
 	{
 		if (ComponentType)
 		{
-			EntityQuery.AddRequirement(ComponentType, ELWComponentAccess::ReadWrite);
+			EntityQuery.AddRequirement(ComponentType, EMassFragmentAccess::ReadWrite);
 		}
 	}
 
-	virtual void Execute(UMassEntitySubsystem& EntitySubsystem, FLWComponentSystemExecutionContext& Context) override;
+	virtual void Execute(UMassEntitySubsystem& EntitySubsystem, FMassExecutionContext& Context) override;
 
 	UPROPERTY(EditAnywhere, Category = "Mass", meta = (ClampMin = 1, UIMin = 1))
 	float MinMagnitude = 1.f;
@@ -48,5 +48,5 @@ public:
 #endif // WITH_EDITOR
 
 protected:
-	FLWComponentQuery EntityQuery;
+	FMassEntityQuery EntityQuery;
 };
diff --git a/Engine/Plugins/Runtime/MassGameplay/Source/MassGameplayDebug/Private/AssignDebugVisProcessor.cpp b/Engine/Plugins/Runtime/MassGameplay/Source/MassGameplayDebug/Private/AssignDebugVisProcessor.cpp
--- a/Engine/Plugins/Runtime/MassGameplay/Source/MassGameplayDebug/Private/AssignDebugVisProcessor.cpp
+++ b/Engine/Plugins/Runtime/MassGameplay/Source/MassGameplayDebug/Private/AssignDebugVisProcessor.cpp
@@ -17,10 +17,10 @@ UAssignDebugVisProcessor::UAssignDebugVisProcessor()
 
 void UAssignDebugVisProcessor::ConfigureQueries()
 {
-	EntityQuery.AddRequirement<FSimDebugVisComponent>(ELWComponentAccess::ReadWrite);
+	EntityQuery.AddRequirement<FSimDebugVisComponent>(EMassFragmentAccess::ReadWrite);
 }
 
-void UAssignDebugVisProcessor::Execute(UMassEntitySubsystem& EntitySubsystem, FLWComponentSystemExecutionContext& Context)
+void UAssignDebugVisProcessor::Execute(UMassEntitySubsystem& EntitySubsystem, FMassExecutionContext& Context)
 {
 	UMassDebugVisualizationComponent* Visualizer = WeakVisualizer.Get();
 	if (!ensure(Visualizer))
@@ -34,7 +34,7 @@ void UAssignDebugVisProcessor::Execute(UMassEntitySubsystem& EntitySubsystem, FL
 	// note that this function will create the "visual components" only it they're missing or out of sync. 
 	Visualizer->ConditionallyConstructVisualComponent();
 
-	EntityQuery.ForEachEntityChunk(EntitySubsystem, Context, [this, Visualizer](FLWComponentSystemExecutionContext& Context)
+	EntityQuery.ForEachEntityChunk(EntitySubsystem, Context, [this, Visualizer](FMassExecutionContext& Context)
 	{
 		const TArrayView<FSimDebugVisComponent> DebugVisList = Context.GetMutableComponentView<FSimDebugVisComponent>();
 		for (FSimDebugVisComponent& VisualComp : DebugVisList)
diff --git a/Engine/Plugins/Runtime/MassGameplay/Source/MassGameplayDebug/Private/DebugVisLocationProcessor.cpp b/Engine/Plugins/Runtime/MassGameplay/Source/MassGameplayDebug/Private/DebugVisLocationProcessor.cpp
--- a/Engine/Plugins/Runtime/MassGameplay/Source/MassGameplayDebug/Private/DebugVisLocationProcessor.cpp
+++ b/Engine/Plugins/Runtime/MassGameplay/Source/MassGameplayDebug/Private/DebugVisLocationProcessor.cpp
@@ -18,12 +18,12 @@ UDebugVisLocationProcessor::UDebugVisLocationProcessor()
 
 void UDebugVisLocationProcessor::ConfigureQueries()
 {
-	EntityQuery.AddRequirement<FDataFragment_Transform>(ELWComponentAccess::ReadOnly);
-	EntityQuery.AddRequirement<FSimDebugVisComponent>(ELWComponentAccess::ReadOnly);
-	EntityQuery.AddRequirement<FMassVelocityFragment>(ELWComponentAccess::ReadOnly);
+	EntityQuery.AddRequirement<FDataFragment_Transform>(EMassFragmentAccess::ReadOnly);
+	EntityQuery.AddRequirement<FSimDebugVisComponent>(EMassFragmentAccess::ReadOnly);
+	EntityQuery.AddRequirement<FMassVelocityFragment>(EMassFragmentAccess::ReadOnly);
 }
 
-void UDebugVisLocationProcessor::Execute(UMassEntitySubsystem& EntitySubsystem, FLWComponentSystemExecutionContext& Context)
+void UDebugVisLocationProcessor::Execute(UMassEntitySubsystem& EntitySubsystem, FMassExecutionContext& Context)
 {
 #if WITH_EDITORONLY_DATA
 	UMassDebugVisualizationComponent* Visualizer = WeakVisualizer.Get();
@@ -37,7 +37,7 @@ void UDebugVisLocationProcessor::Execute(UMassEntitySubsystem& EntitySubsystem,
 	TArrayView<UHierarchicalInstancedStaticMeshComponent*> VisualDataISMCs = Visualizer->GetVisualDataISMCs();
 	if (VisualDataISMCs.Num() > 0)
 	{
-		EntityQuery.ForEachEntityChunk(EntitySubsystem, Context, [this, &VisualDataISMCs](const FLWComponentSystemExecutionContext& Context)
+		EntityQuery.ForEachEntityChunk(EntitySubsystem, Context, [this, &VisualDataISMCs](const FMassExecutionContext& Context)
 		{
 			const int32 NumEntities = Context.GetEntitiesNum();
 			const TConstArrayView<FDataFragment_Transform> LocationList = Context.GetComponentView<FDataFragment_Transform>();
@@ -87,12 +87,12 @@ UMassProcessor_UpdateDebugVis::UMassProcessor_UpdateDebugVis()
 void UMassProcessor_UpdateDebugVis::ConfigureQueries() 
 {
 	// @todo only FDataFragment_DebugVis should be mandatory, rest optional 
-	EntityQuery.AddRequirement<FDataFragment_Transform>(ELWComponentAccess::ReadOnly);
-	EntityQuery.AddRequirement<FDataFragment_DebugVis>(ELWComponentAccess::ReadWrite);
-	EntityQuery.AddRequirement<FDataFragment_AgentRadius>(ELWComponentAccess::ReadWrite);
+	EntityQuery.AddRequirement<FDataFragment_Transform>(EMassFragmentAccess::ReadOnly);
+	EntityQuery.AddRequirement<FDataFragment_DebugVis>(EMassFragmentAccess::ReadWrite);
+	EntityQuery.AddRequirement<FDataFragment_AgentRadius>(EMassFragmentAccess::ReadWrite);
 }
 
-void UMassProcessor_UpdateDebugVis::Execute(UMassEntitySubsystem& EntitySubsystem, FLWComponentSystemExecutionContext& Context)
+void UMassProcessor_UpdateDebugVis::Execute(UMassEntitySubsystem& EntitySubsystem, FMassExecutionContext& Context)
 {
 	UMassDebuggerSubsystem* Debugger = UWorld::GetSubsystem<UMassDebuggerSubsystem>(GetWorld());
 	if (Debugger == nullptr)
@@ -102,7 +102,7 @@ void UMassProcessor_UpdateDebugVis::Execute(UMassEntitySubsystem& EntitySubsyste
 
 	QUICK_SCOPE_CYCLE_COUNTER(UMassProcessor_UpdateDebugVis_Run);
 
-	EntityQuery.ForEachEntityChunk(EntitySubsystem, Context, [this, Debugger](FLWComponentSystemExecutionContext& Context)
+	EntityQuery.ForEachEntityChunk(EntitySubsystem, Context, [this, Debugger](FMassExecutionContext& Context)
 		{
 			const int32 NumEntities = Context.GetEntitiesNum();
 			const TConstArrayView<FDataFragment_Transform> LocationList = Context.GetComponentView<FDataFragment_Transform>();
diff --git a/Engine/Plugins/Runtime/MassGameplay/Source/MassGameplayDebug/Private/MassDebuggerSubsystem.cpp b/Engine/Plugins/Runtime/MassGameplay/Source/MassGameplayDebug/Private/MassDebuggerSubsystem.cpp
--- a/Engine/Plugins/Runtime/MassGameplay/Source/MassGameplayDebug/Private/MassDebuggerSubsystem.cpp
+++ b/Engine/Plugins/Runtime/MassGameplay/Source/MassGameplayDebug/Private/MassDebuggerSubsystem.cpp
@@ -16,7 +16,7 @@ void UMassDebuggerSubsystem::Initialize(FSubsystemCollectionBase& Collection)
 
 	UMassSimulationSubsystem* SimSystem = UWorld::GetSubsystem<UMassSimulationSubsystem>(GetWorld());
 	check(SimSystem);
-	SimSystem->GetOnProcessingPhaseStarted(EPipeProcessingPhase::PrePhysics).AddUObject(this, &UMassDebuggerSubsystem::OnProcessingPhaseStarted);
+	SimSystem->GetOnProcessingPhaseStarted(EMassProcessingPhase::PrePhysics).AddUObject(this, &UMassDebuggerSubsystem::OnProcessingPhaseStarted);
 }
 
 void UMassDebuggerSubsystem::PreTickProcessors()
@@ -32,7 +32,7 @@ void UMassDebuggerSubsystem::PreTickProcessors()
 	SelectedEntityDetails.Empty();
 }
 
-void UMassDebuggerSubsystem::AddEntityLocation(FLWEntity Entity, const FVector& Location)
+void UMassDebuggerSubsystem::AddEntityLocation(FMassEntityHandle Entity, const FVector& Location)
 {
 	Entities.Add(Entity);
 	Locations.Add(Location);
@@ -43,7 +43,7 @@ void UMassDebuggerSubsystem::OnProcessingPhaseStarted(const float DeltaSeconds)
 	PreTickProcessors();
 }
 
-void UMassDebuggerSubsystem::SetSelectedEntity(const FLWEntity& InSelectedEntity)
+void UMassDebuggerSubsystem::SetSelectedEntity(const FMassEntityHandle& InSelectedEntity)
 {
 	SelectedEntity = InSelectedEntity;
 	SelectedEntityDetails.Empty();
diff --git a/Engine/Plugins/Runtime/MassGameplay/Source/MassGameplayDebug/Public/AssignDebugVisProcessor.h b/Engine/Plugins/Runtime/MassGameplay/Source/MassGameplayDebug/Public/AssignDebugVisProcessor.h
--- a/Engine/Plugins/Runtime/MassGameplay/Source/MassGameplayDebug/Public/AssignDebugVisProcessor.h
+++ b/Engine/Plugins/Runtime/MassGameplay/Source/MassGameplayDebug/Public/AssignDebugVisProcessor.h
@@ -16,11 +16,11 @@ class MASSGAMEPLAYDEBUG_API UAssignDebugVisProcessor : public UMassFragmentIniti
 public:
 	UAssignDebugVisProcessor();
 	virtual void ConfigureQueries() override;
-	virtual void Execute(UMassEntitySubsystem& EntitySubsystem, FLWComponentSystemExecutionContext& Context) override;
+	virtual void Execute(UMassEntitySubsystem& EntitySubsystem, FMassExecutionContext& Context) override;
 	virtual void Initialize(UObject& InOwner) override;
 
 protected:
 	TWeakObjectPtr<UMassDebugVisualizationComponent> WeakVisualizer;
 
-	FLWComponentQuery EntityQuery;
+	FMassEntityQuery EntityQuery;
 };
diff --git a/Engine/Plugins/Runtime/MassGameplay/Source/MassGameplayDebug/Public/DebugVisLocationProcessor.h b/Engine/Plugins/Runtime/MassGameplay/Source/MassGameplayDebug/Public/DebugVisLocationProcessor.h
--- a/Engine/Plugins/Runtime/MassGameplay/Source/MassGameplayDebug/Public/DebugVisLocationProcessor.h
+++ b/Engine/Plugins/Runtime/MassGameplay/Source/MassGameplayDebug/Public/DebugVisLocationProcessor.h
@@ -9,7 +9,7 @@ class UMassDebugVisualizationComponent;
 struct FSimDebugVisComponent;
 
 UCLASS()
-class MASSGAMEPLAYDEBUG_API UDebugVisLocationProcessor : public UPipeProcessor
+class MASSGAMEPLAYDEBUG_API UDebugVisLocationProcessor : public UMassProcessor
 {
 	GENERATED_BODY()
 
@@ -18,11 +18,11 @@ public:
 	
 protected:
 	virtual void ConfigureQueries() override;
-	virtual void Execute(UMassEntitySubsystem& EntitySubsystem, FLWComponentSystemExecutionContext& Context) override;
+	virtual void Execute(UMassEntitySubsystem& EntitySubsystem, FMassExecutionContext& Context) override;
 	virtual void Initialize(UObject& InOwner) override;
 
 	TWeakObjectPtr<UMassDebugVisualizationComponent> WeakVisualizer;
-	FLWComponentQuery EntityQuery;
+	FMassEntityQuery EntityQuery;
 };
 
 //----------------------------------------------------------------------//
@@ -31,15 +31,15 @@ protected:
 //class UMassDebugger;
 
 UCLASS()
-class MASSGAMEPLAYDEBUG_API UMassProcessor_UpdateDebugVis : public UPipeProcessor
+class MASSGAMEPLAYDEBUG_API UMassProcessor_UpdateDebugVis : public UMassProcessor
 {
 	GENERATED_BODY()
 public:
 	UMassProcessor_UpdateDebugVis();
 	virtual void ConfigureQueries() override;
-	virtual void Execute(UMassEntitySubsystem& EntitySubsystem, FLWComponentSystemExecutionContext& Context) override;
+	virtual void Execute(UMassEntitySubsystem& EntitySubsystem, FMassExecutionContext& Context) override;
 
 protected:
-	FLWComponentQuery EntityQuery;
+	FMassEntityQuery EntityQuery;
 };
 
diff --git a/Engine/Plugins/Runtime/MassGameplay/Source/MassGameplayDebug/Public/MassDebuggerSubsystem.h b/Engine/Plugins/Runtime/MassGameplay/Source/MassGameplayDebug/Public/MassDebuggerSubsystem.h
--- a/Engine/Plugins/Runtime/MassGameplay/Source/MassGameplayDebug/Public/MassDebuggerSubsystem.h
+++ b/Engine/Plugins/Runtime/MassGameplay/Source/MassGameplayDebug/Public/MassDebuggerSubsystem.h
@@ -30,12 +30,12 @@ public:
 	void AddShape(EMassEntityDebugShape Shape, FVector Location, float Size) { Shapes[uint8(Shape)].Add(FShapeDesc(Location, Size)); }
 	const TArray<FShapeDesc>* GetShapes() const { return Shapes; }
 
-	void AddEntityLocation(FLWEntity Entity, const FVector& Location);
-	TConstArrayView<FLWEntity> GetEntities() const { return Entities; }
+	void AddEntityLocation(FMassEntityHandle Entity, const FVector& Location);
+	TConstArrayView<FMassEntityHandle> GetEntities() const { return Entities; }
 	TConstArrayView<FVector> GetLocations() const { return Locations; }
 
-	FLWEntity GetSelectedEntity() const { return SelectedEntity; }
-	void SetSelectedEntity(const FLWEntity& InSelectedEntity);
+	FMassEntityHandle GetSelectedEntity() const { return SelectedEntity; }
+	void SetSelectedEntity(const FMassEntityHandle& InSelectedEntity);
 
 	void AppendSelectedEntityInfo(const FString& Info);
 	const FString& GetSelectedEntityInfo() const { return SelectedEntityDetails; }
@@ -59,9 +59,9 @@ protected:
 	bool bCollectingData = false;
 
 	TArray<FShapeDesc> Shapes[uint8(EMassEntityDebugShape::MAX)];
-	TArray<FLWEntity> Entities;
+	TArray<FMassEntityHandle> Entities;
 	TArray<FVector> Locations;
-	FLWEntity SelectedEntity;
+	FMassEntityHandle SelectedEntity;
 	FString SelectedEntityDetails;
 
 	UPROPERTY(Transient)
diff --git a/Engine/Plugins/Runtime/MassGameplay/Source/MassGameplayDebug/Public/MassGameplayDebugTypes.h b/Engine/Plugins/Runtime/MassGameplay/Source/MassGameplayDebug/Public/MassGameplayDebugTypes.h
--- a/Engine/Plugins/Runtime/MassGameplay/Source/MassGameplayDebug/Public/MassGameplayDebugTypes.h
+++ b/Engine/Plugins/Runtime/MassGameplay/Source/MassGameplayDebug/Public/MassGameplayDebugTypes.h
@@ -38,7 +38,7 @@ struct FSimDebugDataRow : public FTableRowBase
 };
 
 USTRUCT()
-struct FSimDebugVisComponent : public FLWComponentData
+struct FSimDebugVisComponent : public FMassFragment
 {
 	GENERATED_BODY()
 	int32 InstanceIndex = INDEX_NONE;
@@ -56,7 +56,7 @@ enum class EMassEntityDebugShape : uint8
 };
 
 USTRUCT()
-struct FDataFragment_DebugVis : public FLWComponentData
+struct FDataFragment_DebugVis : public FMassFragment
 {
 	GENERATED_BODY()
 	UPROPERTY(EditAnywhere, Category = Debug)
diff --git a/Engine/Plugins/Runtime/MassGameplay/Source/MassLOD/Private/MassLODCollectorProcessor.cpp b/Engine/Plugins/Runtime/MassGameplay/Source/MassLOD/Private/MassLODCollectorProcessor.cpp
--- a/Engine/Plugins/Runtime/MassGameplay/Source/MassLOD/Private/MassLODCollectorProcessor.cpp
+++ b/Engine/Plugins/Runtime/MassGameplay/Source/MassLOD/Private/MassLODCollectorProcessor.cpp
@@ -30,44 +30,44 @@ void UMassLODCollectorProcessor::ConfigureQueries()
 {
 	for (FMassCollectorLODConfig& LODConfig : LODConfigs)
 	{
-		FLWComponentQuery BaseQuery;
+		FMassEntityQuery BaseQuery;
 		if (LODConfig.TagFilter.GetScriptStruct())
 		{
-			BaseQuery.AddTagRequirement(*LODConfig.TagFilter.GetScriptStruct(), ELWComponentPresence::All);
+			BaseQuery.AddTagRequirement(*LODConfig.TagFilter.GetScriptStruct(), EMassFragmentPresence::All);
 		}
-		BaseQuery.AddRequirement<FDataFragment_Transform>(ELWComponentAccess::ReadOnly);
-		BaseQuery.AddRequirement<FMassLODInfoFragment>(ELWComponentAccess::ReadWrite);
+		BaseQuery.AddRequirement<FDataFragment_Transform>(EMassFragmentAccess::ReadOnly);
+		BaseQuery.AddRequirement<FMassLODInfoFragment>(EMassFragmentAccess::ReadWrite);
 
 		LODConfig.OnLODEntityQuery = BaseQuery;
-		LODConfig.OnLODEntityQuery.AddTagRequirement<FMassOffLODTag>(ELWComponentPresence::None);
+		LODConfig.OnLODEntityQuery.AddTagRequirement<FMassOffLODTag>(EMassFragmentPresence::None);
 		LODConfig.OffLODEntityQuery_Conditional = BaseQuery;
-		LODConfig.OffLODEntityQuery_Conditional.AddTagRequirement<FMassOffLODTag>(ELWComponentPresence::All);
-		LODConfig.OffLODEntityQuery_Conditional.AddChunkRequirement<FMassSimulationVariableTickChunkFragment>(ELWComponentAccess::ReadOnly);
+		LODConfig.OffLODEntityQuery_Conditional.AddTagRequirement<FMassOffLODTag>(EMassFragmentPresence::All);
+		LODConfig.OffLODEntityQuery_Conditional.AddChunkRequirement<FMassSimulationVariableTickChunkFragment>(EMassFragmentAccess::ReadOnly);
 		LODConfig.OffLODEntityQuery_Conditional.SetChunkFilter(&FMassSimulationVariableTickChunkFragment::ShouldTickChunkThisFrame);
 
 		LODConfig.CloseEntityQuery = BaseQuery;
-		LODConfig.CloseEntityQuery.AddTagRequirement<FMassVisibilityCulledByDistanceTag>(ELWComponentPresence::None);
+		LODConfig.CloseEntityQuery.AddTagRequirement<FMassVisibilityCulledByDistanceTag>(EMassFragmentPresence::None);
 		LODConfig.FarEntityQuery_Conditional = BaseQuery;
-		LODConfig.FarEntityQuery_Conditional.AddTagRequirement<FMassVisibilityCulledByDistanceTag>(ELWComponentPresence::All);
-		LODConfig.FarEntityQuery_Conditional.AddChunkRequirement<FMassVisualizationChunkFragment>(ELWComponentAccess::ReadOnly);
+		LODConfig.FarEntityQuery_Conditional.AddTagRequirement<FMassVisibilityCulledByDistanceTag>(EMassFragmentPresence::All);
+		LODConfig.FarEntityQuery_Conditional.AddChunkRequirement<FMassVisualizationChunkFragment>(EMassFragmentAccess::ReadOnly);
 		LODConfig.FarEntityQuery_Conditional.SetChunkFilter(&FMassVisualizationChunkFragment::ShouldUpdateVisualizationForChunk);
 
 		// @todo: Cleanup once we can do more complex logic with entity queries. There is no other way to do "Notall", we need to create by hand 3 queries to achieve it.
 		LODConfig.CloseAndOnLODEntityQuery = BaseQuery;
-		LODConfig.CloseAndOnLODEntityQuery.AddTagRequirement<FMassVisibilityCulledByDistanceTag>(ELWComponentPresence::None);
-		LODConfig.CloseAndOnLODEntityQuery.AddTagRequirement<FMassOffLODTag>(ELWComponentPresence::None);
+		LODConfig.CloseAndOnLODEntityQuery.AddTagRequirement<FMassVisibilityCulledByDistanceTag>(EMassFragmentPresence::None);
+		LODConfig.CloseAndOnLODEntityQuery.AddTagRequirement<FMassOffLODTag>(EMassFragmentPresence::None);
 		LODConfig.CloseAndOffLODEntityQuery = BaseQuery;
-		LODConfig.CloseAndOffLODEntityQuery.AddTagRequirement<FMassVisibilityCulledByDistanceTag>(ELWComponentPresence::None);
-		LODConfig.CloseAndOffLODEntityQuery.AddTagRequirement<FMassOffLODTag>(ELWComponentPresence::All);
+		LODConfig.CloseAndOffLODEntityQuery.AddTagRequirement<FMassVisibilityCulledByDistanceTag>(EMassFragmentPresence::None);
+		LODConfig.CloseAndOffLODEntityQuery.AddTagRequirement<FMassOffLODTag>(EMassFragmentPresence::All);
 		LODConfig.FarAndOnLODEntityQuery = BaseQuery;
-		LODConfig.FarAndOnLODEntityQuery.AddTagRequirement<FMassVisibilityCulledByDistanceTag>(ELWComponentPresence::All);
-		LODConfig.FarAndOnLODEntityQuery.AddTagRequirement<FMassOffLODTag>(ELWComponentPresence::None);
+		LODConfig.FarAndOnLODEntityQuery.AddTagRequirement<FMassVisibilityCulledByDistanceTag>(EMassFragmentPresence::All);
+		LODConfig.FarAndOnLODEntityQuery.AddTagRequirement<FMassOffLODTag>(EMassFragmentPresence::None);
 		LODConfig.FarAndOffLODEntityQuery_Conditional = BaseQuery;
-		LODConfig.FarAndOffLODEntityQuery_Conditional.AddTagRequirement<FMassVisibilityCulledByDistanceTag>(ELWComponentPresence::All);
-		LODConfig.FarAndOffLODEntityQuery_Conditional.AddTagRequirement<FMassOffLODTag>(ELWComponentPresence::All);
-		LODConfig.FarAndOffLODEntityQuery_Conditional.AddChunkRequirement<FMassSimulationVariableTickChunkFragment>(ELWComponentAccess::ReadOnly);
-		LODConfig.FarAndOffLODEntityQuery_Conditional.AddChunkRequirement<FMassVisualizationChunkFragment>(ELWComponentAccess::ReadOnly);
-		LODConfig.FarAndOffLODEntityQuery_Conditional.SetChunkFilter([](const FLWComponentSystemExecutionContext& Context)
+		LODConfig.FarAndOffLODEntityQuery_Conditional.AddTagRequirement<FMassVisibilityCulledByDistanceTag>(EMassFragmentPresence::All);
+		LODConfig.FarAndOffLODEntityQuery_Conditional.AddTagRequirement<FMassOffLODTag>(EMassFragmentPresence::All);
+		LODConfig.FarAndOffLODEntityQuery_Conditional.AddChunkRequirement<FMassSimulationVariableTickChunkFragment>(EMassFragmentAccess::ReadOnly);
+		LODConfig.FarAndOffLODEntityQuery_Conditional.AddChunkRequirement<FMassVisualizationChunkFragment>(EMassFragmentAccess::ReadOnly);
+		LODConfig.FarAndOffLODEntityQuery_Conditional.SetChunkFilter([](const FMassExecutionContext& Context)
 			{
 				return Context.GetChunkComponent<FMassVisualizationChunkFragment>().ShouldUpdateVisualization()
 					|| Context.GetChunkComponent<FMassSimulationVariableTickChunkFragment>().ShouldTickThisFrame();
@@ -76,11 +76,11 @@ void UMassLODCollectorProcessor::ConfigureQueries()
 }
 
 template <typename TCollector>
-void CollectLOD(UMassEntitySubsystem& EntitySubsystem, FLWComponentSystemExecutionContext& Context, const TArray<FViewerInfo>& Viewers, TCollector& Collector, TArrayView<FLWComponentQuery> HighFrequencyTickingEntityQueries, FLWComponentQuery& LowFrequencyTickingEntityQuery)
+void CollectLOD(UMassEntitySubsystem& EntitySubsystem, FMassExecutionContext& Context, const TArray<FViewerInfo>& Viewers, TCollector& Collector, TArrayView<FMassEntityQuery> HighFrequencyTickingEntityQueries, FMassEntityQuery& LowFrequencyTickingEntityQuery)
 {
 	Collector.PrepareExecution(Viewers);
 
-	auto CollectLODInfo = [&Collector](FLWComponentSystemExecutionContext& Context)
+	auto CollectLODInfo = [&Collector](FMassExecutionContext& Context)
 	{
 		TConstArrayView<FDataFragment_Transform> LocationList = Context.GetComponentView<FDataFragment_Transform>();
 		TArrayView<FMassLODInfoFragment> ViewerInfoList = Context.GetMutableComponentView<FMassLODInfoFragment>();
@@ -90,7 +90,7 @@ void CollectLOD(UMassEntitySubsystem& EntitySubsystem, FLWComponentSystemExecuti
 
 	{
 		TRACE_CPUPROFILER_EVENT_SCOPE(TEXT("HighFrequency"))
-		for(FLWComponentQuery& HighFrequencyTickingEntityQuery : HighFrequencyTickingEntityQueries)
+		for(FMassEntityQuery& HighFrequencyTickingEntityQuery : HighFrequencyTickingEntityQueries)
 		{
 			HighFrequencyTickingEntityQuery.ForEachEntityChunk(EntitySubsystem, Context, CollectLODInfo);
 		}
@@ -101,7 +101,7 @@ void CollectLOD(UMassEntitySubsystem& EntitySubsystem, FLWComponentSystemExecuti
 	}
 }
 
-void UMassLODCollectorProcessor::Execute(UMassEntitySubsystem& EntitySubsystem, FLWComponentSystemExecutionContext& Context)
+void UMassLODCollectorProcessor::Execute(UMassEntitySubsystem& EntitySubsystem, FMassExecutionContext& Context)
 {
 	check(LODManager);
 	const TArray<FViewerInfo>& Viewers = LODManager->GetViewers();
diff --git a/Engine/Plugins/Runtime/MassGameplay/Source/MassLOD/Private/MassRepresentationLODCollectorProcessor.cpp b/Engine/Plugins/Runtime/MassGameplay/Source/MassLOD/Private/MassRepresentationLODCollectorProcessor.cpp
--- a/Engine/Plugins/Runtime/MassGameplay/Source/MassLOD/Private/MassRepresentationLODCollectorProcessor.cpp
+++ b/Engine/Plugins/Runtime/MassGameplay/Source/MassLOD/Private/MassRepresentationLODCollectorProcessor.cpp
@@ -24,29 +24,29 @@ void UMassRepresentationLODCollectorProcessor::ConfigureQueries()
 {
 	for (FMassRepresentationLODCollectorConfig& LODConfig : LODConfigs)
 	{
-		FLWComponentQuery BaseQuery;
+		FMassEntityQuery BaseQuery;
 		if (LODConfig.TagFilter.GetScriptStruct())
 		{
-			BaseQuery.AddTagRequirement(*LODConfig.TagFilter.GetScriptStruct(), ELWComponentPresence::All);
+			BaseQuery.AddTagRequirement(*LODConfig.TagFilter.GetScriptStruct(), EMassFragmentPresence::All);
 		}
-		BaseQuery.AddRequirement<FDataFragment_Transform>(ELWComponentAccess::ReadOnly);
-		BaseQuery.AddRequirement<FMassLODInfoFragment>(ELWComponentAccess::ReadWrite);
-		BaseQuery.AddChunkRequirement<FMassVisualizationChunkFragment>(ELWComponentAccess::ReadOnly);
+		BaseQuery.AddRequirement<FDataFragment_Transform>(EMassFragmentAccess::ReadOnly);
+		BaseQuery.AddRequirement<FMassLODInfoFragment>(EMassFragmentAccess::ReadWrite);
+		BaseQuery.AddChunkRequirement<FMassVisualizationChunkFragment>(EMassFragmentAccess::ReadOnly);
 
 		LODConfig.CloseEntityQuery = BaseQuery;
-		LODConfig.CloseEntityQuery.AddTagRequirement<FMassVisibilityCulledByDistanceTag>(ELWComponentPresence::None);
+		LODConfig.CloseEntityQuery.AddTagRequirement<FMassVisibilityCulledByDistanceTag>(EMassFragmentPresence::None);
 		LODConfig.FarEntityQuery_Conditional = BaseQuery;
-		LODConfig.FarEntityQuery_Conditional.AddTagRequirement<FMassVisibilityCulledByDistanceTag>(ELWComponentPresence::All);
+		LODConfig.FarEntityQuery_Conditional.AddTagRequirement<FMassVisibilityCulledByDistanceTag>(EMassFragmentPresence::All);
 		LODConfig.FarEntityQuery_Conditional.SetChunkFilter(&FMassVisualizationChunkFragment::ShouldUpdateVisualizationForChunk);
 	}
 }
 
 template <typename TCollector>
-void CollectLODInfo(UMassEntitySubsystem& EntitySubsystem, FLWComponentSystemExecutionContext& Context, const TArray<FViewerInfo>& Viewers, TCollector& Collector, FLWComponentQuery& HighFrequencyTickingEntityQuery, FLWComponentQuery& LowFrequencyTickingEntityQuery)
+void CollectLODInfo(UMassEntitySubsystem& EntitySubsystem, FMassExecutionContext& Context, const TArray<FViewerInfo>& Viewers, TCollector& Collector, FMassEntityQuery& HighFrequencyTickingEntityQuery, FMassEntityQuery& LowFrequencyTickingEntityQuery)
 {
 	Collector.PrepareExecution(Viewers);
 
-	auto InternalCollectLODInfo = [&Collector](FLWComponentSystemExecutionContext& Context)
+	auto InternalCollectLODInfo = [&Collector](FMassExecutionContext& Context)
 	{
 		TConstArrayView<FDataFragment_Transform> LocationList = Context.GetComponentView<FDataFragment_Transform>();
 		TArrayView<FMassLODInfoFragment> ViewerInfoList = Context.GetMutableComponentView<FMassLODInfoFragment>();
@@ -64,7 +64,7 @@ void CollectLODInfo(UMassEntitySubsystem& EntitySubsystem, FLWComponentSystemExe
 	}
 }
 
-void UMassRepresentationLODCollectorProcessor::Execute(UMassEntitySubsystem& EntitySubsystem, FLWComponentSystemExecutionContext& Context)
+void UMassRepresentationLODCollectorProcessor::Execute(UMassEntitySubsystem& EntitySubsystem, FMassExecutionContext& Context)
 {
 	check(LODManager);
 	const TArray<FViewerInfo>& Viewers = LODManager->GetViewers();
diff --git a/Engine/Plugins/Runtime/MassGameplay/Source/MassLOD/Private/MassSimulationLOD.cpp b/Engine/Plugins/Runtime/MassGameplay/Source/MassLOD/Private/MassSimulationLOD.cpp
--- a/Engine/Plugins/Runtime/MassGameplay/Source/MassLOD/Private/MassSimulationLOD.cpp
+++ b/Engine/Plugins/Runtime/MassGameplay/Source/MassLOD/Private/MassSimulationLOD.cpp
@@ -24,17 +24,17 @@ void UMassProcessor_MassSimulationLODViewersInfo::Initialize(UObject& Owner)
 
 void UMassProcessor_MassSimulationLODViewersInfo::ConfigureQueries()
 {
-	EntityQuery.AddRequirement<FDataFragment_Transform>(ELWComponentAccess::ReadOnly);
-	EntityQuery.AddRequirement<FDataFragment_MassSimulationLODInfo>(ELWComponentAccess::ReadWrite);
+	EntityQuery.AddRequirement<FDataFragment_Transform>(EMassFragmentAccess::ReadOnly);
+	EntityQuery.AddRequirement<FDataFragment_MassSimulationLODInfo>(EMassFragmentAccess::ReadWrite);
 }
 
-void UMassProcessor_MassSimulationLODViewersInfo::Execute(UMassEntitySubsystem& EntitySubsystem, FLWComponentSystemExecutionContext& Context)
+void UMassProcessor_MassSimulationLODViewersInfo::Execute(UMassEntitySubsystem& EntitySubsystem, FMassExecutionContext& Context)
 {
 	check(LODManager);
 	const TArray<FViewerInfo>& Viewers = LODManager->GetViewers();
 	LODCollector.PrepareExecution(Viewers);
 
-	EntityQuery.ForEachEntityChunk(EntitySubsystem, Context, [this](FLWComponentSystemExecutionContext& Context)
+	EntityQuery.ForEachEntityChunk(EntitySubsystem, Context, [this](FMassExecutionContext& Context)
 	{
 		const TConstArrayView<FDataFragment_Transform> LocationList = Context.GetComponentView<FDataFragment_Transform>();
 		const TArrayView<FDataFragment_MassSimulationLODInfo> ViewersInfoList = Context.GetMutableComponentView<FDataFragment_MassSimulationLODInfo>();
@@ -86,12 +86,12 @@ void UMassSimulationLODProcessor::ConfigureQueries()
 	{
 		if(LODConfig.TagFilter.GetScriptStruct())
 		{
-			LODConfig.EntityQuery.AddTagRequirement(*LODConfig.TagFilter.GetScriptStruct(), ELWComponentPresence::All);
+			LODConfig.EntityQuery.AddTagRequirement(*LODConfig.TagFilter.GetScriptStruct(), EMassFragmentPresence::All);
 		}
-		LODConfig.EntityQuery.AddRequirement<FDataFragment_Transform>(ELWComponentAccess::ReadOnly);
-		LODConfig.EntityQuery.AddRequirement<FMassLODInfoFragment>(ELWComponentAccess::ReadOnly);
-		LODConfig.EntityQuery.AddRequirement<FMassSimulationLODFragment>(ELWComponentAccess::ReadWrite);
-		LODConfig.EntityQuery.AddChunkRequirement<FMassSimulationVariableTickChunkFragment>(ELWComponentAccess::ReadOnly);
+		LODConfig.EntityQuery.AddRequirement<FDataFragment_Transform>(EMassFragmentAccess::ReadOnly);
+		LODConfig.EntityQuery.AddRequirement<FMassLODInfoFragment>(EMassFragmentAccess::ReadOnly);
+		LODConfig.EntityQuery.AddRequirement<FMassSimulationLODFragment>(EMassFragmentAccess::ReadWrite);
+		LODConfig.EntityQuery.AddChunkRequirement<FMassSimulationVariableTickChunkFragment>(EMassFragmentAccess::ReadOnly);
 	}
 }
 
@@ -106,7 +106,7 @@ void UMassSimulationLODProcessor::Initialize(UObject& InOwner)
 	Super::Initialize(InOwner);
 }
 
-void UMassSimulationLODProcessor::Execute(UMassEntitySubsystem& EntitySubsystem, FLWComponentSystemExecutionContext& Context)
+void UMassSimulationLODProcessor::Execute(UMassEntitySubsystem& EntitySubsystem, FMassExecutionContext& Context)
 {
 	TRACE_CPUPROFILER_EVENT_SCOPE(TEXT("SimulationLOD"))
 
@@ -116,7 +116,7 @@ void UMassSimulationLODProcessor::Execute(UMassEntitySubsystem& EntitySubsystem,
 	}
 }
 
-void UMassSimulationLODProcessor::CalculateLODForConfig(UMassEntitySubsystem& EntitySubsystem, FLWComponentSystemExecutionContext& Context, FMassSimulationLODConfig& LODConfig)
+void UMassSimulationLODProcessor::CalculateLODForConfig(UMassEntitySubsystem& EntitySubsystem, FMassExecutionContext& Context, FMassSimulationLODConfig& LODConfig)
 {
 	check(LODManager);
 	const TArray<FViewerInfo>& Viewers = LODManager->GetViewers();
@@ -125,7 +125,7 @@ void UMassSimulationLODProcessor::CalculateLODForConfig(UMassEntitySubsystem& En
 	{
 		TRACE_CPUPROFILER_EVENT_SCOPE(TEXT("CalculateLOD"));
 
-		LODConfig.EntityQuery.ForEachEntityChunk(EntitySubsystem, Context, [&LODConfig](FLWComponentSystemExecutionContext& Context)
+		LODConfig.EntityQuery.ForEachEntityChunk(EntitySubsystem, Context, [&LODConfig](FMassExecutionContext& Context)
 		{
 			if (!LODConfig.LODTickRateController.ShouldCalculateLODForChunk(Context))
 			{
@@ -143,7 +143,7 @@ void UMassSimulationLODProcessor::CalculateLODForConfig(UMassEntitySubsystem& En
 
 		if (LODConfig.LODCalculator.AdjustDistancesFromCount())
 		{
-			LODConfig.EntityQuery.ForEachEntityChunk(EntitySubsystem, Context, [&LODConfig](FLWComponentSystemExecutionContext& Context)
+			LODConfig.EntityQuery.ForEachEntityChunk(EntitySubsystem, Context, [&LODConfig](FMassExecutionContext& Context)
 			{
 				if (!LODConfig.LODTickRateController.ShouldAdjustLODFromCountForChunk(Context))
 				{
@@ -161,7 +161,7 @@ void UMassSimulationLODProcessor::CalculateLODForConfig(UMassEntitySubsystem& En
 		TRACE_CPUPROFILER_EVENT_SCOPE(TEXT("VariableTickRates"))
 		check(World);
 		const float Time = World->GetTimeSeconds();
-		LODConfig.EntityQuery.ForEachEntityChunk(EntitySubsystem, Context, [&LODConfig, Time](FLWComponentSystemExecutionContext& Context)
+		LODConfig.EntityQuery.ForEachEntityChunk(EntitySubsystem, Context, [&LODConfig, Time](FMassExecutionContext& Context)
 		{
 			TArrayView<FMassSimulationLODFragment> SimulationLODFragments = Context.GetMutableComponentView<FMassSimulationLODFragment>();
 			LODConfig.LODTickRateController.UpdateTickRateFromLOD(Context, SimulationLODFragments, Time);
@@ -173,7 +173,7 @@ void UMassSimulationLODProcessor::CalculateLODForConfig(UMassEntitySubsystem& En
 	{
 		TRACE_CPUPROFILER_EVENT_SCOPE(TEXT("DebugDisplayLOD"));
 
-		LODConfig.EntityQuery.ForEachEntityChunk(EntitySubsystem, Context, [this, &LODConfig](FLWComponentSystemExecutionContext& Context)
+		LODConfig.EntityQuery.ForEachEntityChunk(EntitySubsystem, Context, [this, &LODConfig](FMassExecutionContext& Context)
 		{
 			const TConstArrayView<FDataFragment_Transform> LocationList = Context.GetComponentView<FDataFragment_Transform>();
 			const TConstArrayView<FMassSimulationLODFragment> SimulationLODList = Context.GetComponentView<FMassSimulationLODFragment>();
diff --git a/Engine/Plugins/Runtime/MassGameplay/Source/MassLOD/Public/MassLODCalculator.h b/Engine/Plugins/Runtime/MassGameplay/Source/MassLOD/Public/MassLODCalculator.h
--- a/Engine/Plugins/Runtime/MassGameplay/Source/MassLOD/Public/MassLODCalculator.h
+++ b/Engine/Plugins/Runtime/MassGameplay/Source/MassLOD/Public/MassLODCalculator.h
@@ -43,7 +43,7 @@ public:
 	 * @Param ViewersInfoList is the source information fragment for LOD calculation
 	 */
 	template< typename FMassLODResultInfo, typename FMassLODSourceInfo >
-	void CalculateLOD(FLWComponentSystemExecutionContext& Context, TArrayView<FMassLODResultInfo> LODList, TConstArrayView<FMassLODSourceInfo> ViewersInfoList);
+	void CalculateLOD(FMassExecutionContext& Context, TArrayView<FMassLODResultInfo> LODList, TConstArrayView<FMassLODSourceInfo> ViewersInfoList);
 
 	/**
 	 * Adjust LOD distances by clamping them to respect the maximum LOD count
@@ -58,7 +58,7 @@ public:
 	 * @Param ViewersInfoList is the source information fragment for LOD calculation
 	 */
 	template< typename FMassLODResultInfo, typename FMassLODSourceInfo >
-	void AdjustLODFromCount(FLWComponentSystemExecutionContext& Context, TArrayView<FMassLODResultInfo> LODList, TConstArrayView<FMassLODSourceInfo> ViewersInfoList);
+	void AdjustLODFromCount(FMassExecutionContext& Context, TArrayView<FMassLODResultInfo> LODList, TConstArrayView<FMassLODSourceInfo> ViewersInfoList);
 
 	/**
 	 * Turn Off all LOD, called for each entity chunks
@@ -66,7 +66,7 @@ public:
 	 * @Param LODList is the fragment where calculation are stored
 	 */
 	template <typename FMassLODResultInfo>
-	void ForceOffLOD(FLWComponentSystemExecutionContext& Context, TArrayView<FMassLODResultInfo> LODList);
+	void ForceOffLOD(FMassExecutionContext& Context, TArrayView<FMassLODResultInfo> LODList);
 
 	/**
 	 * Debug draw the current state of each agent as a color coded square
@@ -76,7 +76,7 @@ public:
 	 * @Param World where the debug display should be drawn
 	 */
 	template< typename FMassLODResultInfo, typename TFragmentLocation >
-	void DebugDisplayLOD(FLWComponentSystemExecutionContext& Context, TConstArrayView<FMassLODResultInfo> LODList, TConstArrayView<TFragmentLocation> LocationList, UWorld* World);
+	void DebugDisplayLOD(FMassExecutionContext& Context, TConstArrayView<FMassLODResultInfo> LODList, TConstArrayView<TFragmentLocation> LocationList, UWorld* World);
 
 protected:
 
@@ -261,7 +261,7 @@ float TMassLODCalculator<FLODLogic>::AccumulateCountInRuntimeData(const EMassLOD
 
 template <typename FLODLogic>
 template< typename FMassLODResultInfo, typename FMassLODSourceInfo >
-void TMassLODCalculator<FLODLogic>::CalculateLOD(FLWComponentSystemExecutionContext& Context, TArrayView<FMassLODResultInfo> LODList, TConstArrayView<FMassLODSourceInfo> ViewersInfoList)
+void TMassLODCalculator<FLODLogic>::CalculateLOD(FMassExecutionContext& Context, TArrayView<FMassLODResultInfo> LODList, TConstArrayView<FMassLODSourceInfo> ViewersInfoList)
 {
 #if WITH_MASSGAMEPLAY_DEBUG
 	if (UE::MassLOD::Debug::bLODCalculationsPaused)
@@ -445,7 +445,7 @@ bool TMassLODCalculator<FLODLogic>::AdjustDistancesFromCount()
 
 template <typename FLODLogic>
 template< typename FMassLODResultInfo, typename FMassLODSourceInfo >
-void TMassLODCalculator<FLODLogic>::AdjustLODFromCount(FLWComponentSystemExecutionContext& Context, TArrayView<FMassLODResultInfo> LODList, TConstArrayView<FMassLODSourceInfo> ViewersInfoList)
+void TMassLODCalculator<FLODLogic>::AdjustLODFromCount(FMassExecutionContext& Context, TArrayView<FMassLODResultInfo> LODList, TConstArrayView<FMassLODSourceInfo> ViewersInfoList)
 {
 	const int32 NumEntities = Context.GetEntitiesNum();
 	// Adjust LOD for each viewer and remember the new highest
@@ -516,7 +516,7 @@ void TMassLODCalculator<FLODLogic>::AdjustLODFromCount(FLWComponentSystemExecuti
 
 template <typename FLODLogic>
 template <typename FMassLODResultInfo>
-void TMassLODCalculator<FLODLogic>::ForceOffLOD(FLWComponentSystemExecutionContext& Context, TArrayView<FMassLODResultInfo> LODList)
+void TMassLODCalculator<FLODLogic>::ForceOffLOD(FMassExecutionContext& Context, TArrayView<FMassLODResultInfo> LODList)
 {
 	const int32 NumEntities = Context.GetEntitiesNum();
 	for (int EntityIdx = 0; EntityIdx < NumEntities; EntityIdx++)
@@ -537,7 +537,7 @@ void TMassLODCalculator<FLODLogic>::ForceOffLOD(FLWComponentSystemExecutionConte
 
 template <typename FLODLogic>
 template< typename FMassLODResultInfo, typename TFragmentLocation >
-void TMassLODCalculator<FLODLogic>::DebugDisplayLOD(FLWComponentSystemExecutionContext& Context, TConstArrayView<FMassLODResultInfo> LODList, TConstArrayView<TFragmentLocation> LocationList, UWorld* World)
+void TMassLODCalculator<FLODLogic>::DebugDisplayLOD(FMassExecutionContext& Context, TConstArrayView<FMassLODResultInfo> LODList, TConstArrayView<TFragmentLocation> LocationList, UWorld* World)
 {
 	const int32 NumEntities = Context.GetEntitiesNum();
 	for (int EntityIdx = 0; EntityIdx < NumEntities; EntityIdx++)
diff --git a/Engine/Plugins/Runtime/MassGameplay/Source/MassLOD/Public/MassLODCollector.h b/Engine/Plugins/Runtime/MassGameplay/Source/MassLOD/Public/MassLODCollector.h
--- a/Engine/Plugins/Runtime/MassGameplay/Source/MassLOD/Public/MassLODCollector.h
+++ b/Engine/Plugins/Runtime/MassGameplay/Source/MassLOD/Public/MassLODCollector.h
@@ -45,7 +45,7 @@ struct TMassLODCollector : public FMassLODBaseLogic
 	 * @Param ViewersInfoList is the fragment where to store source information for LOD calculation
 	 */
 	template< typename FMassTransform, typename FMassLODSourceInfo >
-	void CollectLODInfo(FLWComponentSystemExecutionContext& Context, TConstArrayView<FMassTransform> LocationList, TArrayView<FMassLODSourceInfo> ViewersInfoList);
+	void CollectLODInfo(FMassExecutionContext& Context, TConstArrayView<FMassTransform> LocationList, TArrayView<FMassLODSourceInfo> ViewersInfoList);
 
 protected:
 
@@ -71,7 +71,7 @@ void TMassLODCollector<FLODLogic>::PrepareExecution(TConstArrayView<FViewerInfo>
 
 template <typename FLODLogic>
 template< typename FMassTransform, typename FMassLODSourceInfo >
-void TMassLODCollector<FLODLogic>::CollectLODInfo(FLWComponentSystemExecutionContext& Context, TConstArrayView<FMassTransform> LocationList, TArrayView<FMassLODSourceInfo> ViewersInfoList)
+void TMassLODCollector<FLODLogic>::CollectLODInfo(FMassExecutionContext& Context, TConstArrayView<FMassTransform> LocationList, TArrayView<FMassLODSourceInfo> ViewersInfoList)
 {
 	const int32 NumEntities = Context.GetEntitiesNum();
 	for (int EntityIdx = 0; EntityIdx < NumEntities; EntityIdx++)
diff --git a/Engine/Plugins/Runtime/MassGameplay/Source/MassLOD/Public/MassLODCollectorProcessor.h b/Engine/Plugins/Runtime/MassGameplay/Source/MassLOD/Public/MassLODCollectorProcessor.h
--- a/Engine/Plugins/Runtime/MassGameplay/Source/MassLOD/Public/MassLODCollectorProcessor.h
+++ b/Engine/Plugins/Runtime/MassGameplay/Source/MassLOD/Public/MassLODCollectorProcessor.h
@@ -26,17 +26,17 @@ struct FMassCollectorLODConfig
 
 	/** Runtime data for matching the LOD config */
 	TMassLODCollector<FMassRepresentationLODLogic> RepresentationLODCollector;
-	FLWComponentQuery CloseEntityQuery;
-	FLWComponentQuery FarEntityQuery_Conditional;
+	FMassEntityQuery CloseEntityQuery;
+	FMassEntityQuery FarEntityQuery_Conditional;
 	TMassLODCollector<FMassSimulationLODLogic> SimulationLODCollector;
-	FLWComponentQuery OnLODEntityQuery;
-	FLWComponentQuery OffLODEntityQuery_Conditional;
+	FMassEntityQuery OnLODEntityQuery;
+	FMassEntityQuery OffLODEntityQuery_Conditional;
 	TMassLODCollector<FMassCombinedLODLogic> CombinedLODCollector;
 	// Keep these in this order as we are creating ArrayView from the first 3 queries.
-	FLWComponentQuery CloseAndOnLODEntityQuery;
-	FLWComponentQuery CloseAndOffLODEntityQuery;
-	FLWComponentQuery FarAndOnLODEntityQuery;
-	FLWComponentQuery FarAndOffLODEntityQuery_Conditional;
+	FMassEntityQuery CloseAndOnLODEntityQuery;
+	FMassEntityQuery CloseAndOffLODEntityQuery;
+	FMassEntityQuery FarAndOnLODEntityQuery;
+	FMassEntityQuery FarAndOffLODEntityQuery_Conditional;
 };
 
 /*
@@ -53,7 +53,7 @@ public:
 protected:
 	virtual void Initialize(UObject& Owner) override;
 	virtual void ConfigureQueries() override;
-	virtual void Execute(UMassEntitySubsystem& EntitySubsystem, FLWComponentSystemExecutionContext& Context) override;
+	virtual void Execute(UMassEntitySubsystem& EntitySubsystem, FMassExecutionContext& Context) override;
 
 	UPROPERTY(EditAnywhere, Category = "Mass|LOD", config)
 	TArray<FMassCollectorLODConfig> LODConfigs;
diff --git a/Engine/Plugins/Runtime/MassGameplay/Source/MassLOD/Public/MassLODManager.h b/Engine/Plugins/Runtime/MassGameplay/Source/MassLOD/Public/MassLODManager.h
--- a/Engine/Plugins/Runtime/MassGameplay/Source/MassLOD/Public/MassLODManager.h
+++ b/Engine/Plugins/Runtime/MassGameplay/Source/MassLOD/Public/MassLODManager.h
@@ -16,7 +16,7 @@ class UMassLODManager;
  * Base mass LOD processor to store common information for all LOD processors
  */
 UCLASS(abstract)
-class MASSLOD_API UMassProcessor_LODBase : public UPipeProcessor
+class MASSLOD_API UMassProcessor_LODBase : public UMassProcessor
 {
 	GENERATED_BODY()
 
diff --git a/Engine/Plugins/Runtime/MassGameplay/Source/MassLOD/Public/MassLODTickRateController.h b/Engine/Plugins/Runtime/MassGameplay/Source/MassLOD/Public/MassLODTickRateController.h
--- a/Engine/Plugins/Runtime/MassGameplay/Source/MassLOD/Public/MassLODTickRateController.h
+++ b/Engine/Plugins/Runtime/MassGameplay/Source/MassLOD/Public/MassLODTickRateController.h
@@ -26,14 +26,14 @@ public:
 	 *
 	 * @return if the LOD needs to be calculated
 	 */
-	bool ShouldCalculateLODForChunk(const FLWComponentSystemExecutionContext& Context) const;
+	bool ShouldCalculateLODForChunk(const FMassExecutionContext& Context) const;
 
 	/**
 	 * Retrieve if it is needed to adjust LOD from the newly calculated count for this chunk
 	 * 
 	 * @return if the LOD needs to be adjusted
 	 */
-	bool ShouldAdjustLODFromCountForChunk(const FLWComponentSystemExecutionContext& Context) const;
+	bool ShouldAdjustLODFromCountForChunk(const FMassExecutionContext& Context) const;
 
 	/**
 	 * Updates tick rate for this chunk and its entities
@@ -43,7 +43,7 @@ public:
 	 * @return bool return if the chunk should be tick this frame
 	 */
 	template <typename FMassLODResultInfo>
-	bool UpdateTickRateFromLOD(FLWComponentSystemExecutionContext& Context, TArrayView<FMassLODResultInfo>& LODList, const float Time);
+	bool UpdateTickRateFromLOD(FMassExecutionContext& Context, TArrayView<FMassLODResultInfo>& LODList, const float Time);
 
 protected:
 
@@ -65,7 +65,7 @@ void TMassLODTickRateController<FVariableTickChunkFragment,FLODLogic>::Initializ
 }
 
 template <typename FVariableTickChunkFragment, typename FLODLogic>
-bool TMassLODTickRateController<FVariableTickChunkFragment, FLODLogic>::ShouldCalculateLODForChunk(const FLWComponentSystemExecutionContext& Context) const
+bool TMassLODTickRateController<FVariableTickChunkFragment, FLODLogic>::ShouldCalculateLODForChunk(const FMassExecutionContext& Context) const
 {
 	// EMassLOD::Off does not need to handle max count, so we can use ticking rate for them if available
 	const FMassVariableTickChunkFragment& ChunkData = Context.GetChunkComponent<FVariableTickChunkFragment>();
@@ -73,7 +73,7 @@ bool TMassLODTickRateController<FVariableTickChunkFragment, FLODLogic>::ShouldCa
 }
 
 template <typename FVariableTickChunkFragment, typename FLODLogic>
-bool TMassLODTickRateController<FVariableTickChunkFragment, FLODLogic>::ShouldAdjustLODFromCountForChunk(const FLWComponentSystemExecutionContext& Context) const
+bool TMassLODTickRateController<FVariableTickChunkFragment, FLODLogic>::ShouldAdjustLODFromCountForChunk(const FMassExecutionContext& Context) const
 {
 	// EMassLOD::Off does not need to handle max count, so we can skip it
 	const FMassVariableTickChunkFragment& ChunkData = Context.GetChunkComponent<FVariableTickChunkFragment>();
@@ -82,7 +82,7 @@ bool TMassLODTickRateController<FVariableTickChunkFragment, FLODLogic>::ShouldAd
 
 template <typename FVariableTickChunkFragment, typename FLODLogic>
 template <typename FMassLODResultInfo>
-bool TMassLODTickRateController<FVariableTickChunkFragment, FLODLogic>::UpdateTickRateFromLOD(FLWComponentSystemExecutionContext& Context, TArrayView<FMassLODResultInfo>& LODList, const float Time)
+bool TMassLODTickRateController<FVariableTickChunkFragment, FLODLogic>::UpdateTickRateFromLOD(FMassExecutionContext& Context, TArrayView<FMassLODResultInfo>& LODList, const float Time)
 {
 
 	bool bShouldTickThisFrame = true;
@@ -137,7 +137,7 @@ bool TMassLODTickRateController<FVariableTickChunkFragment, FLODLogic>::UpdateTi
 			SetLastTickedTime<FLODLogic::bDoVariableTickRate>(EntityLOD, Time);
 			if (EntityLOD.LOD != ChunkLOD)
 			{
-				FLWEntity Entity = Context.GetEntity(Index);
+				FMassEntityHandle Entity = Context.GetEntity(Index);
 				Context.Defer().PushCommand(FCommandSwapTags(Entity, UE::MassLOD::GetLODTagFromLOD(ChunkLOD), UE::MassLOD::GetLODTagFromLOD(EntityLOD.LOD)));
 			}
 		}
diff --git a/Engine/Plugins/Runtime/MassGameplay/Source/MassLOD/Public/MassLODTypes.h b/Engine/Plugins/Runtime/MassGameplay/Source/MassLOD/Public/MassLODTypes.h
--- a/Engine/Plugins/Runtime/MassGameplay/Source/MassLOD/Public/MassLODTypes.h
+++ b/Engine/Plugins/Runtime/MassGameplay/Source/MassLOD/Public/MassLODTypes.h
@@ -49,25 +49,25 @@ namespace EMassLOD
 }
 
 USTRUCT()
-struct MASSLOD_API FMassHighLODTag : public FComponentTag
+struct MASSLOD_API FMassHighLODTag : public FMassTag
 {
 	GENERATED_BODY()
 };
 
 USTRUCT()
-struct MASSLOD_API FMassMediumLODTag : public FComponentTag
+struct MASSLOD_API FMassMediumLODTag : public FMassTag
 {
 	GENERATED_BODY()
 };
 
 USTRUCT()
-struct MASSLOD_API FMassLowLODTag : public FComponentTag
+struct MASSLOD_API FMassLowLODTag : public FMassTag
 {
 	GENERATED_BODY()
 };
 
 USTRUCT()
-struct MASSLOD_API FMassOffLODTag : public FComponentTag
+struct MASSLOD_API FMassOffLODTag : public FMassTag
 {
 	GENERATED_BODY()
 };
@@ -76,7 +76,7 @@ struct MASSLOD_API FMassOffLODTag : public FComponentTag
  * Data fragment to store the calculated distances to viewers
  */
 USTRUCT()
-struct MASSLOD_API FMassLODInfoFragment : public FLWComponentData
+struct MASSLOD_API FMassLODInfoFragment : public FMassFragment
 {
 	GENERATED_BODY()
 
@@ -88,7 +88,7 @@ struct MASSLOD_API FMassLODInfoFragment : public FLWComponentData
 };
 
 USTRUCT()
-struct MASSLOD_API FMassVariableTickChunkFragment : public FLWChunkComponent
+struct MASSLOD_API FMassVariableTickChunkFragment : public FMassChunkFragment
 {
 	GENERATED_BODY();
 
@@ -141,25 +141,25 @@ enum class EMassVisibility : uint8
 };
 
 USTRUCT()
-struct MASSLOD_API FMassVisibilityCanBeSeenTag : public FComponentTag
+struct MASSLOD_API FMassVisibilityCanBeSeenTag : public FMassTag
 {
 	GENERATED_BODY()
 };
 
 USTRUCT()
-struct MASSLOD_API FMassVisibilityCulledByFrustumTag : public FComponentTag
+struct MASSLOD_API FMassVisibilityCulledByFrustumTag : public FMassTag
 {
 	GENERATED_BODY()
 };
 
 USTRUCT()
-struct MASSLOD_API FMassVisibilityCulledByDistanceTag : public FComponentTag
+struct MASSLOD_API FMassVisibilityCulledByDistanceTag : public FMassTag
 {
 	GENERATED_BODY()
 };
 
 USTRUCT()
-struct MASSLOD_API FMassVisualizationChunkFragment : public FLWChunkComponent
+struct MASSLOD_API FMassVisualizationChunkFragment : public FMassChunkFragment
 {
 	GENERATED_BODY();
 
@@ -170,7 +170,7 @@ struct MASSLOD_API FMassVisualizationChunkFragment : public FLWChunkComponent
 	 * @param Context of the execution from the entity sub system
 	 * @return true if there is a possibility that the chunk contains a visible entity
 	 */
-	static bool AreAnyEntitiesVisibleInChunk(const FLWComponentSystemExecutionContext& Context)
+	static bool AreAnyEntitiesVisibleInChunk(const FMassExecutionContext& Context)
 	{
 		return Context.GetChunkComponent<FMassVisualizationChunkFragment>().AreAnyEntitiesVisible();
 	}
@@ -191,7 +191,7 @@ struct MASSLOD_API FMassVisualizationChunkFragment : public FLWChunkComponent
 	* @param Context of the execution from the entity sub system
 	 * @return bool true if we should update the representation type for this chunk
 	 */
-	static bool ShouldUpdateVisualizationForChunk(const FLWComponentSystemExecutionContext& Context)
+	static bool ShouldUpdateVisualizationForChunk(const FMassExecutionContext& Context)
 	{
 		return Context.GetChunkComponent<FMassVisualizationChunkFragment>().ShouldUpdateVisualization();
 	}
diff --git a/Engine/Plugins/Runtime/MassGameplay/Source/MassLOD/Public/MassLODUtils.h b/Engine/Plugins/Runtime/MassGameplay/Source/MassLOD/Public/MassLODUtils.h
--- a/Engine/Plugins/Runtime/MassGameplay/Source/MassLOD/Public/MassLODUtils.h
+++ b/Engine/Plugins/Runtime/MassGameplay/Source/MassLOD/Public/MassLODUtils.h
@@ -13,7 +13,7 @@ MASSLOD_API DECLARE_LOG_CATEGORY_EXTERN(LogMassLOD, Log, All);
 namespace UE::MassLOD 
 {
 
-inline EMassLOD::Type GetLODFromArchetype(const FLWComponentSystemExecutionContext& Context)
+inline EMassLOD::Type GetLODFromArchetype(const FMassExecutionContext& Context)
 {
 	if (Context.DoesArchetypeHaveTag<FMassOffLODTag>())
 	{
@@ -53,7 +53,7 @@ inline const UScriptStruct* GetLODTagFromLOD(EMassLOD::Type LOD)
 	}
 }
 
-inline bool IsLODTagSet(const FLWComponentSystemExecutionContext& Context, EMassLOD::Type LOD)
+inline bool IsLODTagSet(const FMassExecutionContext& Context, EMassLOD::Type LOD)
 {
 	switch (LOD)
 	{
diff --git a/Engine/Plugins/Runtime/MassGameplay/Source/MassLOD/Public/MassRepresentationLODCollectorProcessor.h b/Engine/Plugins/Runtime/MassGameplay/Source/MassLOD/Public/MassRepresentationLODCollectorProcessor.h
--- a/Engine/Plugins/Runtime/MassGameplay/Source/MassLOD/Public/MassRepresentationLODCollectorProcessor.h
+++ b/Engine/Plugins/Runtime/MassGameplay/Source/MassLOD/Public/MassRepresentationLODCollectorProcessor.h
@@ -27,8 +27,8 @@ struct FMassRepresentationLODCollectorConfig
 
 	/** Runtime data for matching the LOD config */
 	TMassLODCollector<FMassRepresentationLODLogic> RepresentationLODCollector;
-	FLWComponentQuery CloseEntityQuery;
-	FLWComponentQuery FarEntityQuery_Conditional;
+	FMassEntityQuery CloseEntityQuery;
+	FMassEntityQuery FarEntityQuery_Conditional;
 };
 
 /*
@@ -45,7 +45,7 @@ class MASSLOD_API UMassRepresentationLODCollectorProcessor : public UMassProcess
 	protected:
 	virtual void Initialize(UObject& Owner) override;
 	virtual void ConfigureQueries() override;
-	virtual void Execute(UMassEntitySubsystem& EntitySubsystem, FLWComponentSystemExecutionContext& Context) override;
+	virtual void Execute(UMassEntitySubsystem& EntitySubsystem, FMassExecutionContext& Context) override;
 
 	UPROPERTY(EditAnywhere, Category = "Mass|LOD", config)
 	TArray<FMassRepresentationLODCollectorConfig> LODConfigs;
diff --git a/Engine/Plugins/Runtime/MassGameplay/Source/MassLOD/Public/MassSimulationLOD.h b/Engine/Plugins/Runtime/MassGameplay/Source/MassLOD/Public/MassSimulationLOD.h
--- a/Engine/Plugins/Runtime/MassGameplay/Source/MassLOD/Public/MassSimulationLOD.h
+++ b/Engine/Plugins/Runtime/MassGameplay/Source/MassLOD/Public/MassSimulationLOD.h
@@ -16,7 +16,7 @@
  * @todo To be removed and convert to new Simulation LOD stuff
  */
 USTRUCT()
-struct MASSLOD_API FDataFragment_MassSimulationLODInfo : public FLWComponentData
+struct MASSLOD_API FDataFragment_MassSimulationLODInfo : public FMassFragment
 {
 	GENERATED_BODY()
 
@@ -43,14 +43,14 @@ public:
 protected:
 	virtual void Initialize(UObject& Owner) override;
 	virtual void ConfigureQueries() override;
-	virtual void Execute(UMassEntitySubsystem& EntitySubsystem, FLWComponentSystemExecutionContext& Context) override;
+	virtual void Execute(UMassEntitySubsystem& EntitySubsystem, FMassExecutionContext& Context) override;
 
 	TMassLODCollector<FSimulationLODLogic> LODCollector;
-	FLWComponentQuery EntityQuery;
+	FMassEntityQuery EntityQuery;
 };
 
 USTRUCT()
-struct MASSLOD_API FMassSimulationLODFragment : public FLWComponentData
+struct MASSLOD_API FMassSimulationLODFragment : public FMassFragment
 {
 	GENERATED_BODY()
 
@@ -73,13 +73,13 @@ struct FMassSimulationVariableTickChunkFragment : public FMassVariableTickChunkF
 {
 	GENERATED_BODY();
 
-	static bool ShouldTickChunkThisFrame(const FLWComponentSystemExecutionContext& Context)
+	static bool ShouldTickChunkThisFrame(const FMassExecutionContext& Context)
 	{
 		const FMassSimulationVariableTickChunkFragment& ChunkFragment = Context.GetChunkComponent<FMassSimulationVariableTickChunkFragment>();
 		return ChunkFragment.ShouldTickThisFrame();
 	}
 
-	static EMassLOD::Type GetChunkLOD(const FLWComponentSystemExecutionContext& Context)
+	static EMassLOD::Type GetChunkLOD(const FMassExecutionContext& Context)
 	{
 		const FMassSimulationVariableTickChunkFragment& ChunkFragment = Context.GetChunkComponent<FMassSimulationVariableTickChunkFragment>();
 		return ChunkFragment.GetLOD();
@@ -115,7 +115,7 @@ struct FMassSimulationLODConfig
 	float TickRates[EMassLOD::Max];
 
 	/** Runtime data for matching the LOD config */
-	FLWComponentQuery EntityQuery;
+	FMassEntityQuery EntityQuery;
 	TMassLODCalculator<FMassSimulationLODLogic> LODCalculator;
 	TMassLODTickRateController<FMassSimulationVariableTickChunkFragment, FMassSimulationLODLogic> LODTickRateController;
 };
@@ -132,9 +132,9 @@ protected:
 
 	virtual void ConfigureQueries() override;
 	virtual void Initialize(UObject& InOwner) override;
-	virtual void Execute(UMassEntitySubsystem& EntitySubsystem, FLWComponentSystemExecutionContext& Context) override;
+	virtual void Execute(UMassEntitySubsystem& EntitySubsystem, FMassExecutionContext& Context) override;
 
-	void CalculateLODForConfig(UMassEntitySubsystem& EntitySubsystem, FLWComponentSystemExecutionContext& Context, FMassSimulationLODConfig& LODConfig);
+	void CalculateLODForConfig(UMassEntitySubsystem& EntitySubsystem, FMassExecutionContext& Context, FMassSimulationLODConfig& LODConfig);
 
 	UPROPERTY(EditAnywhere, Category = "Mass|LOD", config)
 	TArray<FMassSimulationLODConfig> LODConfigs;
diff --git a/Engine/Plugins/Runtime/MassGameplay/Source/MassMovement/Public/MassMovementFragments.h b/Engine/Plugins/Runtime/MassGameplay/Source/MassMovement/Public/MassMovementFragments.h
--- a/Engine/Plugins/Runtime/MassGameplay/Source/MassMovement/Public/MassMovementFragments.h
+++ b/Engine/Plugins/Runtime/MassGameplay/Source/MassMovement/Public/MassMovementFragments.h
@@ -8,7 +8,7 @@
 #include "MassMovementFragments.generated.h"
 
 USTRUCT()
-struct MASSMOVEMENT_API FMassVelocityFragment : public FLWComponentData
+struct MASSMOVEMENT_API FMassVelocityFragment : public FMassFragment
 {
 	GENERATED_BODY()
 
diff --git a/Engine/Plugins/Runtime/MassGameplay/Source/MassMovement/Public/MassMovementTypes.h b/Engine/Plugins/Runtime/MassGameplay/Source/MassMovement/Public/MassMovementTypes.h
--- a/Engine/Plugins/Runtime/MassGameplay/Source/MassMovement/Public/MassMovementTypes.h
+++ b/Engine/Plugins/Runtime/MassGameplay/Source/MassMovement/Public/MassMovementTypes.h
@@ -25,7 +25,7 @@ struct FMassAvoidanceObstacleItem
 		return Entity == Other.Entity;
 	}
 
-	FLWEntity Entity;
+	FMassEntityHandle Entity;
 	EMassAvoidanceObstacleItemFlags ItemFlags = EMassAvoidanceObstacleItemFlags::None;
 };
 
diff --git a/Engine/Plugins/Runtime/MassGameplay/Source/MassReplication/Private/MassReplicationManager.cpp b/Engine/Plugins/Runtime/MassGameplay/Source/MassReplication/Private/MassReplicationManager.cpp
--- a/Engine/Plugins/Runtime/MassGameplay/Source/MassReplication/Private/MassReplicationManager.cpp
+++ b/Engine/Plugins/Runtime/MassGameplay/Source/MassReplication/Private/MassReplicationManager.cpp
@@ -423,7 +423,7 @@ FMassBubbleInfoClassHandle UMassReplicationManager::GetBubbleInfoClassHandle(con
 }
 
 #if UE_REPLICATION_COMPILE_CLIENT_CODE
-void UMassReplicationManager::SetEntity(const FMassNetworkID NetworkID, const FLWEntity Entity)
+void UMassReplicationManager::SetEntity(const FMassNetworkID NetworkID, const FMassEntityHandle Entity)
 {
 	FMassReplicationEntityInfo* EntityInfo = FindMassEntityInfoMutable(NetworkID);
 	check(EntityInfo && !EntityInfo->Entity.IsSet());
@@ -434,9 +434,9 @@ void UMassReplicationManager::SetEntity(const FMassNetworkID NetworkID, const FL
 #endif // UE_REPLICATION_COMPILE_CLIENT_CODE
 
 #if UE_REPLICATION_COMPILE_CLIENT_CODE
-FLWEntity UMassReplicationManager::ResetEntityIfValid(const FMassNetworkID NetworkID, int32 ReplicationID)
+FMassEntityHandle UMassReplicationManager::ResetEntityIfValid(const FMassNetworkID NetworkID, int32 ReplicationID)
 {
-	FLWEntity EntityReset;
+	FMassEntityHandle EntityReset;
 
 	FMassReplicationEntityInfo* EntityInfo = FindMassEntityInfoMutable(NetworkID);
 
@@ -452,7 +452,7 @@ FLWEntity UMassReplicationManager::ResetEntityIfValid(const FMassNetworkID Netwo
 		EntityReset = EntityInfo->Entity;
 
 		//Unset the Entity handle, this indicates that its currently removed from the bubble
-		EntityInfo->Entity = FLWEntity();
+		EntityInfo->Entity = FMassEntityHandle();
 	}
 
 	return EntityReset;
@@ -489,7 +489,7 @@ UMassReplicationManager::EFindOrAddMassEntityInfo UMassReplicationManager::FindA
 	}
 	else
 	{
-		MassEntityInfo = &EntityInfoMap.Add(NetworkID, FMassReplicationEntityInfo(FLWEntity(), ReplicationID));
+		MassEntityInfo = &EntityInfoMap.Add(NetworkID, FMassReplicationEntityInfo(FMassEntityHandle(), ReplicationID));
 		FindOrAddStatus = EFindOrAddMassEntityInfo::Added;
 	}
 
@@ -515,11 +515,11 @@ FMassReplicationEntityInfo* UMassReplicationManager::FindMassEntityInfoMutable(c
 #endif // UE_REPLICATION_COMPILE_CLIENT_CODE
 
 #if UE_REPLICATION_COMPILE_CLIENT_CODE
-FLWEntity UMassReplicationManager::FindEntity(const FMassNetworkID NetworkID) const
+FMassEntityHandle UMassReplicationManager::FindEntity(const FMassNetworkID NetworkID) const
 {
 	check(NetworkID.IsValid());
 	const FMassReplicationEntityInfo* Info = EntityInfoMap.Find(NetworkID);
-	return Info ? Info->Entity : FLWEntity();
+	return Info ? Info->Entity : FMassEntityHandle();
 }
 #endif // UE_REPLICATION_COMPILE_CLIENT_CODE
 
