AggregateTicking and Pipe type renames:

diff --git a/Engine/Plugins/Runtime/MassGameplay/Source/MassReplication/Private/MassReplicationProcessorBase.cpp b/Engine/Plugins/Runtime/MassGameplay/Source/MassReplication/Private/MassReplicationProcessorBase.cpp
--- a/Engine/Plugins/Runtime/MassGameplay/Source/MassReplication/Private/MassReplicationProcessorBase.cpp
+++ b/Engine/Plugins/Runtime/MassGameplay/Source/MassReplication/Private/MassReplicationProcessorBase.cpp
@@ -30,16 +30,16 @@ UMassReplicationProcessorBase::UMassReplicationProcessorBase()
 	UpdateInterval[EMassLOD::Off] = 0.5f;
 
 	ExecutionFlags = int32(EProcessorExecutionFlags::None);
-	ProcessingPhase = EPipeProcessingPhase::PostPhysics;
+	ProcessingPhase = EMassProcessingPhase::PostPhysics;
 }
 
 void UMassReplicationProcessorBase::ConfigureQueries()
 {
-	EntityQuery.AddRequirement<FDataFragment_MassSimulationLODInfo>(ELWComponentAccess::ReadOnly);
-	EntityQuery.AddRequirement<FMassNetworkIDFragment>(ELWComponentAccess::ReadOnly);
-	EntityQuery.AddRequirement<FDataFragment_ReplicationTemplateID>(ELWComponentAccess::ReadOnly);
-	EntityQuery.AddRequirement<FMassReplicationViewerLODFragment>(ELWComponentAccess::ReadWrite);
-	EntityQuery.AddRequirement<FMassReplicatedAgentFragment>(ELWComponentAccess::ReadWrite);
+	EntityQuery.AddRequirement<FDataFragment_MassSimulationLODInfo>(EMassFragmentAccess::ReadOnly);
+	EntityQuery.AddRequirement<FMassNetworkIDFragment>(EMassFragmentAccess::ReadOnly);
+	EntityQuery.AddRequirement<FDataFragment_ReplicationTemplateID>(EMassFragmentAccess::ReadOnly);
+	EntityQuery.AddRequirement<FMassReplicationViewerLODFragment>(EMassFragmentAccess::ReadWrite);
+	EntityQuery.AddRequirement<FMassReplicatedAgentFragment>(EMassFragmentAccess::ReadWrite);
 }
 
 void UMassReplicationProcessorBase::Initialize(UObject& Owner)
@@ -129,14 +129,14 @@ void UMassReplicationProcessorBase::PrepareExecution()
 #endif //UE_REPLICATION_COMPILE_SERVER_CODE
 }
 
-void UMassReplicationProcessorBase::Execute(UMassEntitySubsystem& EntitySubsystem, FLWComponentSystemExecutionContext& Context)
+void UMassReplicationProcessorBase::Execute(UMassEntitySubsystem& EntitySubsystem, FMassExecutionContext& Context)
 {
 #if UE_REPLICATION_COMPILE_SERVER_CODE
 	check(World);
 
 	PrepareExecution();
 
-	EntityQuery.ForEachEntityChunk(EntitySubsystem, Context, [this](FLWComponentSystemExecutionContext& Context)
+	EntityQuery.ForEachEntityChunk(EntitySubsystem, Context, [this](FMassExecutionContext& Context)
 		{
 			const TConstArrayView<FDataFragment_MassSimulationLODInfo> ViewersInfoList = Context.GetComponentView<FDataFragment_MassSimulationLODInfo>();
 			const TArrayView<FMassReplicationViewerLODFragment> ViewerLODList = Context.GetMutableComponentView<FMassReplicationViewerLODFragment>();
@@ -145,7 +145,7 @@ void UMassReplicationProcessorBase::Execute(UMassEntitySubsystem& EntitySubsyste
 
 	if (LODCalculator.AdjustDistancesFromCount())
 	{
-		EntityQuery.ForEachEntityChunk(EntitySubsystem, Context, [this](FLWComponentSystemExecutionContext& Context)
+		EntityQuery.ForEachEntityChunk(EntitySubsystem, Context, [this](FMassExecutionContext& Context)
 			{
 				const TConstArrayView<FDataFragment_MassSimulationLODInfo> ViewersInfoList = Context.GetComponentView<FDataFragment_MassSimulationLODInfo>();
 				const TArrayView<FMassReplicationViewerLODFragment> ViewerLODList = Context.GetMutableComponentView<FMassReplicationViewerLODFragment>();
diff --git a/Engine/Plugins/Runtime/MassGameplay/Source/MassReplication/Private/MassReplicationTransformHandlers.cpp b/Engine/Plugins/Runtime/MassGameplay/Source/MassReplication/Private/MassReplicationTransformHandlers.cpp
--- a/Engine/Plugins/Runtime/MassGameplay/Source/MassReplication/Private/MassReplicationTransformHandlers.cpp
+++ b/Engine/Plugins/Runtime/MassGameplay/Source/MassReplication/Private/MassReplicationTransformHandlers.cpp
@@ -5,12 +5,12 @@
 #include "MassCommonFragments.h"
 #include "AIHelpers.h"
 
-void FMassReplicationProcessorTransformHandlerBase::AddRequirements(FLWComponentQuery& InQuery)
+void FMassReplicationProcessorTransformHandlerBase::AddRequirements(FMassEntityQuery& InQuery)
 {
-	InQuery.AddRequirement<FDataFragment_Transform>(ELWComponentAccess::ReadWrite);
+	InQuery.AddRequirement<FDataFragment_Transform>(EMassFragmentAccess::ReadWrite);
 }
 
-void FMassReplicationProcessorTransformHandlerBase::CacheComponentViews(FLWComponentSystemExecutionContext& ExecContext)
+void FMassReplicationProcessorTransformHandlerBase::CacheComponentViews(FMassExecutionContext& ExecContext)
 {
 	TransformList = ExecContext.GetMutableComponentView<FDataFragment_Transform>();
 }
diff --git a/Engine/Plugins/Runtime/MassGameplay/Source/MassReplication/Private/MassReplicationTypes.cpp b/Engine/Plugins/Runtime/MassGameplay/Source/MassReplication/Private/MassReplicationTypes.cpp
--- a/Engine/Plugins/Runtime/MassGameplay/Source/MassReplication/Private/MassReplicationTypes.cpp
+++ b/Engine/Plugins/Runtime/MassGameplay/Source/MassReplication/Private/MassReplicationTypes.cpp
@@ -18,7 +18,7 @@ UMassNetworkIDFragmentInitializer::UMassNetworkIDFragmentInitializer()
 
 void UMassNetworkIDFragmentInitializer::ConfigureQueries()
 {
-	EntityQuery.AddRequirement<FMassNetworkIDFragment>(ELWComponentAccess::ReadWrite);
+	EntityQuery.AddRequirement<FMassNetworkIDFragment>(EMassFragmentAccess::ReadWrite);
 }
 
 void UMassNetworkIDFragmentInitializer::Initialize(UObject& Owner)
@@ -29,7 +29,7 @@ void UMassNetworkIDFragmentInitializer::Initialize(UObject& Owner)
 #endif //UE_REPLICATION_COMPILE_SERVER_CODE
 }
 
-void UMassNetworkIDFragmentInitializer::Execute(UMassEntitySubsystem& EntitySubsystem, FLWComponentSystemExecutionContext& Context)
+void UMassNetworkIDFragmentInitializer::Execute(UMassEntitySubsystem& EntitySubsystem, FMassExecutionContext& Context)
 {
 	QUICK_SCOPE_CYCLE_COUNTER(MassProcessor_InitNetworkID_Run);
 
@@ -41,7 +41,7 @@ void UMassNetworkIDFragmentInitializer::Execute(UMassEntitySubsystem& EntitySubs
 #if UE_REPLICATION_COMPILE_SERVER_CODE
 		check(ReplicationManager);
 
-		EntityQuery.ForEachEntityChunk(EntitySubsystem, Context, [this](FLWComponentSystemExecutionContext& Context)
+		EntityQuery.ForEachEntityChunk(EntitySubsystem, Context, [this](FMassExecutionContext& Context)
 			{
 				const TArrayView<FMassNetworkIDFragment> NetworkIDList = Context.GetMutableComponentView<FMassNetworkIDFragment>();
 				const int32 NumEntities = Context.GetEntitiesNum();
diff --git a/Engine/Plugins/Runtime/MassGameplay/Source/MassReplication/Public/MassClientBubbleHandler.h b/Engine/Plugins/Runtime/MassGameplay/Source/MassReplication/Public/MassClientBubbleHandler.h
--- a/Engine/Plugins/Runtime/MassGameplay/Source/MassReplication/Public/MassClientBubbleHandler.h
+++ b/Engine/Plugins/Runtime/MassGameplay/Source/MassReplication/Public/MassClientBubbleHandler.h
@@ -49,7 +49,7 @@ protected:
 /** Data that can be accessed from a FMassReplicatedAgentHandle on a server */
 struct FMassAgentLookupData
 {
-	FMassAgentLookupData(FLWEntity InEntity, FMassNetworkID InNetID, int32 InAgentsIdx)
+	FMassAgentLookupData(FMassEntityHandle InEntity, FMassNetworkID InNetID, int32 InAgentsIdx)
 		: Entity(InEntity)
 		, NetID(InNetID)
 		, AgentsIdx(InAgentsIdx)
@@ -57,7 +57,7 @@ struct FMassAgentLookupData
 
 	void Invalidate()
 	{
-		Entity = FLWEntity();
+		Entity = FMassEntityHandle();
 		NetID.Invalidate();
 		AgentsIdx = INDEX_NONE;
 	}
@@ -67,7 +67,7 @@ struct FMassAgentLookupData
 		return Entity.IsSet() && NetID.IsValid() && (AgentsIdx >= 0);
 	}
 
-	FLWEntity Entity;
+	FMassEntityHandle Entity;
 	FMassNetworkID NetID;
 	int32 AgentsIdx  = INDEX_NONE;
 };
@@ -131,10 +131,10 @@ public:
 	template<typename T>
 	friend class TMassClientBubbleTransformHandler;
 
-	typedef TFunctionRef<void(FLWComponentQuery&)> FAddRequirementsForSpawnQueryFunction;
-	typedef TFunctionRef<void(FLWComponentSystemExecutionContext&)> FCacheComponentViewsForSpawnQueryFunction;
-	typedef TFunctionRef<void(const FEntityView&, const typename AgentArrayItem::FReplicatedAgentType&, const int32)> FSetSpawnedEntityDataFunction;
-	typedef TFunctionRef<void(const FEntityView&, const typename AgentArrayItem::FReplicatedAgentType&)> FSetModifiedEntityDataFunction;
+	typedef TFunctionRef<void(FMassEntityQuery&)> FAddRequirementsForSpawnQueryFunction;
+	typedef TFunctionRef<void(FMassExecutionContext&)> FCacheComponentViewsForSpawnQueryFunction;
+	typedef TFunctionRef<void(const FMassEntityView&, const typename AgentArrayItem::FReplicatedAgentType&, const int32)> FSetSpawnedEntityDataFunction;
+	typedef TFunctionRef<void(const FMassEntityView&, const typename AgentArrayItem::FReplicatedAgentType&)> FSetModifiedEntityDataFunction;
 
 	/** This must be called from outside before InitializeForWorld() is called. Its called from agent specific bubble implementations */
 	virtual void Initialize(TArray<AgentArrayItem>& InAgents, FMassClientBubbleSerializerBase& InSerializer);
@@ -142,7 +142,7 @@ public:
 	virtual void InitializeForWorld(UWorld& InWorld) override;
 
 #if UE_REPLICATION_COMPILE_SERVER_CODE
-	FMassReplicatedAgentHandle AddAgent(FLWEntity Entity, typename AgentArrayItem::FReplicatedAgentType& Agent);
+	FMassReplicatedAgentHandle AddAgent(FMassEntityHandle Entity, typename AgentArrayItem::FReplicatedAgentType& Agent);
 
 	bool RemoveAgent(FMassNetworkID NetID);
 	bool RemoveAgent(FMassReplicatedAgentHandle AgentHandle);
@@ -236,7 +236,7 @@ void TClientBubbleHandlerBase<AgentArrayItem>::InitializeForWorld(UWorld& InWorl
 
 #if UE_REPLICATION_COMPILE_SERVER_CODE
 template<typename AgentArrayItem>
-FMassReplicatedAgentHandle TClientBubbleHandlerBase<AgentArrayItem>::AddAgent(FLWEntity Entity, typename AgentArrayItem::FReplicatedAgentType& Agent)
+FMassReplicatedAgentHandle TClientBubbleHandlerBase<AgentArrayItem>::AddAgent(FMassEntityHandle Entity, typename AgentArrayItem::FReplicatedAgentType& Agent)
 {
 	checkf(Agent.GetNetID().IsValid(), TEXT("Agent.NetID must be valid!"));
 	checkf(NetworkIDToAgentHandleMap.Find(Agent.GetNetID()) == nullptr, TEXT("Only add agents once"));
@@ -381,7 +381,7 @@ void TClientBubbleHandlerBase<AgentArrayItem>::UpdateAgentsToRemove()
 template<typename AgentArrayItem>
 void TClientBubbleHandlerBase<AgentArrayItem>::PreReplicatedRemove(const TArrayView<int32> RemovedIndices, int32 FinalSize)
 {
-	TMap<FMassEntityTemplateID, TArray<FLWEntity>> EntitiesToDestroy;
+	TMap<FMassEntityTemplateID, TArray<FMassEntityHandle>> EntitiesToDestroy;
 
 	UWorld* World = Serializer->GetWorld();
 	check(World);
@@ -396,13 +396,13 @@ void TClientBubbleHandlerBase<AgentArrayItem>::PreReplicatedRemove(const TArrayV
 	{
 		const AgentArrayItem& RemovedItem = (*Agents)[Idx];
 
-		FLWEntity Entity = ReplicationManager->ResetEntityIfValid(RemovedItem.Agent.GetNetID(), RemovedItem.ReplicationID);
+		FMassEntityHandle Entity = ReplicationManager->ResetEntityIfValid(RemovedItem.Agent.GetNetID(), RemovedItem.ReplicationID);
 
 		// Only remove the item if its currently Set / Valid and its the most recent ReplicationID. Stale removes after more recent adds are ignored
 		// We do need to check the ReplicationID in this case
 		if (Entity.IsSet())
 		{
-			TArray<FLWEntity>& EntityArray = EntitiesToDestroy.FindOrAdd(RemovedItem.Agent.GetTemplateID());
+			TArray<FMassEntityHandle>& EntityArray = EntitiesToDestroy.FindOrAdd(RemovedItem.Agent.GetTemplateID());
 			EntityArray.Push(Entity);
 
 			check(AgentsRemoveDataMap.Find(RemovedItem.Agent.GetNetID()) == nullptr);
@@ -411,10 +411,10 @@ void TClientBubbleHandlerBase<AgentArrayItem>::PreReplicatedRemove(const TArrayV
 	}
 
 	// Batch destroy agents per template
-	for (const TPair<FMassEntityTemplateID, TArray<FLWEntity>>& Item : EntitiesToDestroy)
+	for (const TPair<FMassEntityTemplateID, TArray<FMassEntityHandle>>& Item : EntitiesToDestroy)
 	{
 		const FMassEntityTemplateID& TemplateID = Item.Key;
-		const TArray<FLWEntity>& EntityArray = Item.Value;
+		const TArray<FMassEntityHandle>& EntityArray = Item.Value;
 		SpawnerSubsystem->DestroyEntities(TemplateID, EntityArray);
 	}
 }
@@ -425,7 +425,7 @@ template<typename AgentArrayItem>
 void TClientBubbleHandlerBase<AgentArrayItem>::PostReplicatedAddHelper(const TArrayView<int32> AddedIndices, FAddRequirementsForSpawnQueryFunction AddRequirementsForSpawnQuery
 	, FCacheComponentViewsForSpawnQueryFunction CacheComponentViewsForSpawnQuery, FSetSpawnedEntityDataFunction SetSpawnedEntityData, FSetModifiedEntityDataFunction SetModifiedEntityData)
 {
-	TArray<FLWEntity> EntitiesDestroy;
+	TArray<FMassEntityHandle> EntitiesDestroy;
 
 	UMassEntitySubsystem* EntitySystem = Serializer->GetEntitySystem();
 	check(EntitySystem);
@@ -454,7 +454,7 @@ void TClientBubbleHandlerBase<AgentArrayItem>::PostReplicatedAddHelper(const TAr
 			// If EntityData IsSet() it means we have had multiple Adds without a remove and we treat this add as modifying an existing agent.
 			if (EntityInfo->Entity.IsSet())
 			{
-				FEntityView EntityView(*EntitySystem, EntityInfo->Entity);
+				FMassEntityView EntityView(*EntitySystem, EntityInfo->Entity);
 
 				SetModifiedEntityData(EntityView, AddedItem.Agent);
 			}
@@ -536,7 +536,7 @@ void TClientBubbleHandlerBase<AgentArrayItem>::PostReplicatedAddEntitiesHelper(c
 		Config.MinNumber = AgentsSpawn.Num();
 		Config.MaxNumber = AgentsSpawn.Num();
 
-		TArray<FLWEntity> Entities;
+		TArray<FMassEntityHandle> Entities;
 
 		SpawnerSubsystem->SpawnEntities(TemplateID, Config, FStructView(), Entities);
 
@@ -545,16 +545,16 @@ void TClientBubbleHandlerBase<AgentArrayItem>::PostReplicatedAddEntitiesHelper(c
 		const FArchetypeHandle& ArchetypeHandle = MassEntityTemplate->GetArchetype();
 
 
-		FLWComponentSystemExecutionContext ExecContext;
-		FLWComponentQuery Query;
+		FMassExecutionContext ExecContext;
+		FMassEntityQuery Query;
 
 		AddRequirementsForSpawnQuery(Query);
 
-		Query.AddRequirement<FMassNetworkIDFragment>(ELWComponentAccess::ReadWrite);
+		Query.AddRequirement<FMassNetworkIDFragment>(EMassFragmentAccess::ReadWrite);
 
 		int32 AgentsSpawnIdx = 0;
 
-		Query.ForEachEntityChunk(FArchetypeChunkCollection(ArchetypeHandle, Entities), *EntitySystem, ExecContext, [&AgentsSpawn, &AgentsSpawnIdx, this, ReplicationManager, &ExecContext, &CacheComponentViewsForSpawnQuery, &SetSpawnedEntityData, &EntitySystem](FLWComponentSystemExecutionContext& Context)
+		Query.ForEachEntityChunk(FArchetypeChunkCollection(ArchetypeHandle, Entities), *EntitySystem, ExecContext, [&AgentsSpawn, &AgentsSpawnIdx, this, ReplicationManager, &ExecContext, &CacheComponentViewsForSpawnQuery, &SetSpawnedEntityData, &EntitySystem](FMassExecutionContext& Context)
 			{
 				CacheComponentViewsForSpawnQuery(ExecContext);
 
@@ -563,14 +563,14 @@ void TClientBubbleHandlerBase<AgentArrayItem>::PostReplicatedAddEntitiesHelper(c
 				for (int32 i = 0; i < Context.GetEntitiesNum(); ++i)
 				{
 					const typename AgentArrayItem::FReplicatedAgentType& AgentSpawn = *AgentsSpawn[AgentsSpawnIdx];
-					const FLWEntity Entity = Context.GetEntity(i);
+					const FMassEntityHandle Entity = Context.GetEntity(i);
 
 					FMassNetworkIDFragment& NetIDFragment = NetworkIDList[i];
 
 					NetIDFragment.NetID = AgentSpawn.GetNetID();
 					ReplicationManager->SetEntity(NetIDFragment.NetID, Entity);
 
-					FEntityView EntityView(*EntitySystem, Entity);
+					FMassEntityView EntityView(*EntitySystem, Entity);
 					SetSpawnedEntityData(EntityView, AgentSpawn, i);
 
 					++AgentsSpawnIdx;
@@ -603,7 +603,7 @@ void TClientBubbleHandlerBase<AgentArrayItem>::PostReplicatedChangeHelper(const
 		// Currently we don't think this should be needed, but are leaving it in for bomb proofing.
 		if (ensure(EntityInfo->ReplicationID == ChangedItem.ReplicationID))
 		{
-			FEntityView EntityView(*EntitySystem, EntityInfo->Entity);
+			FMassEntityView EntityView(*EntitySystem, EntityInfo->Entity);
 			SetModifiedEntityData(EntityView, ChangedItem.Agent);
 		}
 	}
diff --git a/Engine/Plugins/Runtime/MassGameplay/Source/MassReplication/Public/MassReplicationManager.h b/Engine/Plugins/Runtime/MassGameplay/Source/MassReplication/Public/MassReplicationManager.h
--- a/Engine/Plugins/Runtime/MassGameplay/Source/MassReplication/Public/MassReplicationManager.h
+++ b/Engine/Plugins/Runtime/MassGameplay/Source/MassReplication/Public/MassReplicationManager.h
@@ -58,7 +58,7 @@ struct FMassClientBubbleInfoData
 
 namespace UE::MassReplication
 {
-	DECLARE_MULTICAST_DELEGATE_TwoParams(FMassAgentDelegate, FMassNetworkID, FLWEntity);
+	DECLARE_MULTICAST_DELEGATE_TwoParams(FMassAgentDelegate, FMassNetworkID, FMassEntityHandle);
 
 } // UE::MassReplication
 
@@ -169,14 +169,14 @@ public:
 	/** 
 	 * Must be called immediately after an Entity is added to the simulation and after FindAndUpdateOrAddMassEntityInfo has been called
 	 */
-	void SetEntity(const FMassNetworkID NetworkID, const FLWEntity Entity);
+	void SetEntity(const FMassNetworkID NetworkID, const FMassEntityHandle Entity);
 
 	/** 
-	 * Resets the item with NetworkID in the  EntityInfoMap if its currently Set / Valid and its the most recent ReplicationID. In this case the FLWEntity that was reset is returned.
-	 * Otherwise FLWEntity::IsSet() will be false on the returned FLWEntity
-	 * Must be called after SetEntity and just before the Entity is removed on the client. This sets the Entity member of the AddToEntityInfoMap item to FLWEntity so it will be not IsSet() 
+	 * Resets the item with NetworkID in the  EntityInfoMap if its currently Set / Valid and its the most recent ReplicationID. In this case the FMassEntityHandle that was reset is returned.
+	 * Otherwise FMassEntityHandle::IsSet() will be false on the returned FMassEntityHandle
+	 * Must be called after SetEntity and just before the Entity is removed on the client. This sets the Entity member of the AddToEntityInfoMap item to FMassEntityHandle so it will be not IsSet() 
 	 */
-	FLWEntity ResetEntityIfValid(const FMassNetworkID NetworkID, int32 ReplicationID);
+	FMassEntityHandle ResetEntityIfValid(const FMassNetworkID NetworkID, int32 ReplicationID);
 
 	void RemoveFromEntityInfoMap(const FMassNetworkID NetworkID);
 
@@ -197,7 +197,7 @@ public:
 	 */
 	const FMassReplicationEntityInfo* FindMassEntityInfo(const FMassNetworkID NetworkID) const;
 
-	FLWEntity FindEntity(const FMassNetworkID NetworkID) const;
+	FMassEntityHandle FindEntity(const FMassNetworkID NetworkID) const;
 
 	const TMap<FMassNetworkID, FMassReplicationEntityInfo>& GetEntityInfoMap() const { return EntityInfoMap; }
 
diff --git a/Engine/Plugins/Runtime/MassGameplay/Source/MassReplication/Public/MassReplicationProcessorBase.h b/Engine/Plugins/Runtime/MassGameplay/Source/MassReplication/Public/MassReplicationProcessorBase.h
--- a/Engine/Plugins/Runtime/MassGameplay/Source/MassReplication/Public/MassReplicationProcessorBase.h
+++ b/Engine/Plugins/Runtime/MassGameplay/Source/MassReplication/Public/MassReplicationProcessorBase.h
@@ -32,7 +32,7 @@ public:
 protected:
 	virtual void ConfigureQueries() override;
 	virtual void Initialize(UObject& Owner) override;
-	virtual void Execute(UMassEntitySubsystem& EntitySubsystem, FLWComponentSystemExecutionContext& Context) override;
+	virtual void Execute(UMassEntitySubsystem& EntitySubsystem, FMassExecutionContext& Context) override;
 
 	void PrepareExecution();
 
@@ -41,7 +41,7 @@ protected:
 	 *  but Functors can also be used as well as TFunctionRefs etc. Its also fairly straight forward to call member functions via some Lamda glue code
 	 */
 	template<typename AgentArrayItem, typename CacheViewsCallback, typename AddEntityCallback, typename ModifyEntityCallback, typename RemoveEntityCallback>
-	void CalculateClientReplication(FLWComponentSystemExecutionContext& Context, CacheViewsCallback&& CacheViews, AddEntityCallback&& AddEntity, ModifyEntityCallback&& ModifyEntity, RemoveEntityCallback&& RemoveEntity);
+	void CalculateClientReplication(FMassExecutionContext& Context, CacheViewsCallback&& CacheViews, AddEntityCallback&& AddEntity, ModifyEntityCallback&& ModifyEntity, RemoveEntityCallback&& RemoveEntity);
 
 	template<typename T>
 	T& GetTypedClientBubbleInfoChecked(FMassClientHandle Handle);
@@ -83,11 +83,11 @@ protected:
 
 	TMassLODCalculator<FReplicationLODLogic> LODCalculator;
 
-	FLWComponentQuery EntityQuery;
+	FMassEntityQuery EntityQuery;
 };
 
 template<typename AgentArrayItem, typename CacheViewsCallback, typename AddEntityCallback, typename ModifyEntityCallback, typename RemoveEntityCallback>
-void UMassReplicationProcessorBase::CalculateClientReplication(FLWComponentSystemExecutionContext& Context, CacheViewsCallback&& CacheViews, AddEntityCallback&& AddEntity, ModifyEntityCallback&& ModifyEntity, RemoveEntityCallback&& RemoveEntity)
+void UMassReplicationProcessorBase::CalculateClientReplication(FMassExecutionContext& Context, CacheViewsCallback&& CacheViews, AddEntityCallback&& AddEntity, ModifyEntityCallback&& ModifyEntity, RemoveEntityCallback&& RemoveEntity)
 {
 #if UE_REPLICATION_COMPILE_SERVER_CODE
 
diff --git a/Engine/Plugins/Runtime/MassGameplay/Source/MassReplication/Public/MassReplicationTransformHandlers.h b/Engine/Plugins/Runtime/MassGameplay/Source/MassReplication/Public/MassReplicationTransformHandlers.h
--- a/Engine/Plugins/Runtime/MassGameplay/Source/MassReplication/Public/MassReplicationTransformHandlers.h
+++ b/Engine/Plugins/Runtime/MassGameplay/Source/MassReplication/Public/MassReplicationTransformHandlers.h
@@ -11,7 +11,7 @@
 
 #include "MassReplicationTransformHandlers.generated.h"
 
-struct FLWComponentQuery;
+struct FMassEntityQuery;
 
 namespace UE::Mass::Replication
 {
@@ -76,14 +76,14 @@ public:
 	 * When entities are spawned in Mass by the replication system on the client, a spawn query is used to set the data on the spawned entities.
 	 * The following functions are used to configure the query and then set the position and yaw data.
 	 */
-	static void AddRequirementsForSpawnQuery(FLWComponentQuery& InQuery);
-	void CacheComponentViewsForSpawnQuery(FLWComponentSystemExecutionContext& InExecContext);
+	static void AddRequirementsForSpawnQuery(FMassEntityQuery& InQuery);
+	void CacheComponentViewsForSpawnQuery(FMassExecutionContext& InExecContext);
 	void ClearComponentViewsForSpawnQuery();
 
 	void SetSpawnedEntityData(const int32 EntityIdx, const FReplicatedAgentPositionYawData& ReplicatedPathData) const;
 
 	/** Call this when an Entity that has already been spawned is modified on the client */
-	static void SetModifiedEntityData(const FEntityView& EntityView, const FReplicatedAgentPositionYawData& ReplicatedPathData);
+	static void SetModifiedEntityData(const FMassEntityView& EntityView, const FReplicatedAgentPositionYawData& ReplicatedPathData);
 
 	// We could easily add support replicating FReplicatedAgentTransformData here if required
 #endif // UE_REPLICATION_COMPILE_CLIENT_CODE
@@ -143,15 +143,15 @@ void TMassClientBubbleTransformHandler<AgentArrayItem>::SetBubblePositionYawFrom
 
 #if UE_REPLICATION_COMPILE_CLIENT_CODE
 template<typename AgentArrayItem>
-void TMassClientBubbleTransformHandler<AgentArrayItem>::AddRequirementsForSpawnQuery(FLWComponentQuery& InQuery)
+void TMassClientBubbleTransformHandler<AgentArrayItem>::AddRequirementsForSpawnQuery(FMassEntityQuery& InQuery)
 {
-	InQuery.AddRequirement<FDataFragment_Transform>(ELWComponentAccess::ReadWrite);
+	InQuery.AddRequirement<FDataFragment_Transform>(EMassFragmentAccess::ReadWrite);
 }
 #endif // UE_REPLICATION_COMPILE_CLIENT_CODE
 
 #if UE_REPLICATION_COMPILE_CLIENT_CODE
 template<typename AgentArrayItem>
-void TMassClientBubbleTransformHandler<AgentArrayItem>::CacheComponentViewsForSpawnQuery(FLWComponentSystemExecutionContext& InExecContext)
+void TMassClientBubbleTransformHandler<AgentArrayItem>::CacheComponentViewsForSpawnQuery(FMassExecutionContext& InExecContext)
 {
 	TransformList = InExecContext.GetMutableComponentView<FDataFragment_Transform>();
 }
@@ -177,7 +177,7 @@ void TMassClientBubbleTransformHandler<AgentArrayItem>::SetSpawnedEntityData(con
 
 #if UE_REPLICATION_COMPILE_CLIENT_CODE
 template<typename AgentArrayItem>
-void TMassClientBubbleTransformHandler<AgentArrayItem>::SetModifiedEntityData(const FEntityView& EntityView, const FReplicatedAgentPositionYawData& ReplicatedPositionYawData)
+void TMassClientBubbleTransformHandler<AgentArrayItem>::SetModifiedEntityData(const FMassEntityView& EntityView, const FReplicatedAgentPositionYawData& ReplicatedPositionYawData)
 {
 	FDataFragment_Transform& TransformFragment = EntityView.GetComponentData<FDataFragment_Transform>();
 
@@ -199,8 +199,8 @@ void TMassClientBubbleTransformHandler<AgentArrayItem>::SetEntityData(FDataFragm
 class MASSREPLICATION_API FMassReplicationProcessorTransformHandlerBase
 {
 public:
-	static void AddRequirements(FLWComponentQuery& InQuery);
-	void CacheComponentViews(FLWComponentSystemExecutionContext& ExecContext);
+	static void AddRequirements(FMassEntityQuery& InQuery);
+	void CacheComponentViews(FMassExecutionContext& ExecContext);
 
 protected:
 	TArrayView<FDataFragment_Transform> TransformList;
diff --git a/Engine/Plugins/Runtime/MassGameplay/Source/MassReplication/Public/MassReplicationTypes.h b/Engine/Plugins/Runtime/MassGameplay/Source/MassReplication/Public/MassReplicationTypes.h
--- a/Engine/Plugins/Runtime/MassGameplay/Source/MassReplication/Public/MassReplicationTypes.h
+++ b/Engine/Plugins/Runtime/MassGameplay/Source/MassReplication/Public/MassReplicationTypes.h
@@ -121,7 +121,7 @@ struct MASSREPLICATION_API FMassBubbleInfoClassHandle : public FSimpleIndexedHan
  *  Light weight component data for the mass network id of a light weight entity
  */
 USTRUCT()
-struct MASSREPLICATION_API FMassNetworkIDFragment : public FLWComponentData
+struct MASSREPLICATION_API FMassNetworkIDFragment : public FMassFragment
 {
 	GENERATED_BODY()
 
@@ -151,7 +151,7 @@ struct FMassReplicatedAgentArrayData
  * Agent handle per client, these will be at TArray indices of the Client handles indicies (used as a free list array)
  */
 USTRUCT()
-struct MASSREPLICATION_API FMassReplicatedAgentFragment : public FLWComponentData
+struct MASSREPLICATION_API FMassReplicatedAgentFragment : public FMassFragment
 {
 	GENERATED_BODY()
 
@@ -195,7 +195,7 @@ struct FReplicationLODLogic : public FLODDefaultLogic
 };
 
 USTRUCT()
-struct MASSREPLICATION_API FMassReplicationViewerLODFragment : public FLWComponentData
+struct MASSREPLICATION_API FMassReplicationViewerLODFragment : public FMassFragment
 {
 	GENERATED_BODY()
 
@@ -221,12 +221,12 @@ public:
 protected:
 	virtual void ConfigureQueries() override;
 	virtual void Initialize(UObject& Owner) override;
-	virtual void Execute(UMassEntitySubsystem& EntitySubsystem, FLWComponentSystemExecutionContext& Context) override;
+	virtual void Execute(UMassEntitySubsystem& EntitySubsystem, FMassExecutionContext& Context) override;
 
 	UPROPERTY()
 	UMassReplicationManager* ReplicationManager = nullptr;
 
-	FLWComponentQuery EntityQuery;
+	FMassEntityQuery EntityQuery;
 };
 
 #if UE_REPLICATION_COMPILE_CLIENT_CODE
@@ -234,13 +234,13 @@ protected:
 struct FMassReplicationEntityInfo
 {
 	FMassReplicationEntityInfo() = default;
-	FMassReplicationEntityInfo(FLWEntity InEntity, int32 InReplicationID)
+	FMassReplicationEntityInfo(FMassEntityHandle InEntity, int32 InReplicationID)
 		: Entity(InEntity)
 		, ReplicationID(InReplicationID)
 	{}
 
 	/** If this is not IsSet() then the entity has been removed from the client simulation */
-	FLWEntity Entity; 
+	FMassEntityHandle Entity; 
 
 	/** This is stored between removes and adds, however this item in the UMassReplicationManager::EntityInfoMap will eventually get cleaned up if Entity.IsSet() == false for a fairly substantial length of time. */
 	int32 ReplicationID = INDEX_NONE;
diff --git a/Engine/Plugins/Runtime/MassGameplay/Source/MassRepresentation/Private/MassRepresentationProcessor.cpp b/Engine/Plugins/Runtime/MassGameplay/Source/MassRepresentation/Private/MassRepresentationProcessor.cpp
--- a/Engine/Plugins/Runtime/MassGameplay/Source/MassRepresentation/Private/MassRepresentationProcessor.cpp
+++ b/Engine/Plugins/Runtime/MassGameplay/Source/MassRepresentation/Private/MassRepresentationProcessor.cpp
@@ -27,11 +27,11 @@ UMassRepresentationProcessor::UMassRepresentationProcessor()
 
 void UMassRepresentationProcessor::ConfigureQueries()
 {
-	EntityQuery.AddRequirement<FDataFragment_Transform>(ELWComponentAccess::ReadOnly);
-	EntityQuery.AddRequirement<FMassRepresentationFragment>(ELWComponentAccess::ReadWrite);
-	EntityQuery.AddRequirement<FMassRepresentationLODFragment>(ELWComponentAccess::ReadOnly);
-	EntityQuery.AddRequirement<FDataFragment_Actor>(ELWComponentAccess::ReadWrite);
-	EntityQuery.AddChunkRequirement<FMassVisualizationChunkFragment>(ELWComponentAccess::ReadWrite);
+	EntityQuery.AddRequirement<FDataFragment_Transform>(EMassFragmentAccess::ReadOnly);
+	EntityQuery.AddRequirement<FMassRepresentationFragment>(EMassFragmentAccess::ReadWrite);
+	EntityQuery.AddRequirement<FMassRepresentationLODFragment>(EMassFragmentAccess::ReadOnly);
+	EntityQuery.AddRequirement<FDataFragment_Actor>(EMassFragmentAccess::ReadWrite);
+	EntityQuery.AddChunkRequirement<FMassVisualizationChunkFragment>(EMassFragmentAccess::ReadWrite);
 }
 
 void UMassRepresentationProcessor::Initialize(UObject& Owner)
@@ -57,7 +57,7 @@ void UMassRepresentationProcessor::Initialize(UObject& Owner)
 	}
 }
 
-void UMassRepresentationProcessor::UpdateRepresentation(FLWComponentSystemExecutionContext& Context)
+void UMassRepresentationProcessor::UpdateRepresentation(FMassExecutionContext& Context)
 {
 	check(RepresentationSubsystem);
 
@@ -69,7 +69,7 @@ void UMassRepresentationProcessor::UpdateRepresentation(FLWComponentSystemExecut
 	const int32 NumEntities = Context.GetEntitiesNum();
 	for (int32 EntityIdx = 0; EntityIdx < NumEntities; EntityIdx++)
 	{
-		const FLWEntity Entity = Context.GetEntity(EntityIdx);
+		const FMassEntityHandle Entity = Context.GetEntity(EntityIdx);
 		const FMassHandle MassAgent(Entity);
 		const FDataFragment_Transform& TransformFragment = TransformList[EntityIdx];
 		const FMassRepresentationLODFragment& RepresentationLOD = RepresentationLODList[EntityIdx];
@@ -207,12 +207,12 @@ void UMassRepresentationProcessor::UpdateRepresentation(FLWComponentSystemExecut
 	}
 }
 
-void UMassRepresentationProcessor::Execute(UMassEntitySubsystem& EntitySubsystem, FLWComponentSystemExecutionContext& Context)
+void UMassRepresentationProcessor::Execute(UMassEntitySubsystem& EntitySubsystem, FMassExecutionContext& Context)
 {
 	check(RepresentationSubsystem);
 
 	// Visualize entities
-	EntityQuery.ForEachEntityChunk(EntitySubsystem, Context, [this](FLWComponentSystemExecutionContext& Context)
+	EntityQuery.ForEachEntityChunk(EntitySubsystem, Context, [this](FMassExecutionContext& Context)
 	{
 		UpdateRepresentation(Context);
 	});
@@ -227,7 +227,7 @@ AActor* UMassRepresentationProcessor::GetOrSpawnActor(const FMassHandle MassAgen
 		FMassActorPostSpawnDelegate::CreateUObject(this, &UMassRepresentationProcessor::OnActorPostSpawn));
 }
 
-bool UMassRepresentationProcessor::ReleaseActorOrCancelSpawning(const FMassHandle MassAgent, FDataFragment_Actor& ActorInfo, const int16 TemplateActorIndex, FMassHandle_ActorSpawnRequest& SpawnRequestHandle, FLWComponentSystemExecutionContext& Context, bool bCancelSpawningOnly /*= false*/)
+bool UMassRepresentationProcessor::ReleaseActorOrCancelSpawning(const FMassHandle MassAgent, FDataFragment_Actor& ActorInfo, const int16 TemplateActorIndex, FMassHandle_ActorSpawnRequest& SpawnRequestHandle, FMassExecutionContext& Context, bool bCancelSpawningOnly /*= false*/)
 {
 	check(RepresentationSubsystem);
 	check(!ActorInfo.IsValid() || ActorInfo.IsOwnedByMass());
@@ -258,7 +258,7 @@ bool UMassRepresentationProcessor::ReleaseActorOrCancelSpawning(const FMassHandl
 	return false;
 }
 
-void UMassRepresentationProcessor::SetActorEnabled(const EActorEnabledType EnabledType, AActor& Actor, const int32 EntityIdx, FLWComponentSystemExecutionContext& Context)
+void UMassRepresentationProcessor::SetActorEnabled(const EActorEnabledType EnabledType, AActor& Actor, const int32 EntityIdx, FMassExecutionContext& Context)
 {
 	const bool bEnabled = EnabledType != EActorEnabledType::Disabled;
 	if (Actor.IsActorTickEnabled() != bEnabled)
@@ -275,7 +275,7 @@ void UMassRepresentationProcessor::SetActorEnabled(const EActorEnabledType Enabl
 	}
 }
 
-void UMassRepresentationProcessor::TeleportActor(const FTransform& Transform, AActor& Actor, FLWComponentSystemExecutionContext& Context)
+void UMassRepresentationProcessor::TeleportActor(const FTransform& Transform, AActor& Actor, FMassExecutionContext& Context)
 {
 	if (!Actor.GetTransform().Equals(Transform))
 	{
@@ -319,7 +319,7 @@ void UMassRepresentationProcessor::ReleaseAnyActorOrCancelAnySpawning(UMassRepre
 	}
 }
 
-void UMassRepresentationProcessor::UpdateVisualization(FLWComponentSystemExecutionContext& Context)
+void UMassRepresentationProcessor::UpdateVisualization(FMassExecutionContext& Context)
 {
 	FMassVisualizationChunkFragment& ChunkData = UpdateChunkVisibility(Context);
 	if (!ChunkData.ShouldUpdateVisualization())
@@ -336,14 +336,14 @@ void UMassRepresentationProcessor::UpdateVisualization(FLWComponentSystemExecuti
 	const int32 NumEntities = Context.GetEntitiesNum();
 	for (int32 EntityIdx = 0; EntityIdx < NumEntities; EntityIdx++)
 	{
-		const FLWEntity Entity = Context.GetEntity(EntityIdx);
+		const FMassEntityHandle Entity = Context.GetEntity(EntityIdx);
 		FMassRepresentationFragment& Representation = RepresentationList[EntityIdx];
 		const FMassRepresentationLODFragment& RepresentationLOD = RepresentationLODList[EntityIdx];
 		UpdateEntityVisibility(Entity, Representation, RepresentationLOD, ChunkData, Context);
 	}
 }
 
-FMassVisualizationChunkFragment& UMassRepresentationProcessor::UpdateChunkVisibility(FLWComponentSystemExecutionContext& Context) const
+FMassVisualizationChunkFragment& UMassRepresentationProcessor::UpdateChunkVisibility(FMassExecutionContext& Context) const
 {
 	// Setup chunk component data about visibility
 	FMassVisualizationChunkFragment& ChunkData = Context.GetMutableChunkComponent<FMassVisualizationChunkFragment>();
@@ -374,7 +374,7 @@ FMassVisualizationChunkFragment& UMassRepresentationProcessor::UpdateChunkVisibi
 	return ChunkData;
 }
 
-void UMassRepresentationProcessor::UpdateEntityVisibility(const FLWEntity Entity, const FMassRepresentationFragment& Representation, const FMassRepresentationLODFragment& RepresentationLOD, FMassVisualizationChunkFragment& ChunkData, FLWComponentSystemExecutionContext& Context)
+void UMassRepresentationProcessor::UpdateEntityVisibility(const FMassEntityHandle Entity, const FMassRepresentationFragment& Representation, const FMassRepresentationLODFragment& RepresentationLOD, FMassVisualizationChunkFragment& ChunkData, FMassExecutionContext& Context)
 {
 	// Move the visible entities together into same chunks so we can skip entire chunk when not visible as an optimization
 	const EMassVisibility Visibility = Representation.CurrentRepresentation != ERepresentationType::None ? EMassVisibility::CanBeSeen : 
@@ -393,7 +393,7 @@ void UMassRepresentationProcessor::OnActorPreSpawn(const FMassHandle_ActorSpawnR
 	check(RepresentationSubsystem);
 
 	const FMassActorSpawnRequest& MassActorSpawnRequest = SpawnRequest.Get<FMassActorSpawnRequest>();
-	const FEntityView View(*PipeEntitySubsystem, MassActorSpawnRequest.MassAgent.GetLWEntity());
+	const FMassEntityView View(*PipeEntitySubsystem, MassActorSpawnRequest.MassAgent.GetLWEntity());
 
 	if (FDataFragment_Actor* ActorInfo = View.GetComponentDataPtr<FDataFragment_Actor>())
 	{
@@ -465,15 +465,15 @@ void UMassRepresentationFragmentDestructor::Initialize(UObject& Owner)
 
 void UMassRepresentationFragmentDestructor::ConfigureQueries()
 {
-	EntityQuery.AddRequirement<FMassRepresentationFragment>(ELWComponentAccess::ReadOnly);
-	EntityQuery.AddRequirement<FDataFragment_Actor>(ELWComponentAccess::ReadWrite);
+	EntityQuery.AddRequirement<FMassRepresentationFragment>(EMassFragmentAccess::ReadOnly);
+	EntityQuery.AddRequirement<FDataFragment_Actor>(EMassFragmentAccess::ReadWrite);
 }
 
-void UMassRepresentationFragmentDestructor::Execute(UMassEntitySubsystem& EntitySubsystem, FLWComponentSystemExecutionContext& Context)
+void UMassRepresentationFragmentDestructor::Execute(UMassEntitySubsystem& EntitySubsystem, FMassExecutionContext& Context)
 {
 	check(RepresentationSubsystem);
 
-	EntityQuery.ForEachEntityChunk(EntitySubsystem, Context, [this](FLWComponentSystemExecutionContext& Context)
+	EntityQuery.ForEachEntityChunk(EntitySubsystem, Context, [this](FMassExecutionContext& Context)
 	{
 		const TArrayView<FMassRepresentationFragment> RepresentationList = Context.GetMutableComponentView<FMassRepresentationFragment>();
 		const TArrayView<FDataFragment_Actor> ActorList = Context.GetMutableComponentView<FDataFragment_Actor>();
@@ -484,7 +484,7 @@ void UMassRepresentationFragmentDestructor::Execute(UMassEntitySubsystem& Entity
 			FMassRepresentationFragment& Representation = RepresentationList[i];
 			FDataFragment_Actor& ActorInfo = ActorList[i];
 
-			const FLWEntity Entity = Context.GetEntity(i);
+			const FMassEntityHandle Entity = Context.GetEntity(i);
 			const FMassHandle MassAgent(Entity);
 
 			UMassRepresentationProcessor::ReleaseAnyActorOrCancelAnySpawning(*RepresentationSubsystem, MassAgent, ActorInfo, Representation);
diff --git a/Engine/Plugins/Runtime/MassGameplay/Source/MassRepresentation/Private/MassRepresentationSubsystem.cpp b/Engine/Plugins/Runtime/MassGameplay/Source/MassRepresentation/Private/MassRepresentationSubsystem.cpp
--- a/Engine/Plugins/Runtime/MassGameplay/Source/MassRepresentation/Private/MassRepresentationSubsystem.cpp
+++ b/Engine/Plugins/Runtime/MassGameplay/Source/MassRepresentation/Private/MassRepresentationSubsystem.cpp
@@ -319,8 +319,8 @@ void UMassRepresentationSubsystem::Initialize(FSubsystemCollectionBase& Collecti
 
 		UMassSimulationSubsystem* SimSystem = World->GetSubsystem<UMassSimulationSubsystem>();
 		check(SimSystem);
-		SimSystem->GetOnProcessingPhaseStarted(EPipeProcessingPhase::PrePhysics).AddUObject(this, &UMassRepresentationSubsystem::OnProcessingPhaseStarted, EPipeProcessingPhase::PrePhysics);
-		SimSystem->GetOnProcessingPhaseFinished(EPipeProcessingPhase::PostPhysics).AddUObject(this, &UMassRepresentationSubsystem::OnProcessingPhaseStarted, EPipeProcessingPhase::PostPhysics);
+		SimSystem->GetOnProcessingPhaseStarted(EMassProcessingPhase::PrePhysics).AddUObject(this, &UMassRepresentationSubsystem::OnProcessingPhaseStarted, EMassProcessingPhase::PrePhysics);
+		SimSystem->GetOnProcessingPhaseFinished(EMassProcessingPhase::PostPhysics).AddUObject(this, &UMassRepresentationSubsystem::OnProcessingPhaseStarted, EMassProcessingPhase::PostPhysics);
 
 		UMassAgentSubsystem* AgentSystem = World->GetSubsystem<UMassAgentSubsystem>();
 		check(AgentSystem);
@@ -338,8 +338,8 @@ void UMassRepresentationSubsystem::Deinitialize()
 	{
 		if (const UMassSimulationSubsystem* SimSystem = World->GetSubsystem<UMassSimulationSubsystem>())
 		{
-			SimSystem->GetOnProcessingPhaseStarted(EPipeProcessingPhase::PrePhysics).RemoveAll(this);
-			SimSystem->GetOnProcessingPhaseFinished(EPipeProcessingPhase::PostPhysics).RemoveAll(this);
+			SimSystem->GetOnProcessingPhaseStarted(EMassProcessingPhase::PrePhysics).RemoveAll(this);
+			SimSystem->GetOnProcessingPhaseFinished(EMassProcessingPhase::PostPhysics).RemoveAll(this);
 		}
 
 		if (UMassAgentSubsystem* AgentSystem = World->GetSubsystem<UMassAgentSubsystem>())
@@ -351,15 +351,15 @@ void UMassRepresentationSubsystem::Deinitialize()
 
 }
 
-void UMassRepresentationSubsystem::OnProcessingPhaseStarted(const float DeltaSeconds, const EPipeProcessingPhase Phase) const
+void UMassRepresentationSubsystem::OnProcessingPhaseStarted(const float DeltaSeconds, const EMassProcessingPhase Phase) const
 {
 	check(VisualizationComponent);
 	switch (Phase)
 	{
-		case EPipeProcessingPhase::PrePhysics:
+		case EMassProcessingPhase::PrePhysics:
 			VisualizationComponent->BeginVisualChanges();
 			break;
-		case EPipeProcessingPhase::PostPhysics:/* Currently this is the end of phases signal */
+		case EMassProcessingPhase::PostPhysics:/* Currently this is the end of phases signal */
 			VisualizationComponent->EndVisualChanges();
 			break;
 		default:
@@ -395,7 +395,7 @@ void UMassRepresentationSubsystem::OnMassAgentComponentEntityDetaching(const UMa
 	checkf(EntitySubsystem->IsEntityValid(MassAgent.GetLWEntity()), TEXT("Expecting a valid mass entity"));
 	if (EntitySubsystem->IsEntityValid(MassAgent.GetLWEntity()) && AgentComponent.IsNetSimulating())
 	{
-		const FEntityView EntityView(*EntitySubsystem, MassAgent.GetLWEntity());
+		const FMassEntityView EntityView(*EntitySubsystem, MassAgent.GetLWEntity());
 		if (FMassRepresentationFragment* Representation = EntityView.GetComponentDataPtr<FMassRepresentationFragment>())
 		{
 			// Force a reevaluate of the current representation
diff --git a/Engine/Plugins/Runtime/MassGameplay/Source/MassRepresentation/Private/MassUpdateISMProcessor.cpp b/Engine/Plugins/Runtime/MassGameplay/Source/MassRepresentation/Private/MassUpdateISMProcessor.cpp
--- a/Engine/Plugins/Runtime/MassGameplay/Source/MassRepresentation/Private/MassUpdateISMProcessor.cpp
+++ b/Engine/Plugins/Runtime/MassGameplay/Source/MassRepresentation/Private/MassUpdateISMProcessor.cpp
@@ -27,27 +27,27 @@ void UMassUpdateISMProcessor::ConfigureQueries()
 {
 	for (const FInstancedStruct& TagFilter : TagFilters)
 	{
-		FLWComponentQuery& EntityQuery = EntityQueries.AddDefaulted_GetRef();
+		FMassEntityQuery& EntityQuery = EntityQueries.AddDefaulted_GetRef();
 		if (TagFilter.GetScriptStruct())
 		{
-			EntityQuery.AddTagRequirement(*TagFilter.GetScriptStruct(), ELWComponentPresence::All);
+			EntityQuery.AddTagRequirement(*TagFilter.GetScriptStruct(), EMassFragmentPresence::All);
 		}
-		EntityQuery.AddRequirement<FDataFragment_Transform>(ELWComponentAccess::ReadOnly);
-		EntityQuery.AddRequirement<FMassRepresentationFragment>(ELWComponentAccess::ReadWrite);
-		EntityQuery.AddRequirement<FMassRepresentationLODFragment>(ELWComponentAccess::ReadWrite);
-		EntityQuery.AddChunkRequirement<FMassVisualizationChunkFragment>(ELWComponentAccess::ReadWrite);
+		EntityQuery.AddRequirement<FDataFragment_Transform>(EMassFragmentAccess::ReadOnly);
+		EntityQuery.AddRequirement<FMassRepresentationFragment>(EMassFragmentAccess::ReadWrite);
+		EntityQuery.AddRequirement<FMassRepresentationLODFragment>(EMassFragmentAccess::ReadWrite);
+		EntityQuery.AddChunkRequirement<FMassVisualizationChunkFragment>(EMassFragmentAccess::ReadWrite);
 		EntityQuery.SetChunkFilter(&FMassVisualizationChunkFragment::AreAnyEntitiesVisibleInChunk);
 	}
 }
 
-void UMassUpdateISMProcessor::Execute(UMassEntitySubsystem& EntitySubsystem, FLWComponentSystemExecutionContext& Context)
+void UMassUpdateISMProcessor::Execute(UMassEntitySubsystem& EntitySubsystem, FMassExecutionContext& Context)
 {
 	check(RepresentationSubsystem);
 	FMassInstancedStaticMeshInfoArrayView ISMInfo = RepresentationSubsystem->GetMutableInstancedStaticMeshInfos();
 
-	for (FLWComponentQuery& EntityQuery : EntityQueries)
+	for (FMassEntityQuery& EntityQuery : EntityQueries)
 	{
-		EntityQuery.ForEachEntityChunk(EntitySubsystem, Context, [&ISMInfo](FLWComponentSystemExecutionContext& Context)
+		EntityQuery.ForEachEntityChunk(EntitySubsystem, Context, [&ISMInfo](FMassExecutionContext& Context)
 		{
 			const TConstArrayView<FDataFragment_Transform> TransformList = Context.GetComponentView<FDataFragment_Transform>();
 			const TArrayView<FMassRepresentationFragment> RepresentationList = Context.GetMutableComponentView<FMassRepresentationFragment>();
diff --git a/Engine/Plugins/Runtime/MassGameplay/Source/MassRepresentation/Private/MassVisualizationLODProcessor.cpp b/Engine/Plugins/Runtime/MassGameplay/Source/MassRepresentation/Private/MassVisualizationLODProcessor.cpp
--- a/Engine/Plugins/Runtime/MassGameplay/Source/MassRepresentation/Private/MassVisualizationLODProcessor.cpp
+++ b/Engine/Plugins/Runtime/MassGameplay/Source/MassRepresentation/Private/MassVisualizationLODProcessor.cpp
@@ -33,16 +33,16 @@ UMassVisualizationLODProcessor::UMassVisualizationLODProcessor()
 
 void UMassVisualizationLODProcessor::ConfigureQueries()
 {
-	CloseEntityQuery.AddRequirement<FMassLODInfoFragment>(ELWComponentAccess::ReadOnly);
-	CloseEntityQuery.AddRequirement<FMassRepresentationLODFragment>(ELWComponentAccess::ReadWrite);
-	CloseEntityQuery.AddRequirement<FDataFragment_Transform>(ELWComponentAccess::ReadOnly);
+	CloseEntityQuery.AddRequirement<FMassLODInfoFragment>(EMassFragmentAccess::ReadOnly);
+	CloseEntityQuery.AddRequirement<FMassRepresentationLODFragment>(EMassFragmentAccess::ReadWrite);
+	CloseEntityQuery.AddRequirement<FDataFragment_Transform>(EMassFragmentAccess::ReadOnly);
 
 	FarEntityQuery = CloseEntityQuery;
 
-	CloseEntityQuery.AddTagRequirement<FMassVisibilityCulledByDistanceTag>(ELWComponentPresence::None);
+	CloseEntityQuery.AddTagRequirement<FMassVisibilityCulledByDistanceTag>(EMassFragmentPresence::None);
 	
-	FarEntityQuery.AddTagRequirement<FMassVisibilityCulledByDistanceTag>(ELWComponentPresence::All);
-	FarEntityQuery.AddChunkRequirement<FMassVisualizationChunkFragment>(ELWComponentAccess::ReadOnly);
+	FarEntityQuery.AddTagRequirement<FMassVisibilityCulledByDistanceTag>(EMassFragmentPresence::All);
+	FarEntityQuery.AddChunkRequirement<FMassVisualizationChunkFragment>(EMassFragmentAccess::ReadOnly);
 }
 
 void UMassVisualizationLODProcessor::Initialize(UObject& Owner)
@@ -61,7 +61,7 @@ void UMassVisualizationLODProcessor::PrepareExecution()
 	LODCalculator.PrepareExecution(Viewers);
 }
 
-void UMassVisualizationLODProcessor::Execute(UMassEntitySubsystem& EntitySubsystem, FLWComponentSystemExecutionContext& Context)
+void UMassVisualizationLODProcessor::Execute(UMassEntitySubsystem& EntitySubsystem, FMassExecutionContext& Context)
 {
 	ExecuteInternal(EntitySubsystem, Context);
 }
\ No newline at end of file
diff --git a/Engine/Plugins/Runtime/MassGameplay/Source/MassRepresentation/Private/MassVisualizationTrait.cpp b/Engine/Plugins/Runtime/MassGameplay/Source/MassRepresentation/Private/MassVisualizationTrait.cpp
--- a/Engine/Plugins/Runtime/MassGameplay/Source/MassRepresentation/Private/MassVisualizationTrait.cpp
+++ b/Engine/Plugins/Runtime/MassGameplay/Source/MassRepresentation/Private/MassVisualizationTrait.cpp
@@ -34,7 +34,7 @@ void UMassVisualizationTrait::BuildTemplate(FMassEntityTemplateBuildContext& Bui
 		RepresentationFragment.HighResTemplateActorIndex = HighResTemplateActor.Get() ? RepresentationSubsystem->FindOrAddTemplateActor(HighResTemplateActor.Get()) : INDEX_NONE;
 		RepresentationFragment.LowResTemplateActorIndex = LowResTemplateActor.Get() ? RepresentationSubsystem->FindOrAddTemplateActor(LowResTemplateActor.Get()) : INDEX_NONE;
 
-		BuildContext.AddDeinitializer(*(RepresentationFragmentDeinitializerClass->GetDefaultObject<UPipeProcessor>()));
+		BuildContext.AddDeinitializer(*(RepresentationFragmentDeinitializerClass->GetDefaultObject<UMassProcessor>()));
 
 		BuildContext.AddFragment<FMassRepresentationLODFragment>();
 
diff --git a/Engine/Plugins/Runtime/MassGameplay/Source/MassRepresentation/Public/MassRepresentationFragments.h b/Engine/Plugins/Runtime/MassGameplay/Source/MassRepresentation/Public/MassRepresentationFragments.h
--- a/Engine/Plugins/Runtime/MassGameplay/Source/MassRepresentation/Public/MassRepresentationFragments.h
+++ b/Engine/Plugins/Runtime/MassGameplay/Source/MassRepresentation/Public/MassRepresentationFragments.h
@@ -11,7 +11,7 @@
 #include "MassRepresentationFragments.generated.h"
 
 USTRUCT()
-struct MASSREPRESENTATION_API FMassRepresentationLODFragment : public FLWComponentData
+struct MASSREPRESENTATION_API FMassRepresentationLODFragment : public FMassFragment
 {
 	GENERATED_BODY()
 
@@ -33,7 +33,7 @@ struct MASSREPRESENTATION_API FMassRepresentationLODFragment : public FLWCompone
 };
 
 USTRUCT()
-struct MASSREPRESENTATION_API FMassRepresentationFragment : public FLWComponentData
+struct MASSREPRESENTATION_API FMassRepresentationFragment : public FMassFragment
 {
 	GENERATED_BODY()
 
diff --git a/Engine/Plugins/Runtime/MassGameplay/Source/MassRepresentation/Public/MassRepresentationProcessor.h b/Engine/Plugins/Runtime/MassGameplay/Source/MassRepresentation/Public/MassRepresentationProcessor.h
--- a/Engine/Plugins/Runtime/MassGameplay/Source/MassRepresentation/Public/MassRepresentationProcessor.h
+++ b/Engine/Plugins/Runtime/MassGameplay/Source/MassRepresentation/Public/MassRepresentationProcessor.h
@@ -19,7 +19,7 @@ enum EActorEnabledType
 };
 
 UCLASS()
-class MASSREPRESENTATION_API UMassRepresentationProcessor : public UPipeProcessor
+class MASSREPRESENTATION_API UMassRepresentationProcessor : public UMassProcessor
 {
 	GENERATED_BODY()
 
@@ -28,7 +28,7 @@ public:
 
 protected:
 
-	/** Configure the owned FLWComponentQuery instances to express processor's requirements */
+	/** Configure the owned FMassEntityQuery instances to express processor's requirements */
 	virtual void ConfigureQueries() override;
 
 	/**
@@ -42,7 +42,7 @@ protected:
 	 * @param EntitySubsystem is the system to execute the lambdas on each entity chunk
 	 * @param Context is the execution context to be passed when executing the lambdas
 	 */
-	virtual void Execute(UMassEntitySubsystem& EntitySubsystem, FLWComponentSystemExecutionContext& Context) override;
+	virtual void Execute(UMassEntitySubsystem& EntitySubsystem, FMassExecutionContext& Context) override;
 
 	/** 
 	 * Returns an actor of the template type and setup fragments values from it
@@ -66,7 +66,7 @@ protected:
 	 * @param bCancelSpawningOnly tell to only cancel the existing spawning request and to not release the associated actor it any.
 	 * @return if the actor was release or the spawning was canceled.
 	 */
-	virtual bool ReleaseActorOrCancelSpawning(const FMassHandle MassAgent, FDataFragment_Actor& ActorInfo, const int16 TemplateActorIndex, FMassHandle_ActorSpawnRequest& SpawnRequestHandle, FLWComponentSystemExecutionContext& Context, const bool bCancelSpawningOnly = false);
+	virtual bool ReleaseActorOrCancelSpawning(const FMassHandle MassAgent, FDataFragment_Actor& ActorInfo, const int16 TemplateActorIndex, FMassHandle_ActorSpawnRequest& SpawnRequestHandle, FMassExecutionContext& Context, const bool bCancelSpawningOnly = false);
 
 	/** 
 	 * Enable/disable a spawned actor
@@ -75,7 +75,7 @@ protected:
 	 * @param EntityIdx is the entity index currently processing
 	 * @param Context is the current light weight component execution context 
 	 */
-	virtual void SetActorEnabled(const EActorEnabledType EnabledType, AActor& Actor, const int32 EntityIdx, FLWComponentSystemExecutionContext& Context);
+	virtual void SetActorEnabled(const EActorEnabledType EnabledType, AActor& Actor, const int32 EntityIdx, FMassExecutionContext& Context);
 
 	/**
 	 * Teleports the actor at the specified transform by preserving its velocity and without collision.
@@ -84,7 +84,7 @@ protected:
 	 * @param Actor is the actual actor to teleport
 	 * @param Context is the current light weight component execution context
 	 */
-	virtual void TeleportActor(const FTransform& Transform, AActor& Actor, FLWComponentSystemExecutionContext& Context);
+	virtual void TeleportActor(const FTransform& Transform, AActor& Actor, FMassExecutionContext& Context);
 
 	/**
 	 * Method that will be bound to a delegate called before the spawning of an actor to let the requester prepare it
@@ -113,7 +113,7 @@ protected:
 	 * @param Context of the execution from the entity sub system
 	 * @return The visualization chunk fragment
 	 */
-	FMassVisualizationChunkFragment& UpdateChunkVisibility(FLWComponentSystemExecutionContext& Context) const;
+	FMassVisualizationChunkFragment& UpdateChunkVisibility(FMassExecutionContext& Context) const;
 
 	/**
 	 * Updates entity visibility tag for later chunk logic optimization
@@ -123,7 +123,7 @@ protected:
 	 * @param ChunkData is the visualization chunk fragment
      * @param Context of the execution from the entity sub system
 	 */
-	static void UpdateEntityVisibility(const FLWEntity Entity, const FMassRepresentationFragment& Representation, const FMassRepresentationLODFragment& RepresentationLOD, FMassVisualizationChunkFragment& ChunkData, FLWComponentSystemExecutionContext& Context);
+	static void UpdateEntityVisibility(const FMassEntityHandle Entity, const FMassRepresentationFragment& Representation, const FMassRepresentationLODFragment& RepresentationLOD, FMassVisualizationChunkFragment& ChunkData, FMassExecutionContext& Context);
 
 public:
 	/**
@@ -158,13 +158,13 @@ protected:
 	 * Update representation type for each entity, must be called within a ForEachEntityChunk
 	 * @param Context of the execution from the entity sub system
 	 */
-	void UpdateRepresentation(FLWComponentSystemExecutionContext& Context);
+	void UpdateRepresentation(FMassExecutionContext& Context);
 
 	/** 
 	 * Update representation and visibility for each entity, must be called within a ForEachEntityChunk
 	 * @param Context of the execution from the entity sub system
 	 */
-	void UpdateVisualization(FLWComponentSystemExecutionContext& Context);
+	void UpdateVisualization(FMassExecutionContext& Context);
 
 	/** What should be the representation of this entity for each specificLOD */
 	UPROPERTY(EditAnywhere, Category = "Mass|Representation", config)
@@ -192,7 +192,7 @@ protected:
 	/** Default representation when unable to spawn an actor */
 	ERepresentationType DefaultRepresentationType = ERepresentationType::None;
 
-	FLWComponentQuery EntityQuery;
+	FMassEntityQuery EntityQuery;
 
 	/** At what rate should the not visible entity be updated in seconds */
 	UPROPERTY(EditAnywhere, Category = "Mass|Visualization", config)
@@ -210,11 +210,11 @@ public:
 protected:
 	virtual void Initialize(UObject& Owner) override;
 	virtual void ConfigureQueries() override;
-	virtual void Execute(UMassEntitySubsystem& EntitySubsystem, FLWComponentSystemExecutionContext& Context) override;
+	virtual void Execute(UMassEntitySubsystem& EntitySubsystem, FMassExecutionContext& Context) override;
 
 	/** A cache pointer to the representation subsystem */
 	UPROPERTY(Transient)
 	UMassRepresentationSubsystem* RepresentationSubsystem;
 
-	FLWComponentQuery EntityQuery;
+	FMassEntityQuery EntityQuery;
 };
\ No newline at end of file
diff --git a/Engine/Plugins/Runtime/MassGameplay/Source/MassRepresentation/Public/MassRepresentationSubsystem.h b/Engine/Plugins/Runtime/MassGameplay/Source/MassRepresentation/Public/MassRepresentationSubsystem.h
--- a/Engine/Plugins/Runtime/MassGameplay/Source/MassRepresentation/Public/MassRepresentationSubsystem.h
+++ b/Engine/Plugins/Runtime/MassGameplay/Source/MassRepresentation/Public/MassRepresentationSubsystem.h
@@ -19,7 +19,7 @@ struct FMassHandle_ActorSpawnRequest;
 class UMassActorSpawnerSubsystem;
 class UMassAgentComponent;
 class UMassEntitySubsystem;
-enum class EPipeProcessingPhase : uint8;
+enum class EMassProcessingPhase : uint8;
 class UWorldPartitionSubsystem;
 
 /**
@@ -123,7 +123,7 @@ protected:
 	// USubsystem END
 
 	/** Needed for batching the update of static mesh transform */
-	void OnProcessingPhaseStarted(const float DeltaSeconds, const EPipeProcessingPhase Phase) const;
+	void OnProcessingPhaseStarted(const float DeltaSeconds, const EMassProcessingPhase Phase) const;
 
 	void OnMassAgentComponentEntityAssociated(const UMassAgentComponent& AgentComponent);
 	void OnMassAgentComponentEntityDetaching(const UMassAgentComponent& AgentComponent);
diff --git a/Engine/Plugins/Runtime/MassGameplay/Source/MassRepresentation/Public/MassRepresentationUtils.h b/Engine/Plugins/Runtime/MassGameplay/Source/MassRepresentation/Public/MassRepresentationUtils.h
--- a/Engine/Plugins/Runtime/MassGameplay/Source/MassRepresentation/Public/MassRepresentationUtils.h
+++ b/Engine/Plugins/Runtime/MassGameplay/Source/MassRepresentation/Public/MassRepresentationUtils.h
@@ -10,7 +10,7 @@
 namespace UE::MassRepresentation
 {
 
-inline EMassVisibility GetVisibilityFromArchetype(const FLWComponentSystemExecutionContext& Context)
+inline EMassVisibility GetVisibilityFromArchetype(const FMassExecutionContext& Context)
 {
 	if (Context.DoesArchetypeHaveTag<FMassVisibilityCanBeSeenTag>())
 	{
@@ -44,7 +44,7 @@ inline const UScriptStruct* GetTagFromVisibility(EMassVisibility Visibility)
 	}
 }
 
-inline bool IsVisibilityTagSet(const FLWComponentSystemExecutionContext& Context, EMassVisibility Visibility)
+inline bool IsVisibilityTagSet(const FMassExecutionContext& Context, EMassVisibility Visibility)
 {
 	switch (Visibility)
 	{
diff --git a/Engine/Plugins/Runtime/MassGameplay/Source/MassRepresentation/Public/MassUpdateISMProcessor.h b/Engine/Plugins/Runtime/MassGameplay/Source/MassRepresentation/Public/MassUpdateISMProcessor.h
--- a/Engine/Plugins/Runtime/MassGameplay/Source/MassRepresentation/Public/MassUpdateISMProcessor.h
+++ b/Engine/Plugins/Runtime/MassGameplay/Source/MassRepresentation/Public/MassUpdateISMProcessor.h
@@ -11,7 +11,7 @@ class UMassRepresentationSubsystem;
 struct FMassInstancedStaticMeshInfo;
 
 UCLASS()
-class MASSREPRESENTATION_API UMassUpdateISMProcessor : public UPipeProcessor
+class MASSREPRESENTATION_API UMassUpdateISMProcessor : public UMassProcessor
 {
 	GENERATED_BODY()
 
@@ -26,14 +26,14 @@ protected:
 	 * @param Owner of the Processor */
 	virtual void Initialize(UObject& Owner) override;
 
-	/** Configure the owned FLWComponentQuery instances to express processor's requirements */
+	/** Configure the owned FMassEntityQuery instances to express processor's requirements */
 	virtual void ConfigureQueries() override;
 
 	/**
 	 * Execution method for this processor
 	 * @param EntitySubsystem is the system to execute the lambdas on each entity chunk
 	 * @param Context is the execution context to be passed when executing the lambdas */
-	virtual void Execute(UMassEntitySubsystem& EntitySubsystem, FLWComponentSystemExecutionContext& Context) override;
+	virtual void Execute(UMassEntitySubsystem& EntitySubsystem, FMassExecutionContext& Context) override;
 
 	/** Component Tag that will be used to filter out entities to update ISM */
 	UPROPERTY(EditAnywhere, Category = "Mass|LOD", config, meta = (BaseStruct = "ComponentTag"))
@@ -43,5 +43,5 @@ protected:
 	UPROPERTY(Transient)
 	UMassRepresentationSubsystem* RepresentationSubsystem;
 
-	TArray<FLWComponentQuery> EntityQueries;
+	TArray<FMassEntityQuery> EntityQueries;
 };
\ No newline at end of file
diff --git a/Engine/Plugins/Runtime/MassGameplay/Source/MassRepresentation/Public/MassVisualizationLODProcessor.h b/Engine/Plugins/Runtime/MassGameplay/Source/MassRepresentation/Public/MassVisualizationLODProcessor.h
--- a/Engine/Plugins/Runtime/MassGameplay/Source/MassRepresentation/Public/MassVisualizationLODProcessor.h
+++ b/Engine/Plugins/Runtime/MassGameplay/Source/MassRepresentation/Public/MassVisualizationLODProcessor.h
@@ -18,7 +18,7 @@ public:
 
 protected:
 
-	/** Configure the owned FLWComponentQuery instances to express processor's requirements */
+	/** Configure the owned FMassEntityQuery instances to express processor's requirements */
 	virtual void ConfigureQueries() override;
 
 	/**
@@ -32,12 +32,12 @@ protected:
 	 * @param EntitySubsystem is the system to execute the lambdas on each entity chunk
 	 * @param Context is the execution context to be passed when executing the lambdas
 	 */
-	virtual void Execute(UMassEntitySubsystem& EntitySubsystem, FLWComponentSystemExecutionContext& Context) override;
+	virtual void Execute(UMassEntitySubsystem& EntitySubsystem, FMassExecutionContext& Context) override;
 
 	void PrepareExecution();
 
 	template <typename TMassViewerLODInfoFragment = FMassLODInfoFragment>
-	void ExecuteInternal(UMassEntitySubsystem& EntitySubsystem, FLWComponentSystemExecutionContext& Context);
+	void ExecuteInternal(UMassEntitySubsystem& EntitySubsystem, FMassExecutionContext& Context);
 
 	
 	/**
@@ -61,8 +61,8 @@ protected:
 
 	TMassLODCalculator<FMassRepresentationLODLogic> LODCalculator;
 
-	FLWComponentQuery CloseEntityQuery;
-	FLWComponentQuery FarEntityQuery;
+	FMassEntityQuery CloseEntityQuery;
+	FMassEntityQuery FarEntityQuery;
 
 	bool bForceOFFLOD = false;
 };
@@ -73,11 +73,11 @@ namespace UE::MassRepresentation
 } // UE::MassRepresentation
 
 template <typename TMassViewerLODInfoFragment>
-void UMassVisualizationLODProcessor::ExecuteInternal(UMassEntitySubsystem& EntitySubsystem, FLWComponentSystemExecutionContext& Context)
+void UMassVisualizationLODProcessor::ExecuteInternal(UMassEntitySubsystem& EntitySubsystem, FMassExecutionContext& Context)
 {
 	if (bForceOFFLOD)
 	{
-		CloseEntityQuery.ForEachEntityChunk(EntitySubsystem, Context, [this](FLWComponentSystemExecutionContext& Context)
+		CloseEntityQuery.ForEachEntityChunk(EntitySubsystem, Context, [this](FMassExecutionContext& Context)
 		{
 			TArrayView<FMassRepresentationLODFragment> RepresentationLODList = Context.GetMutableComponentView<FMassRepresentationLODFragment>();
 			LODCalculator.ForceOffLOD(Context, RepresentationLODList);
@@ -96,7 +96,7 @@ void UMassVisualizationLODProcessor::ExecuteInternal(UMassEntitySubsystem& Entit
 
 		// @todo this is the only block in this processor actually using FMassRepresentationFragment. Consider refactoring it into a separate Processor
 		// the "debug" section below is non-production, should be separated as well
-		auto CalculateLOD = [this](FLWComponentSystemExecutionContext& Context)
+		auto CalculateLOD = [this](FMassExecutionContext& Context)
 		{
 			TArrayView<FMassRepresentationLODFragment> RepresentationLODList = Context.GetMutableComponentView<FMassRepresentationLODFragment>();
 			TConstArrayView<TMassViewerLODInfoFragment> ViewerInfoList = Context.GetComponentView<TMassViewerLODInfoFragment>();
@@ -107,7 +107,7 @@ void UMassVisualizationLODProcessor::ExecuteInternal(UMassEntitySubsystem& Entit
 		// @todo this is the only block in this processor actually using FMassRepresentationFragment. Consider refactoring it into a separate Processor
 		// the "debug" section below is non-production, should be separated as well
 		FarEntityQuery.SetChunkFilter(&FMassVisualizationChunkFragment::ShouldUpdateVisualizationForChunk);
-		FarEntityQuery.ForEachEntityChunk(EntitySubsystem, Context, [this, &CalculateLOD](FLWComponentSystemExecutionContext& Context)
+		FarEntityQuery.ForEachEntityChunk(EntitySubsystem, Context, [this, &CalculateLOD](FMassExecutionContext& Context)
 			{
 				CalculateLOD(Context);
 			});
@@ -118,7 +118,7 @@ void UMassVisualizationLODProcessor::ExecuteInternal(UMassEntitySubsystem& Entit
 		TRACE_CPUPROFILER_EVENT_SCOPE(AdjustDistanceAndLODFromCount)
 		if (LODCalculator.AdjustDistancesFromCount())
 		{
-			CloseEntityQuery.ForEachEntityChunk(EntitySubsystem, Context, [this](FLWComponentSystemExecutionContext& Context)
+			CloseEntityQuery.ForEachEntityChunk(EntitySubsystem, Context, [this](FMassExecutionContext& Context)
 			{
 				TConstArrayView<TMassViewerLODInfoFragment> ViewerInfoList = Context.GetComponentView<TMassViewerLODInfoFragment>();
 				TArrayView<FMassRepresentationLODFragment> RepresentationLODList = Context.GetMutableComponentView<FMassRepresentationLODFragment>();
@@ -133,7 +133,7 @@ void UMassVisualizationLODProcessor::ExecuteInternal(UMassEntitySubsystem& Entit
 	{
 		TRACE_CPUPROFILER_EVENT_SCOPE(DebugDisplayLOD)
 
-		auto DebugDisplayLOD = [this](FLWComponentSystemExecutionContext& Context)
+		auto DebugDisplayLOD = [this](FMassExecutionContext& Context)
 		{
 			TConstArrayView<FMassRepresentationLODFragment> RepresentationLODList = Context.GetComponentView<FMassRepresentationLODFragment>();
 			TConstArrayView<FDataFragment_Transform> TransformList = Context.GetComponentView<FDataFragment_Transform>();
diff --git a/Engine/Plugins/Runtime/MassGameplay/Source/MassRepresentation/Public/MassVisualizationTrait.h b/Engine/Plugins/Runtime/MassGameplay/Source/MassRepresentation/Public/MassVisualizationTrait.h
--- a/Engine/Plugins/Runtime/MassGameplay/Source/MassRepresentation/Public/MassVisualizationTrait.h
+++ b/Engine/Plugins/Runtime/MassGameplay/Source/MassRepresentation/Public/MassVisualizationTrait.h
@@ -8,7 +8,7 @@
 #include "MassVisualizationTrait.generated.h"
 
 class UMassRepresentationSubsystem;
-class UPipeProcessor;
+class UMassProcessor;
 
 UCLASS(BlueprintType, EditInlineNew, CollapseCategories, meta=(DisplayName="Visualization"))
 class MASSREPRESENTATION_API UMassVisualizationTrait : public UMassEntityTraitBase
@@ -37,6 +37,6 @@ protected:
 	TSubclassOf<UMassRepresentationSubsystem> RepresentationSubsystemClass;
 
 	/** Allow subclasses to override the representation fragment deinitializer */
-	TSubclassOf<UPipeProcessor> RepresentationFragmentDeinitializerClass;
+	TSubclassOf<UMassProcessor> RepresentationFragmentDeinitializerClass;
 
 };
diff --git a/Engine/Plugins/Runtime/MassGameplay/Source/MassSignals/Private/MassSignalProcessorBase.cpp b/Engine/Plugins/Runtime/MassGameplay/Source/MassSignals/Private/MassSignalProcessorBase.cpp
--- a/Engine/Plugins/Runtime/MassGameplay/Source/MassSignals/Private/MassSignalProcessorBase.cpp
+++ b/Engine/Plugins/Runtime/MassGameplay/Source/MassSignals/Private/MassSignalProcessorBase.cpp
@@ -31,7 +31,7 @@ void UMassSignalProcessorBase::SubscribeToSignal(const FName SignalName)
 	SignalSubsystem->GetSignalDelegateByName(SignalName).AddUObject(this, &UMassSignalProcessorBase::OnSignalReceived);
 }
 
-void UMassSignalProcessorBase::Execute(UMassEntitySubsystem& EntitySubsystem, FLWComponentSystemExecutionContext& Context)
+void UMassSignalProcessorBase::Execute(UMassEntitySubsystem& EntitySubsystem, FMassExecutionContext& Context)
 {
 	QUICK_SCOPE_CYCLE_COUNTER(SignalEntities);
 
@@ -40,7 +40,7 @@ void UMassSignalProcessorBase::Execute(UMassEntitySubsystem& EntitySubsystem, FL
 	CurrentFrameBufferIndex = (CurrentFrameBufferIndex + 1) % 2;
 	FFrameReceivedSignals& ProcessingFrameBuffer = FrameReceivedSignals[ProcessingFrameBufferIndex];
 	TArray<FEntitySignalRange>& ReceivedSignalRanges = ProcessingFrameBuffer.ReceivedSignalRanges;
-	TArray<FLWEntity>& SignaledEntities = ProcessingFrameBuffer.SignaledEntities;
+	TArray<FMassEntityHandle>& SignaledEntities = ProcessingFrameBuffer.SignaledEntities;
 
 	if (ReceivedSignalRanges.IsEmpty())
 	{
@@ -61,8 +61,8 @@ void UMassSignalProcessorBase::Execute(UMassEntitySubsystem& EntitySubsystem, FL
 			}
 
 			FArchetypeHandle Archetype;
-			TArray<FLWEntity> Entities;
-			TSet<FLWEntity> AddedEntities;
+			TArray<FMassEntityHandle> Entities;
+			TSet<FMassEntityHandle> AddedEntities;
 		};
 		TArray<FEntitySet> EntitySets;
 
@@ -94,9 +94,9 @@ void UMassSignalProcessorBase::Execute(UMassEntitySubsystem& EntitySubsystem, FL
 				}
 
 				// Entities for this signal
-				TArrayView<FLWEntity> Entities(&SignaledEntities[Range.Begin], Range.End - Range.Begin);
+				TArrayView<FMassEntityHandle> Entities(&SignaledEntities[Range.Begin], Range.End - Range.Begin);
 				FEntitySet* PrevSet = &EntitySets[0];
-				for (const FLWEntity Entity : Entities)
+				for (const FMassEntityHandle Entity : Entities)
 				{
 					// Add to set of supported archetypes. Dont process if we don't care about the type.
 					FArchetypeHandle Archetype = EntitySubsystem.GetArchetypeForEntity(Entity);
@@ -136,7 +136,7 @@ void UMassSignalProcessorBase::Execute(UMassEntitySubsystem& EntitySubsystem, FL
 	SignaledEntities.Reset();
 }
 
-void UMassSignalProcessorBase::OnSignalReceived(FName SignalName, TConstArrayView<FLWEntity> Entities)
+void UMassSignalProcessorBase::OnSignalReceived(FName SignalName, TConstArrayView<FMassEntityHandle> Entities)
 {
 	FFrameReceivedSignals& CurrentFrameBuffer = FrameReceivedSignals[CurrentFrameBufferIndex];
 
diff --git a/Engine/Plugins/Runtime/MassGameplay/Source/MassSignals/Private/MassSignalSubsystem.cpp b/Engine/Plugins/Runtime/MassGameplay/Source/MassSignals/Private/MassSignalSubsystem.cpp
--- a/Engine/Plugins/Runtime/MassGameplay/Source/MassSignals/Private/MassSignalSubsystem.cpp
+++ b/Engine/Plugins/Runtime/MassGameplay/Source/MassSignals/Private/MassSignalSubsystem.cpp
@@ -32,13 +32,13 @@ TStatId UMassSignalSubsystem::GetStatId() const
 	RETURN_QUICK_DECLARE_CYCLE_STAT(UMassSignalSubsystem, STATGROUP_Tickables);
 }
 
-void UMassSignalSubsystem::SignalEntity(FName SignalName, const FLWEntity Entity)
+void UMassSignalSubsystem::SignalEntity(FName SignalName, const FMassEntityHandle Entity)
 {
 	checkf(Entity.IsSet(), TEXT("Expecting a valid entity to signal"));
 	SignalEntities(SignalName, MakeArrayView(&Entity, 1));
 }
 
-void UMassSignalSubsystem::SignalEntities(FName SignalName, TConstArrayView<FLWEntity> Entities)
+void UMassSignalSubsystem::SignalEntities(FName SignalName, TConstArrayView<FMassEntityHandle> Entities)
 {
 	checkf(Entities.Num() > 0, TEXT("Expecting entities to signal"));
 	const UE::MassSignal::FSignalDelegate& SignalDelegate = GetSignalDelegateByName(SignalName);
@@ -52,13 +52,13 @@ void UMassSignalSubsystem::SignalEntities(FName SignalName, TConstArrayView<FLWE
 	UE_CVLOG(Entities.Num() > 1, this, LogMassSignals, Log, TEXT("Raising signal [%s] to %d entities"), *SignalName.ToString(), Entities.Num());
 }
 
-void UMassSignalSubsystem::DelaySignalEntity(FName SignalName, const FLWEntity Entity, const float DelayInSeconds)
+void UMassSignalSubsystem::DelaySignalEntity(FName SignalName, const FMassEntityHandle Entity, const float DelayInSeconds)
 {
 	checkf(Entity.IsSet(), TEXT("Expecting a valid entity to signal"));
 	DelaySignalEntities(SignalName, MakeArrayView(&Entity, 1), DelayInSeconds);
 }
 
-void UMassSignalSubsystem::DelaySignalEntities(FName SignalName, TConstArrayView<FLWEntity> Entities, const float DelayInSeconds)
+void UMassSignalSubsystem::DelaySignalEntities(FName SignalName, TConstArrayView<FMassEntityHandle> Entities, const float DelayInSeconds)
 {
 	FDelayedSignal& DelayedSignal = DelayedSignals.Emplace_GetRef();
 	DelayedSignal.SignalName = SignalName;
@@ -68,13 +68,13 @@ void UMassSignalSubsystem::DelaySignalEntities(FName SignalName, TConstArrayView
 	UE_CVLOG(Entities.Num() > 1,this, LogMassSignals, Log, TEXT("Delay signal [%s] to %d entities in %.2f"), *SignalName.ToString(), Entities.Num(), DelayInSeconds);
 }
 
-void UMassSignalSubsystem::SignalEntityDeferred(FLWComponentSystemExecutionContext& Context, FName SignalName, const FLWEntity Entity)
+void UMassSignalSubsystem::SignalEntityDeferred(FMassExecutionContext& Context, FName SignalName, const FMassEntityHandle Entity)
 {
 	checkf(Entity.IsSet(), TEXT("Expecting a valid entity to signal"));
 	SignalEntitiesDeferred(Context, SignalName, MakeArrayView(&Entity, 1));
 }
 
-void UMassSignalSubsystem::SignalEntitiesDeferred(FLWComponentSystemExecutionContext& Context, FName SignalName, TConstArrayView<FLWEntity> Entities)
+void UMassSignalSubsystem::SignalEntitiesDeferred(FMassExecutionContext& Context, FName SignalName, TConstArrayView<FMassEntityHandle> Entities)
 {
 	checkf(Entities.Num() > 0, TEXT("Expecting entities to signal"));
 	Context.Defer().EmplaceCommand<FDeferredCommand>([SignalName, Entities](const UMassEntitySubsystem& System)
@@ -87,13 +87,13 @@ void UMassSignalSubsystem::SignalEntitiesDeferred(FLWComponentSystemExecutionCon
 	UE_CVLOG(Entities.Num() > 1, this, LogMassSignals, Log, TEXT("Raising deferred signal [%s] to %d entities"), *SignalName.ToString(), Entities.Num());
 }
 
-void UMassSignalSubsystem::DelaySignalEntityDeferred(FLWComponentSystemExecutionContext& Context, FName SignalName, const FLWEntity Entity, const float DelayInSeconds)
+void UMassSignalSubsystem::DelaySignalEntityDeferred(FMassExecutionContext& Context, FName SignalName, const FMassEntityHandle Entity, const float DelayInSeconds)
 {
 	checkf(Entity.IsSet(), TEXT("Expecting a valid entity to signal"));
 	DelaySignalEntitiesDeferred(Context, SignalName, MakeArrayView(&Entity, 1), DelayInSeconds);
 }
 
-void UMassSignalSubsystem::DelaySignalEntitiesDeferred(FLWComponentSystemExecutionContext& Context, FName SignalName, TConstArrayView<FLWEntity> Entities, const float DelayInSeconds)
+void UMassSignalSubsystem::DelaySignalEntitiesDeferred(FMassExecutionContext& Context, FName SignalName, TConstArrayView<FMassEntityHandle> Entities, const float DelayInSeconds)
 {
 	checkf(Entities.Num() > 0, TEXT("Expecting entities to signal"));
 	Context.Defer().EmplaceCommand<FDeferredCommand>([SignalName, Entities, DelayInSeconds](const UMassEntitySubsystem& System)
diff --git a/Engine/Plugins/Runtime/MassGameplay/Source/MassSignals/Private/MassSignalTypes.cpp b/Engine/Plugins/Runtime/MassGameplay/Source/MassSignals/Private/MassSignalTypes.cpp
--- a/Engine/Plugins/Runtime/MassGameplay/Source/MassSignals/Private/MassSignalTypes.cpp
+++ b/Engine/Plugins/Runtime/MassGameplay/Source/MassSignals/Private/MassSignalTypes.cpp
@@ -17,13 +17,13 @@ uint64 FMassSignalNameLookup::GetOrAddSignalName(const FName SignalName)
 	return 1ULL << Index;
 }
 
-void FMassSignalNameLookup::AddSignalToEntity(const FLWEntity Entity, const uint64 SignalFlag)
+void FMassSignalNameLookup::AddSignalToEntity(const FMassEntityHandle Entity, const uint64 SignalFlag)
 {
 	uint64& Signals = EntitySignals.FindOrAdd(Entity, 0);
 	Signals |= SignalFlag;
 }
 
-void FMassSignalNameLookup::GetSignalsForEntity(const FLWEntity Entity, TArray<FName>& OutSignals) const
+void FMassSignalNameLookup::GetSignalsForEntity(const FMassEntityHandle Entity, TArray<FName>& OutSignals) const
 {
 	OutSignals.Reset();
 	if (const uint64* Signals = EntitySignals.Find(Entity))
diff --git a/Engine/Plugins/Runtime/MassGameplay/Source/MassSignals/Public/MassSignalProcessorBase.h b/Engine/Plugins/Runtime/MassGameplay/Source/MassSignals/Public/MassSignalProcessorBase.h
--- a/Engine/Plugins/Runtime/MassGameplay/Source/MassSignals/Public/MassSignalProcessorBase.h
+++ b/Engine/Plugins/Runtime/MassGameplay/Source/MassSignals/Public/MassSignalProcessorBase.h
@@ -15,7 +15,7 @@ class UMassSignalSubsystem;
  * The derived classes only need to implement the method SignalEntities to actually received the raised signals for the entities they subscribed to 
  */
 UCLASS(abstract)
-class MASSSIGNALS_API UMassSignalProcessorBase : public UPipeProcessor
+class MASSSIGNALS_API UMassSignalProcessorBase : public UMassProcessor
 {
 	GENERATED_BODY()
 
@@ -28,7 +28,7 @@ protected:
 
 	virtual void BeginDestroy() override;
 
-	/** Configure the owned FLWComponentQuery instances to express processor queries requirements */
+	/** Configure the owned FMassEntityQuery instances to express processor queries requirements */
 	virtual void ConfigureQueries() override {}
 
 	/**
@@ -37,21 +37,21 @@ protected:
 	 * @param Context is the execution context to be passed when executing the lambdas
 	 * @param EntitySignals Look up to retrieve for each entities their raised signal via GetSignalsForEntity
 	 */
-	 virtual void SignalEntities(UMassEntitySubsystem& EntitySubsystem, FLWComponentSystemExecutionContext& Context, FMassSignalNameLookup& EntitySignals) PURE_VIRTUAL(UMassSignalProcessorBase::SignalEntities, );
+	 virtual void SignalEntities(UMassEntitySubsystem& EntitySubsystem, FMassExecutionContext& Context, FMassSignalNameLookup& EntitySignals) PURE_VIRTUAL(UMassSignalProcessorBase::SignalEntities, );
 
 	/**
 	 * Callback that is being called when new signal is raised
 	 * @param SignalName is the name of the signal being raised
 	 * @param Entities are the targeted entities for this signal
 	 */
-	 virtual void OnSignalReceived(FName SignalName, TConstArrayView<FLWEntity> Entities);
+	 virtual void OnSignalReceived(FName SignalName, TConstArrayView<FMassEntityHandle> Entities);
 
 	/**
 	 * Execution method for this processor
 	 * @param EntitySubsystem is the system to execute the lambdas on each entity chunk
 	 * @param Context is the execution context to be passed when executing the lambdas
 	 */
-	virtual void Execute(UMassEntitySubsystem& EntitySubsystem, FLWComponentSystemExecutionContext& Context) override;
+	virtual void Execute(UMassEntitySubsystem& EntitySubsystem, FMassExecutionContext& Context) override;
 
 	/**
 	 * To receive notification about a particular signal, you need to subscribe to it.
@@ -63,7 +63,7 @@ protected:
 	UPROPERTY(Transient)
 	UMassSignalSubsystem* SignalSubsystem = nullptr;
 
-	FLWComponentQuery EntityQuery;
+	FMassEntityQuery EntityQuery;
 
 private:
 
@@ -82,7 +82,7 @@ private:
 		TArray<FEntitySignalRange> ReceivedSignalRanges;
 
 		/** the list of all signaled entities, can contain duplicates */
-		TArray<FLWEntity> SignaledEntities;
+		TArray<FMassEntityHandle> SignaledEntities;
 	};
 
 	/** Double buffer frame received signal as we can receive new signals as we are processing them */
diff --git a/Engine/Plugins/Runtime/MassGameplay/Source/MassSignals/Public/MassSignalSubsystem.h b/Engine/Plugins/Runtime/MassGameplay/Source/MassSignals/Public/MassSignalSubsystem.h
--- a/Engine/Plugins/Runtime/MassGameplay/Source/MassSignals/Public/MassSignalSubsystem.h
+++ b/Engine/Plugins/Runtime/MassGameplay/Source/MassSignals/Public/MassSignalSubsystem.h
@@ -8,7 +8,7 @@
 
 namespace UE::MassSignal 
 {
-	DECLARE_MULTICAST_DELEGATE_TwoParams(FSignalDelegate, FName /*SignalName*/, TConstArrayView<FLWEntity> /*Entities*/);
+	DECLARE_MULTICAST_DELEGATE_TwoParams(FSignalDelegate, FName /*SignalName*/, TConstArrayView<FMassEntityHandle> /*Entities*/);
 } // UE::MassSignal
 
 /**
@@ -35,14 +35,14 @@ public:
 	 * @param SignalName is the name of the signal raised
 	 * @param Entity entity that should be informed that signal 'SignalName' was raised
 	 */
-	void SignalEntity(FName SignalName, const FLWEntity Entity);
+	void SignalEntity(FName SignalName, const FMassEntityHandle Entity);
 
 	/**
 	 * Inform multiple entities of a signal being raised
 	 * @param SignalName is the name of the signal raised
 	 * @param Entities list of entities that should be informed that signal 'SignalName' was raised
 	 */
-	void SignalEntities(FName SignalName, TConstArrayView<FLWEntity> Entities);
+	void SignalEntities(FName SignalName, TConstArrayView<FMassEntityHandle> Entities);
 
 	/**
 	 * Inform a single entity of a signal being raised in a certain amount of seconds
@@ -50,7 +50,7 @@ public:
 	 * @param Entity entity that should be informed that signal 'SignalName' was raised
 	 * @param DelayInSeconds is the amount of time before signaling the entity
 	 */
-	void DelaySignalEntity(FName SignalName, const FLWEntity Entity, const float DelayInSeconds);
+	void DelaySignalEntity(FName SignalName, const FMassEntityHandle Entity, const float DelayInSeconds);
 
  	/**
 	 * Inform multiple entities of a signal being raised in a certain amount of seconds
@@ -58,41 +58,41 @@ public:
 	 * @param Entities being informed of the raised signal
 	 * @param DelayInSeconds is the amount of time before signaling the entities
 	 */
-	void DelaySignalEntities(FName SignalName, TConstArrayView<FLWEntity> Entities, const float DelayInSeconds);
+	void DelaySignalEntities(FName SignalName, TConstArrayView<FMassEntityHandle> Entities, const float DelayInSeconds);
 
 	/**
-	 * Inform single entity of a signal being raised asynchronously using the LWComponent Command Buffer
+	 * Inform single entity of a signal being raised asynchronously using the Mass Command Buffer
 	 * @param Context is the Entity System execution context to push the async command
 	 * @param SignalName is the name of the signal raised
 	 * @param Entity entity that should be informed that signal 'SignalName' was raised
 	 */
-	void SignalEntityDeferred(FLWComponentSystemExecutionContext& Context, FName SignalName, const FLWEntity Entity);
+	void SignalEntityDeferred(FMassExecutionContext& Context, FName SignalName, const FMassEntityHandle Entity);
 
 	/**
-	 * Inform multiple entities of a signal being raised asynchronously using the LWComponent Command Buffer
+	 * Inform multiple entities of a signal being raised asynchronously using the Mass Command Buffer
 	 * @param Context is the Entity System execution context to push the async command
 	 * @param SignalName is the name of the signal raised
 	 * @param Entities list of entities that should be informed that signal 'SignalName' was raised
 	 */
-	void SignalEntitiesDeferred(FLWComponentSystemExecutionContext& Context, FName SignalName, TConstArrayView<FLWEntity> Entities);
+	void SignalEntitiesDeferred(FMassExecutionContext& Context, FName SignalName, TConstArrayView<FMassEntityHandle> Entities);
 
  	/**
-	 * Inform single entity of a signal being raised asynchronously using the LWComponent Command Buffer
+	 * Inform single entity of a signal being raised asynchronously using the Mass Command Buffer
 	 * @param Context is the Entity System execution context to push the async command
 	 * @param SignalName is the name of the signal raised
 	 * @param Entity entity that should be informed that signal 'SignalName' was raised
 	 * @param DelayInSeconds is the amount of time before signaling the entities
 	 */
-	void DelaySignalEntityDeferred(FLWComponentSystemExecutionContext& Context, FName SignalName, const FLWEntity Entity, const float DelayInSeconds);
+	void DelaySignalEntityDeferred(FMassExecutionContext& Context, FName SignalName, const FMassEntityHandle Entity, const float DelayInSeconds);
 
  	/**
-	 * Inform multiple entities of a signal being raised asynchronously using the LWComponent Command Buffer
+	 * Inform multiple entities of a signal being raised asynchronously using the Mass Command Buffer
 	 * @param Context is the Entity System execution context to push the async command
 	 * @param SignalName is the name of the signal raised
 	 * @param Entities being informed of that signal was raised
 	 * @param DelayInSeconds is the amount of time before signaling the entities
 	 */
-	void DelaySignalEntitiesDeferred(FLWComponentSystemExecutionContext& Context, FName SignalName, TConstArrayView<FLWEntity> Entities, const float DelayInSeconds);
+	void DelaySignalEntitiesDeferred(FMassExecutionContext& Context, FName SignalName, TConstArrayView<FMassEntityHandle> Entities, const float DelayInSeconds);
 
 protected:
 	virtual void Tick(float DeltaTime) override;
@@ -103,7 +103,7 @@ protected:
 	struct FDelayedSignal
 	{
 		FName SignalName;
-		TArray<FLWEntity> Entities;
+		TArray<FMassEntityHandle> Entities;
 		float DelayInSeconds;
 	};
 
diff --git a/Engine/Plugins/Runtime/MassGameplay/Source/MassSignals/Public/MassSignalTypes.h b/Engine/Plugins/Runtime/MassGameplay/Source/MassSignals/Public/MassSignalTypes.h
--- a/Engine/Plugins/Runtime/MassGameplay/Source/MassSignals/Public/MassSignalTypes.h
+++ b/Engine/Plugins/Runtime/MassGameplay/Source/MassSignals/Public/MassSignalTypes.h
@@ -31,13 +31,13 @@ struct MASSSIGNALS_API FMassSignalNameLookup
 	 * @param Entity is the entity where the signal has been raised
 	 * @param SignalFlag is the actual bitflag describing the signal
 	 */
-	void AddSignalToEntity(const FLWEntity Entity, const uint64 SignalFlag);
+	void AddSignalToEntity(const FMassEntityHandle Entity, const uint64 SignalFlag);
 
 	/** 
 	 * Retrieve for a specific entity the raised signal this frame
 	 * @return Array of signal names raised for this entity 
 	 */
-	void GetSignalsForEntity(const FLWEntity Entity, TArray<FName>& OutSignals) const;
+	void GetSignalsForEntity(const FMassEntityHandle Entity, TArray<FName>& OutSignals) const;
 
 	/** Empties the name lookup and entity signals */
 	void Reset();
@@ -47,5 +47,5 @@ protected:
 	TArray<FName> SignalNames;
 
 	/** Map from entity id to name bitmask */
-	TMap<FLWEntity, uint64> EntitySignals;
+	TMap<FMassEntityHandle, uint64> EntitySignals;
 };
\ No newline at end of file
diff --git a/Engine/Plugins/Runtime/MassGameplay/Source/MassSimulation/Private/MassSimulationSettings.cpp b/Engine/Plugins/Runtime/MassGameplay/Source/MassSimulation/Private/MassSimulationSettings.cpp
--- a/Engine/Plugins/Runtime/MassGameplay/Source/MassSimulation/Private/MassSimulationSettings.cpp
+++ b/Engine/Plugins/Runtime/MassGameplay/Source/MassSimulation/Private/MassSimulationSettings.cpp
@@ -33,9 +33,9 @@ void UMassSimulationSettings::PostEditChangeProperty(struct FPropertyChangedEven
 	}
 }
 
-void UMassSimulationSettings::OnAssetPropertiesChanged(UPipeSchematic* PipeSchematic, const FPropertyChangedEvent& PropertyChangedEvent)
+void UMassSimulationSettings::OnAssetPropertiesChanged(UMassSchematic* MassSchematic, const FPropertyChangedEvent& PropertyChangedEvent)
 {
-	if (TickSchematics.Find(PipeSchematic) != INDEX_NONE)
+	if (TickSchematics.Find(MassSchematic) != INDEX_NONE)
 	{
 		OnTickSchematicChanged.Broadcast();
 	}
diff --git a/Engine/Plugins/Runtime/MassGameplay/Source/MassSimulation/Private/MassSimulationSubsystem.cpp b/Engine/Plugins/Runtime/MassGameplay/Source/MassSimulation/Private/MassSimulationSubsystem.cpp
--- a/Engine/Plugins/Runtime/MassGameplay/Source/MassSimulation/Private/MassSimulationSubsystem.cpp
+++ b/Engine/Plugins/Runtime/MassGameplay/Source/MassSimulation/Private/MassSimulationSubsystem.cpp
@@ -24,7 +24,7 @@ namespace UE::MassSimulation
 UMassSimulationSubsystem::UMassSimulationSubsystem(const FObjectInitializer& ObjectInitializer)
 	: Super()
 {
-	PhaseManager = CreateDefaultSubobject<UPipeProcessingPhaseManager>(TEXT("PhaseManager"));
+	PhaseManager = CreateDefaultSubobject<UMassProcessingPhaseManager>(TEXT("PhaseManager"));
 }
 
 void UMassSimulationSubsystem::PostInitProperties()
@@ -56,13 +56,13 @@ void UMassSimulationSubsystem::BeginDestroy()
 	Super::BeginDestroy();
 }
 
-FPipeProcessingPhase::FOnPhaseEvent& UMassSimulationSubsystem::GetOnProcessingPhaseStarted(const EPipeProcessingPhase Phase) const
+FMassProcessingPhase::FOnPhaseEvent& UMassSimulationSubsystem::GetOnProcessingPhaseStarted(const EMassProcessingPhase Phase) const
 {
 	check(PhaseManager);
 	return PhaseManager->GetOnPhaseStart(Phase);
 }
 
-FPipeProcessingPhase::FOnPhaseEvent& UMassSimulationSubsystem::GetOnProcessingPhaseFinished(const EPipeProcessingPhase Phase) const
+FMassProcessingPhase::FOnPhaseEvent& UMassSimulationSubsystem::GetOnProcessingPhaseFinished(const EMassProcessingPhase Phase) const
 {
 	check(PhaseManager);
 	return PhaseManager->GetOnPhaseEnd(Phase);
@@ -75,7 +75,7 @@ void UMassSimulationSubsystem::Initialize(FSubsystemCollectionBase& Collection)
 	EntitySubsystem = Collection.InitializeDependency<UMassEntitySubsystem>();
 	check(EntitySubsystem);
 	
-	GetOnProcessingPhaseStarted(EPipeProcessingPhase::PrePhysics).AddUObject(this, &UMassSimulationSubsystem::OnProcessingPhaseStarted, EPipeProcessingPhase::PrePhysics);
+	GetOnProcessingPhaseStarted(EMassProcessingPhase::PrePhysics).AddUObject(this, &UMassSimulationSubsystem::OnProcessingPhaseStarted, EMassProcessingPhase::PrePhysics);
 }
 
 void UMassSimulationSubsystem::Deinitialize()
@@ -83,7 +83,7 @@ void UMassSimulationSubsystem::Deinitialize()
 #if WITH_EDITOR
 	GET_MASS_CONFIG_VALUE(GetOnTickSchematicChanged()).RemoveAll(this);
 #endif // WITH_EDITOR
-	GetOnProcessingPhaseStarted(EPipeProcessingPhase::PrePhysics).RemoveAll(this);
+	GetOnProcessingPhaseStarted(EMassProcessingPhase::PrePhysics).RemoveAll(this);
 	StopSimulation();
 
 	Super::Deinitialize();
@@ -141,11 +141,11 @@ void UMassSimulationSubsystem::StopSimulation()
 	bSimulationStarted = false;
 }
 
-void UMassSimulationSubsystem::OnProcessingPhaseStarted(const float DeltaSeconds, const EPipeProcessingPhase Phase) const
+void UMassSimulationSubsystem::OnProcessingPhaseStarted(const float DeltaSeconds, const EMassProcessingPhase Phase) const
 {
 	switch (Phase)
 	{
-		case EPipeProcessingPhase::PrePhysics:
+		case EMassProcessingPhase::PrePhysics:
 			{
 				TRACE_CPUPROFILER_EVENT_SCOPE(DoEntityCompation);
 				check(EntitySubsystem);
diff --git a/Engine/Plugins/Runtime/MassGameplay/Source/MassSimulation/Public/MassSimulationLocalCoordinator.h b/Engine/Plugins/Runtime/MassGameplay/Source/MassSimulation/Public/MassSimulationLocalCoordinator.h
--- a/Engine/Plugins/Runtime/MassGameplay/Source/MassSimulation/Public/MassSimulationLocalCoordinator.h
+++ b/Engine/Plugins/Runtime/MassGameplay/Source/MassSimulation/Public/MassSimulationLocalCoordinator.h
@@ -15,7 +15,7 @@
 #include "MassSimulationLocalCoordinator.generated.h"
 
 
-class UPipeSchematic;
+class UMassSchematic;
 class UMassEntitySubsystem;
 class UHierarchicalInstancedStaticMeshComponent;
 class UMassSpawningManager;
@@ -45,11 +45,11 @@ public:
 protected:
 	// all spawned entities will go through these pipelines
 	UPROPERTY(EditAnywhere, Category = "Mass|Spawning")
-	TArray<TSoftObjectPtr<UPipeSchematic>> PostSpawnSchematics;
+	TArray<TSoftObjectPtr<UMassSchematic>> PostSpawnSchematics;
 
 	// Every tick these pipelines will be spawned 
 	UPROPERTY(EditAnywhere, Category = "Mass|Simulation")
-	TArray<TSoftObjectPtr<UPipeSchematic>> TickSchematics;
+	TArray<TSoftObjectPtr<UMassSchematic>> TickSchematics;
 
 	UPROPERTY(EditAnywhere, Category = "Mass|Spawning", meta=(DisplayName="LEGACY_LocallyAppliedSpawnComponents"))
 	TArray<const UScriptStruct*> LocallyAppliedSpawnComponents;
diff --git a/Engine/Plugins/Runtime/MassGameplay/Source/MassSimulation/Public/MassSimulationSettings.h b/Engine/Plugins/Runtime/MassGameplay/Source/MassSimulation/Public/MassSimulationSettings.h
--- a/Engine/Plugins/Runtime/MassGameplay/Source/MassSimulation/Public/MassSimulationSettings.h
+++ b/Engine/Plugins/Runtime/MassGameplay/Source/MassSimulation/Public/MassSimulationSettings.h
@@ -7,7 +7,7 @@
 
 #define GET_MASS_CONFIG_VALUE(a) (GetMutableDefault<UMassSimulationSettings>()->a)
 
-class UPipeSchematic;
+class UMassSchematic;
 
 /**
  * Implements the settings for MassSimulation
@@ -26,12 +26,12 @@ public:
 protected:
 	virtual void PostInitProperties() override;
 	virtual void PostEditChangeProperty(struct FPropertyChangedEvent& PropertyChangedEvent) override;
-	void OnAssetPropertiesChanged(UPipeSchematic* PipeSchematic, const FPropertyChangedEvent& PropertyChangedEvent);	
+	void OnAssetPropertiesChanged(UMassSchematic* MassSchematic, const FPropertyChangedEvent& PropertyChangedEvent);	
 #endif // WITH_EDITOR
 
 public:
 	UPROPERTY(EditDefaultsOnly, Category = "Runtime", config, meta=(DisplayName="DEPRECATED_TickSchematics"))
-	TArray<TSoftObjectPtr<UPipeSchematic>> TickSchematics;
+	TArray<TSoftObjectPtr<UMassSchematic>> TickSchematics;
 
 	/** The desired budget in seconds allowed to do actor spawning per frame */
 	UPROPERTY(EditDefaultsOnly, Category = "Runtime", config)
diff --git a/Engine/Plugins/Runtime/MassGameplay/Source/MassSimulation/Public/MassSimulationSubsystem.h b/Engine/Plugins/Runtime/MassGameplay/Source/MassSimulation/Public/MassSimulationSubsystem.h
--- a/Engine/Plugins/Runtime/MassGameplay/Source/MassSimulation/Public/MassSimulationSubsystem.h
+++ b/Engine/Plugins/Runtime/MassGameplay/Source/MassSimulation/Public/MassSimulationSubsystem.h
@@ -11,7 +11,7 @@
 
 class AMassSimulationLocalCoordinator;
 class UMassEntitySubsystem;
-class UPipeProcessingPhaseManager;
+class UMassProcessingPhaseManager;
 
 DECLARE_LOG_CATEGORY_EXTERN(LogMassSim, Log, All);
 
@@ -20,16 +20,16 @@ class MASSSIMULATION_API UMassSimulationSubsystem : public UWorldSubsystem
 {
 	GENERATED_BODY()
 public:
-	DECLARE_MULTICAST_DELEGATE_TwoParams(FOnAdjustTickSchematics, UWorld* /*World*/, TArray<TSoftObjectPtr<UPipeSchematic>>& /*InOutTickSchematics*/);
+	DECLARE_MULTICAST_DELEGATE_TwoParams(FOnAdjustTickSchematics, UWorld* /*World*/, TArray<TSoftObjectPtr<UMassSchematic>>& /*InOutTickSchematics*/);
 	DECLARE_MULTICAST_DELEGATE_OneParam(FOnSimulationStarted, UWorld* /*World*/);
 	
 	UMassSimulationSubsystem(const FObjectInitializer& ObjectInitializer = FObjectInitializer::Get());
 
 	//UMassEntitySubsystem* GetEntitySubsystem() const { return CachedEntitySubsystem; }
-	const UPipeProcessingPhaseManager& GetPhaseManager() const { check(PhaseManager); return *PhaseManager; }
+	const UMassProcessingPhaseManager& GetPhaseManager() const { check(PhaseManager); return *PhaseManager; }
 
-	FPipeProcessingPhase::FOnPhaseEvent& GetOnProcessingPhaseStarted(const EPipeProcessingPhase Phase) const;
-	FPipeProcessingPhase::FOnPhaseEvent& GetOnProcessingPhaseFinished(const EPipeProcessingPhase Phase) const;
+	FMassProcessingPhase::FOnPhaseEvent& GetOnProcessingPhaseStarted(const EMassProcessingPhase Phase) const;
+	FMassProcessingPhase::FOnPhaseEvent& GetOnProcessingPhaseFinished(const EMassProcessingPhase Phase) const;
 	static FOnAdjustTickSchematics& GetOnAdjustTickSchematics() { return OnAdjustTickSchematics; }
 	static FOnSimulationStarted& GetOnSimulationStarted() { return OnSimulationStarted; }
 
@@ -48,7 +48,7 @@ protected:
 	void StartSimulation(UWorld& InWorld);
 	void StopSimulation();
 
-	void OnProcessingPhaseStarted(const float DeltaSeconds, const EPipeProcessingPhase Phase) const;
+	void OnProcessingPhaseStarted(const float DeltaSeconds, const EMassProcessingPhase Phase) const;
 
 #if WITH_EDITOR
 	void OnPieBegin(const bool bIsSimulation);
@@ -61,13 +61,13 @@ protected:
 	UMassEntitySubsystem* EntitySubsystem;
 
 	UPROPERTY()
-	UPipeProcessingPhaseManager* PhaseManager;
+	UMassProcessingPhaseManager* PhaseManager;
 
 	inline static FOnAdjustTickSchematics OnAdjustTickSchematics={};
 	inline static FOnSimulationStarted OnSimulationStarted={};
 
 	UPROPERTY()
-	FRuntimePipeline RuntimePipeline;
+	FMassRuntimePipeline RuntimePipeline;
 
 	float CurrentDeltaSeconds = 0.f;
 	bool bTickInProgress = false;
diff --git a/Engine/Plugins/Runtime/MassGameplay/Source/MassSmartObjects/Private/MassSmartObjectBehaviorConfig.cpp b/Engine/Plugins/Runtime/MassGameplay/Source/MassSmartObjects/Private/MassSmartObjectBehaviorConfig.cpp
--- a/Engine/Plugins/Runtime/MassGameplay/Source/MassSmartObjects/Private/MassSmartObjectBehaviorConfig.cpp
+++ b/Engine/Plugins/Runtime/MassGameplay/Source/MassSmartObjects/Private/MassSmartObjectBehaviorConfig.cpp
@@ -4,7 +4,7 @@
 #include "MassSmartObjectProcessor.h"
 
 void USmartObjectMassBehaviorConfig::Activate(UMassEntitySubsystem& EntitySubsystem,
-                                              FLWComponentSystemExecutionContext& Context,
+                                              FMassExecutionContext& Context,
                                               const FMassBehaviorEntityContext& EntityContext) const
 {
 	EntityContext.SOUser.SetUseTime(UseTime);
