AggregateTicking and Pipe type renames:

diff --git a/Engine/Plugins/Runtime/MassGameplay/Source/MassSmartObjects/Private/MassSmartObjectHandler.cpp b/Engine/Plugins/Runtime/MassGameplay/Source/MassSmartObjects/Private/MassSmartObjectHandler.cpp
--- a/Engine/Plugins/Runtime/MassGameplay/Source/MassSmartObjects/Private/MassSmartObjectHandler.cpp
+++ b/Engine/Plugins/Runtime/MassGameplay/Source/MassSmartObjects/Private/MassSmartObjectHandler.cpp
@@ -11,7 +11,7 @@ namespace UE::Mass::SmartObject
 {
 struct FPayload
 {
-	FLWEntity Entity;
+	FMassEntityHandle Entity;
 	TWeakObjectPtr<UMassEntitySubsystem> EntitySubsystem;
 	TWeakObjectPtr<USmartObjectSubsystem> SmartObjectSubsystem;
 };
@@ -32,9 +32,9 @@ void OnSlotInvalidated(const FSmartObjectClaimHandle& ClaimHandle, const ESmartO
 //----------------------------------------------------------------------//
 // FMassSmartObjectHandler
 //----------------------------------------------------------------------//
-FMassSmartObjectRequestID FMassSmartObjectHandler::FindCandidatesAsync(const FLWEntity RequestingEntity, const FVector& Location) const
+FMassSmartObjectRequestID FMassSmartObjectHandler::FindCandidatesAsync(const FMassEntityHandle RequestingEntity, const FVector& Location) const
 {
-	const FLWEntity ReservedEntity = EntitySubsystem.ReserveEntity();
+	const FMassEntityHandle ReservedEntity = EntitySubsystem.ReserveEntity();
 
 	FMassSmartObjectWorldLocationRequestFragment RequestFragment;
 	RequestFragment.SearchOrigin = Location;
@@ -52,9 +52,9 @@ FMassSmartObjectRequestID FMassSmartObjectHandler::FindCandidatesAsync(const FLW
 	return FMassSmartObjectRequestID(ReservedEntity);
 }
 
-FMassSmartObjectRequestID FMassSmartObjectHandler::FindCandidatesAsync(const FLWEntity RequestingEntity, const FZoneGraphCompactLaneLocation& LaneLocation) const
+FMassSmartObjectRequestID FMassSmartObjectHandler::FindCandidatesAsync(const FMassEntityHandle RequestingEntity, const FZoneGraphCompactLaneLocation& LaneLocation) const
 {
-	const FLWEntity ReservedEntity = EntitySubsystem.ReserveEntity();
+	const FMassEntityHandle ReservedEntity = EntitySubsystem.ReserveEntity();
 
 	FMassSmartObjectLaneLocationRequestFragment RequestFragment;
 	RequestFragment.CompactLaneLocation = LaneLocation;
@@ -74,7 +74,7 @@ FMassSmartObjectRequestID FMassSmartObjectHandler::FindCandidatesAsync(const FLW
 
 FMassSmartObjectRequestResult FMassSmartObjectHandler::GetRequestResult(const FMassSmartObjectRequestID& RequestID) const
 {
-	const FLWEntity RequestEntity = static_cast<FLWEntity>(RequestID);
+	const FMassEntityHandle RequestEntity = static_cast<FMassEntityHandle>(RequestID);
 	if (!ensureMsgf(EntitySubsystem.IsEntityValid(RequestEntity), TEXT("Invalid request.")))
 	{
 		return {};
@@ -92,18 +92,18 @@ FMassSmartObjectRequestResult FMassSmartObjectHandler::GetRequestResult(const FM
 
 void FMassSmartObjectHandler::RemoveRequest(const FMassSmartObjectRequestID& RequestID) const
 {
-	const FLWEntity RequestEntity = static_cast<FLWEntity>(RequestID);
+	const FMassEntityHandle RequestEntity = static_cast<FMassEntityHandle>(RequestID);
 	ExecutionContext.Defer().DestroyEntity(RequestEntity);
 }
 
-EMassSmartObjectClaimResult FMassSmartObjectHandler::ClaimCandidate(const FLWEntity Entity, FDataFragment_SmartObjectUser& User, const FMassSmartObjectRequestID& RequestID) const
+EMassSmartObjectClaimResult FMassSmartObjectHandler::ClaimCandidate(const FMassEntityHandle Entity, FDataFragment_SmartObjectUser& User, const FMassSmartObjectRequestID& RequestID) const
 {
 	if (!ensureMsgf(RequestID.IsSet(), TEXT("Trying to claim using an invalid request.")))
 	{
 		return EMassSmartObjectClaimResult::Failed_InvalidRequest;
 	}
 
-	const FLWEntity RequestEntity = static_cast<FLWEntity>(RequestID);
+	const FMassEntityHandle RequestEntity = static_cast<FMassEntityHandle>(RequestID);
 	const FMassSmartObjectRequestResultFragment& RequestFragment = EntitySubsystem.GetComponentDataChecked<FMassSmartObjectRequestResultFragment>(RequestEntity);
 	if (!ensureMsgf(RequestFragment.Result.bProcessed, TEXT("Trying to claim using a request that is not processed.")))
 	{
@@ -113,7 +113,7 @@ EMassSmartObjectClaimResult FMassSmartObjectHandler::ClaimCandidate(const FLWEnt
 	return ClaimCandidate(Entity, User, RequestFragment.Result);
 }
 
-EMassSmartObjectClaimResult FMassSmartObjectHandler::ClaimCandidate(const FLWEntity Entity, FDataFragment_SmartObjectUser& User, const FMassSmartObjectRequestResult& SearchRequestResult) const
+EMassSmartObjectClaimResult FMassSmartObjectHandler::ClaimCandidate(const FMassEntityHandle Entity, FDataFragment_SmartObjectUser& User, const FMassSmartObjectRequestResult& SearchRequestResult) const
 {
 	if (!ensureMsgf(SearchRequestResult.bProcessed, TEXT("Trying to claim using a search request that is not processed.")))
 	{
@@ -143,7 +143,7 @@ EMassSmartObjectClaimResult FMassSmartObjectHandler::ClaimCandidate(const FLWEnt
 	return User.ClaimHandle.IsValid() ? EMassSmartObjectClaimResult::Succeeded : EMassSmartObjectClaimResult::Failed_NoAvailableCandidate;
 }
 
-bool FMassSmartObjectHandler::ClaimSmartObject(const FLWEntity Entity, FDataFragment_SmartObjectUser& User, const FSmartObjectID& ObjectID) const
+bool FMassSmartObjectHandler::ClaimSmartObject(const FMassEntityHandle Entity, FDataFragment_SmartObjectUser& User, const FSmartObjectID& ObjectID) const
 {
 	FSmartObjectRequestFilter Filter;
 	Filter.BehaviorConfigurationClass = USmartObjectMassBehaviorConfig::StaticClass();
@@ -179,7 +179,7 @@ bool FMassSmartObjectHandler::ClaimSmartObject(const FLWEntity Entity, FDataFrag
 }
 
 bool FMassSmartObjectHandler::UseSmartObject(
-	const FLWEntity Entity,
+	const FMassEntityHandle Entity,
 	FDataFragment_SmartObjectUser& User,
 	const FDataFragment_Transform& Transform) const
 {
@@ -207,7 +207,7 @@ bool FMassSmartObjectHandler::UseSmartObject(
 	return true;
 }
 
-void FMassSmartObjectHandler::ReleaseSmartObject(const FLWEntity Entity, FDataFragment_SmartObjectUser& User, const EMassSmartObjectInteractionStatus NewStatus) const
+void FMassSmartObjectHandler::ReleaseSmartObject(const FMassEntityHandle Entity, FDataFragment_SmartObjectUser& User, const EMassSmartObjectInteractionStatus NewStatus) const
 {
 #if WITH_MASSGAMEPLAY_DEBUG
 	UE_CVLOG(UE::MassDebug::IsDebuggingEntity(Entity),
diff --git a/Engine/Plugins/Runtime/MassGameplay/Source/MassSmartObjects/Private/MassSmartObjectProcessor.cpp b/Engine/Plugins/Runtime/MassGameplay/Source/MassSmartObjects/Private/MassSmartObjectProcessor.cpp
--- a/Engine/Plugins/Runtime/MassGameplay/Source/MassSmartObjects/Private/MassSmartObjectProcessor.cpp
+++ b/Engine/Plugins/Runtime/MassGameplay/Source/MassSmartObjects/Private/MassSmartObjectProcessor.cpp
@@ -38,13 +38,13 @@ void UMassProcessor_SmartObjectCandidatesFinder::Initialize(UObject& Owner)
 
 void UMassProcessor_SmartObjectCandidatesFinder::ConfigureQueries()
 {
-	WorldRequestQuery.AddRequirement<FMassSmartObjectWorldLocationRequestFragment>(ELWComponentAccess::ReadOnly);
-	WorldRequestQuery.AddRequirement<FMassSmartObjectRequestResultFragment>(ELWComponentAccess::ReadWrite);
-	WorldRequestQuery.AddTagRequirement<FMassSmartObjectCompletedRequestTag>(ELWComponentPresence::None);
+	WorldRequestQuery.AddRequirement<FMassSmartObjectWorldLocationRequestFragment>(EMassFragmentAccess::ReadOnly);
+	WorldRequestQuery.AddRequirement<FMassSmartObjectRequestResultFragment>(EMassFragmentAccess::ReadWrite);
+	WorldRequestQuery.AddTagRequirement<FMassSmartObjectCompletedRequestTag>(EMassFragmentPresence::None);
 
-	LaneRequestQuery.AddRequirement<FMassSmartObjectLaneLocationRequestFragment>(ELWComponentAccess::ReadOnly);
-	LaneRequestQuery.AddRequirement<FMassSmartObjectRequestResultFragment>(ELWComponentAccess::ReadWrite);
-	LaneRequestQuery.AddTagRequirement<FMassSmartObjectCompletedRequestTag>(ELWComponentPresence::None);
+	LaneRequestQuery.AddRequirement<FMassSmartObjectLaneLocationRequestFragment>(EMassFragmentAccess::ReadOnly);
+	LaneRequestQuery.AddRequirement<FMassSmartObjectRequestResultFragment>(EMassFragmentAccess::ReadWrite);
+	LaneRequestQuery.AddTagRequirement<FMassSmartObjectCompletedRequestTag>(EMassFragmentPresence::None);
 }
 
 UMassProcessor_SmartObjectCandidatesFinder::UMassProcessor_SmartObjectCandidatesFinder()
@@ -56,7 +56,7 @@ UMassProcessor_SmartObjectCandidatesFinder::UMassProcessor_SmartObjectCandidates
 	ExecutionOrder.ExecuteBefore.Add(UE::Mass::ProcessorGroupNames::Behavior);
 }
 
-void UMassProcessor_SmartObjectCandidatesFinder::Execute(UMassEntitySubsystem& EntitySubsystem, FLWComponentSystemExecutionContext& Context)
+void UMassProcessor_SmartObjectCandidatesFinder::Execute(UMassEntitySubsystem& EntitySubsystem, FMassExecutionContext& Context)
 {
 	checkf(SmartObjectSubsystem != nullptr, TEXT("SmartObjectSubsystem should exist when executing processors."));
 	checkf(SignalSubsystem != nullptr, TEXT("MassSignalSubsystem should exist when executing processors."));
@@ -68,15 +68,15 @@ void UMassProcessor_SmartObjectCandidatesFinder::Execute(UMassEntitySubsystem& E
 	Filter.BehaviorConfigurationClass = USmartObjectMassBehaviorConfig::StaticClass();
 
 	// Build list of request owner entities to send a completion signal
-	TArray<FLWEntity> EntitiesToSignal;
+	TArray<FMassEntityHandle> EntitiesToSignal;
 
-	auto BeginRequestProcessing = [](const FLWEntity Entity, FLWComponentSystemExecutionContext& Context, FMassSmartObjectRequestResult& Result)
+	auto BeginRequestProcessing = [](const FMassEntityHandle Entity, FMassExecutionContext& Context, FMassSmartObjectRequestResult& Result)
 	{
 		Context.Defer().AddTag<FMassSmartObjectCompletedRequestTag>(Entity);
 		Result.NumCandidates = 0;
 	};
 
-	auto EndRequestProcessing = [](const UObject* LogOwner, const FLWEntity Entity, FMassSmartObjectRequestResult& Result)
+	auto EndRequestProcessing = [](const UObject* LogOwner, const FMassEntityHandle Entity, FMassSmartObjectRequestResult& Result)
 	{
 		if (Result.NumCandidates > 0)
 		{
@@ -91,7 +91,7 @@ void UMassProcessor_SmartObjectCandidatesFinder::Execute(UMassEntitySubsystem& E
 	};
 
 	// Process world location based requests
-	WorldRequestQuery.ForEachEntityChunk(EntitySubsystem, Context, [this, &Filter, &EntitiesToSignal, &BeginRequestProcessing, &EndRequestProcessing](FLWComponentSystemExecutionContext& Context)
+	WorldRequestQuery.ForEachEntityChunk(EntitySubsystem, Context, [this, &Filter, &EntitiesToSignal, &BeginRequestProcessing, &EndRequestProcessing](FMassExecutionContext& Context)
 	{
 		const FSmartObjectOctree& Octree = SmartObjectSubsystem->GetOctree();
 
@@ -110,7 +110,7 @@ void UMassProcessor_SmartObjectCandidatesFinder::Execute(UMassEntitySubsystem& E
 			const FVector& SearchOrigin = RequestFragment.SearchOrigin;
 			const FBox& SearchBounds = FBox::BuildAABB(SearchOrigin, FVector(SearchExtents));
 
-			const FLWEntity Entity = Context.GetEntity(i);
+			const FMassEntityHandle Entity = Context.GetEntity(i);
 			bool bDisplayDebug = false;
 			FColor DebugColor(FColor::White);
 
@@ -152,7 +152,7 @@ void UMassProcessor_SmartObjectCandidatesFinder::Execute(UMassEntitySubsystem& E
 	USmartObjectZoneAnnotations* SmartObjectBehavior = Cast<USmartObjectZoneAnnotations>(AnnotationSubsystem->GetFirstAnnotationForTag(SmartObjectTag));
 	if (SmartObjectBehavior != nullptr)
 	{
-		LaneRequestQuery.ForEachEntityChunk(EntitySubsystem, Context, [this, SmartObjectBehavior, &Filter, SmartObjectTag, &EntitiesToSignal, &BeginRequestProcessing, &EndRequestProcessing](FLWComponentSystemExecutionContext& Context)
+		LaneRequestQuery.ForEachEntityChunk(EntitySubsystem, Context, [this, SmartObjectBehavior, &Filter, SmartObjectTag, &EntitiesToSignal, &BeginRequestProcessing, &EndRequestProcessing](FMassExecutionContext& Context)
 		{
 			const int32 NumEntities = Context.GetEntitiesNum();
 			EntitiesToSignal.Reserve(EntitiesToSignal.Num() + NumEntities);
@@ -173,7 +173,7 @@ void UMassProcessor_SmartObjectCandidatesFinder::Execute(UMassEntitySubsystem& E
 				const FZoneGraphCompactLaneLocation RequestLocation = RequestFragment.CompactLaneLocation;
 				const FZoneGraphLaneHandle RequestLaneHandle = RequestLocation.LaneHandle;
 
-				const FLWEntity Entity = Context.GetEntity(i);
+				const FMassEntityHandle Entity = Context.GetEntity(i);
 				bool bDisplayDebug = false;
 #if WITH_MASSGAMEPLAY_DEBUG
 				FColor DebugColor(FColor::White);
@@ -268,8 +268,8 @@ void UMassProcessor_SmartObjectCandidatesFinder::Execute(UMassEntitySubsystem& E
 //----------------------------------------------------------------------//
 void UMassProcessor_SmartObjectTimedBehavior::ConfigureQueries()
 {
-	EntityQuery.AddRequirement<FDataFragment_SmartObjectUser>(ELWComponentAccess::ReadWrite);
-	EntityQuery.AddTagRequirement<FMassSmartObjectTimedBehaviorTag>(ELWComponentPresence::All);
+	EntityQuery.AddRequirement<FDataFragment_SmartObjectUser>(EMassFragmentAccess::ReadWrite);
+	EntityQuery.AddTagRequirement<FMassSmartObjectTimedBehaviorTag>(EMassFragmentPresence::All);
 }
 
 UMassProcessor_SmartObjectTimedBehavior::UMassProcessor_SmartObjectTimedBehavior()
@@ -278,16 +278,16 @@ UMassProcessor_SmartObjectTimedBehavior::UMassProcessor_SmartObjectTimedBehavior
 	ExecutionOrder.ExecuteAfter.Add(UE::Mass::ProcessorGroupNames::Behavior);
 }
 
-void UMassProcessor_SmartObjectTimedBehavior::Execute(UMassEntitySubsystem& EntitySubsystem, FLWComponentSystemExecutionContext& Context)
+void UMassProcessor_SmartObjectTimedBehavior::Execute(UMassEntitySubsystem& EntitySubsystem, FMassExecutionContext& Context)
 {
 	checkf(SmartObjectSubsystem != nullptr, TEXT("SmartObjectSubsystem should exist when executing processors."));
 	checkf(SignalSubsystem != nullptr, TEXT("MassSignalSubsystem should exist when executing processors."));
 
-	TArray<TTuple<FLWEntity, FDataFragment_SmartObjectUser*>> ToRelease;
+	TArray<TTuple<FMassEntityHandle, FDataFragment_SmartObjectUser*>> ToRelease;
 
 	QUICK_SCOPE_CYCLE_COUNTER(UMassProcessor_SmartObjectTestBehavior_Run);
 
-	EntityQuery.ForEachEntityChunk(EntitySubsystem, Context, [this, &ToRelease](FLWComponentSystemExecutionContext& Context)
+	EntityQuery.ForEachEntityChunk(EntitySubsystem, Context, [this, &ToRelease](FMassExecutionContext& Context)
 	{
 		const int32 NumEntities = Context.GetEntitiesNum();
 		const TArrayView<FDataFragment_SmartObjectUser> UserList = Context.GetMutableComponentView<FDataFragment_SmartObjectUser>();
@@ -306,7 +306,7 @@ void UMassProcessor_SmartObjectTimedBehavior::Execute(UMassEntitySubsystem& Enti
 			const bool bMustRelease = SOUser.GetUseTime() <= 0.f;
 
 #if WITH_MASSGAMEPLAY_DEBUG
-			const FLWEntity Entity = Context.GetEntity(i);
+			const FMassEntityHandle Entity = Context.GetEntity(i);
 			FColor DebugColor(FColor::White);
 			const bool bIsDebuggingEntity = UE::MassDebug::IsDebuggingEntity(Entity, &DebugColor);
 			if (bIsDebuggingEntity)
@@ -334,9 +334,9 @@ void UMassProcessor_SmartObjectTimedBehavior::Execute(UMassEntitySubsystem& Enti
 		}
 	});
 
-	for (const TTuple<FLWEntity, FDataFragment_SmartObjectUser*> EntryToRelease : ToRelease)
+	for (const TTuple<FMassEntityHandle, FDataFragment_SmartObjectUser*> EntryToRelease : ToRelease)
 	{
-		const FLWEntity Entity = EntryToRelease.Get<0>();
+		const FMassEntityHandle Entity = EntryToRelease.Get<0>();
 		SignalSubsystem->SignalEntity(UE::Mass::Signals::SmartObjectInteractionDone, Entity);
 
 		FDataFragment_SmartObjectUser* User = EntryToRelease.Get<1>();
diff --git a/Engine/Plugins/Runtime/MassGameplay/Source/MassSmartObjects/Public/MassSmartObjectBehaviorConfig.h b/Engine/Plugins/Runtime/MassGameplay/Source/MassSmartObjects/Public/MassSmartObjectBehaviorConfig.h
--- a/Engine/Plugins/Runtime/MassGameplay/Source/MassSmartObjects/Public/MassSmartObjectBehaviorConfig.h
+++ b/Engine/Plugins/Runtime/MassGameplay/Source/MassSmartObjects/Public/MassSmartObjectBehaviorConfig.h
@@ -7,8 +7,8 @@
 
 class UMassEntitySubsystem;
 class USmartObjectSubsystem;
-struct FLWEntity;
-struct FLWComponentSystemExecutionContext;
+struct FMassEntityHandle;
+struct FMassExecutionContext;
 struct FDataFragment_Transform;
 struct FDataFragment_SmartObjectUser;
 
@@ -19,11 +19,11 @@ struct MASSSMARTOBJECTS_API FMassBehaviorEntityContext
 {
 	FMassBehaviorEntityContext() = delete;
 
-	FMassBehaviorEntityContext(const FLWEntity& InEntity, const FDataFragment_Transform& InTransformFragment, FDataFragment_SmartObjectUser& InSOUser, USmartObjectSubsystem& InSubsystem)
+	FMassBehaviorEntityContext(const FMassEntityHandle& InEntity, const FDataFragment_Transform& InTransformFragment, FDataFragment_SmartObjectUser& InSOUser, USmartObjectSubsystem& InSubsystem)
 		: Entity(InEntity), TransformFragment(InTransformFragment), SOUser(InSOUser), Subsystem(InSubsystem)
 	{}
 
-	const FLWEntity& Entity;
+	const FMassEntityHandle& Entity;
 	const FDataFragment_Transform& TransformFragment;
 	FDataFragment_SmartObjectUser& SOUser;
 	USmartObjectSubsystem& Subsystem;
@@ -43,7 +43,7 @@ public:
 	 * This virtual method allows subclasses to configure the LW Entity based on their
 	 * parameters (e.g. Add new fragments)
 	 */
-	virtual void Activate(UMassEntitySubsystem& EntitySubsystem, FLWComponentSystemExecutionContext& Context, const FMassBehaviorEntityContext& EntityContext) const;
+	virtual void Activate(UMassEntitySubsystem& EntitySubsystem, FMassExecutionContext& Context, const FMassBehaviorEntityContext& EntityContext) const;
 
 	/**
 	 * Indicates the amount of time the LW entity
diff --git a/Engine/Plugins/Runtime/MassGameplay/Source/MassSmartObjects/Public/MassSmartObjectHandler.h b/Engine/Plugins/Runtime/MassGameplay/Source/MassSmartObjects/Public/MassSmartObjectHandler.h
--- a/Engine/Plugins/Runtime/MassGameplay/Source/MassSmartObjects/Public/MassSmartObjectHandler.h
+++ b/Engine/Plugins/Runtime/MassGameplay/Source/MassSmartObjects/Public/MassSmartObjectHandler.h
@@ -7,8 +7,8 @@
 
 class UMassEntitySubsystem;
 class USmartObjectSubsystem;
-struct FLWComponentSystemExecutionContext;
-struct FLWEntity;
+struct FMassExecutionContext;
+struct FMassEntityHandle;
 struct FDataFragment_SmartObjectUser;
 struct FDataFragment_Transform;
 struct FSmartObjectClaimHandle;
@@ -28,7 +28,7 @@ struct MASSSMARTOBJECTS_API FMassSmartObjectHandler
 	 * @param InExecutionContext is the current execution context of the entity subsystem
 	 * @param InSmartObjectSubsystem is the smart object subsystem
 	 */
-	FMassSmartObjectHandler(UMassEntitySubsystem& InEntitySubsystem, FLWComponentSystemExecutionContext& InExecutionContext, USmartObjectSubsystem& InSmartObjectSubsystem)
+	FMassSmartObjectHandler(UMassEntitySubsystem& InEntitySubsystem, FMassExecutionContext& InExecutionContext, USmartObjectSubsystem& InSmartObjectSubsystem)
 		: EntitySubsystem(InEntitySubsystem)
 		, ExecutionContext(InExecutionContext)
 		, SmartObjectSubsystem(InSmartObjectSubsystem)
@@ -43,7 +43,7 @@ struct MASSSMARTOBJECTS_API FMassSmartObjectHandler
 	 * @param Location The center of the query
 	 * @return Request identifier that can be used to try claiming a result once available
 	 */
-	FMassSmartObjectRequestID FindCandidatesAsync(const FLWEntity RequestingEntity, const FVector& Location) const;
+	FMassSmartObjectRequestID FindCandidatesAsync(const FMassEntityHandle RequestingEntity, const FVector& Location) const;
 
 	/**
 	 * Creates an async request to build a list of compatible smart objects
@@ -53,7 +53,7 @@ struct MASSSMARTOBJECTS_API FMassSmartObjectHandler
 	 * @param LaneLocation The lane location as reference for the query
 	 * @return Request identifier that can be used to try claiming a result once available
 	 */
-	FMassSmartObjectRequestID FindCandidatesAsync(const FLWEntity RequestingEntity, const FZoneGraphCompactLaneLocation& LaneLocation) const;
+	FMassSmartObjectRequestID FindCandidatesAsync(const FMassEntityHandle RequestingEntity, const FZoneGraphCompactLaneLocation& LaneLocation) const;
 
 	/**
 	 * Provides the result of a previously created request from FindCandidatesAsync to indicate if it has been processed
@@ -79,7 +79,7 @@ struct MASSSMARTOBJECTS_API FMassSmartObjectHandler
 	 * @param RequestID A valid request identifier (method will ensure otherwise)
 	 * @return Whether the slot has been successfully claimed or not
 	 */
-	EMassSmartObjectClaimResult ClaimCandidate(const FLWEntity Entity, FDataFragment_SmartObjectUser& User, const FMassSmartObjectRequestID& RequestID) const;
+	EMassSmartObjectClaimResult ClaimCandidate(const FMassEntityHandle Entity, FDataFragment_SmartObjectUser& User, const FMassSmartObjectRequestID& RequestID) const;
 
 	/**
 	 * Claims the first available smart object from the provided results.
@@ -88,7 +88,7 @@ struct MASSSMARTOBJECTS_API FMassSmartObjectHandler
 	 * @param SearchRequestResult Results of completed search request
 	 * @return Whether the slot has been successfully claimed or not
 	 */
-	EMassSmartObjectClaimResult ClaimCandidate(const FLWEntity Entity, FDataFragment_SmartObjectUser& User, const FMassSmartObjectRequestResult& SearchRequestResult) const;
+	EMassSmartObjectClaimResult ClaimCandidate(const FMassEntityHandle Entity, FDataFragment_SmartObjectUser& User, const FMassSmartObjectRequestResult& SearchRequestResult) const;
 
 	/**
 	 * Claims the first available slot holding any type of USmartObjectMassBehaviorConfig in the smart object
@@ -98,7 +98,7 @@ struct MASSSMARTOBJECTS_API FMassSmartObjectHandler
 	 * @param ObjectID A valid smart object identifier (method will ensure otherwise)
 	 * @return Whether the slot has been successfully claimed or not
 	 */
-	bool ClaimSmartObject(const FLWEntity Entity, FDataFragment_SmartObjectUser& User, const FSmartObjectID& ObjectID) const;
+	bool ClaimSmartObject(const FMassEntityHandle Entity, FDataFragment_SmartObjectUser& User, const FSmartObjectID& ObjectID) const;
 
 	/**
 	 * Activates the mass gameplay behavior associated to the previously claimed smart object.
@@ -107,7 +107,7 @@ struct MASSSMARTOBJECTS_API FMassSmartObjectHandler
 	 * @param Transform Fragment holding the transform of the user claiming
 	 * @return Whether the slot has been successfully claimed or not
 	 */
-	bool UseSmartObject(const FLWEntity Entity, FDataFragment_SmartObjectUser& User, const FDataFragment_Transform& Transform) const;
+	bool UseSmartObject(const FMassEntityHandle Entity, FDataFragment_SmartObjectUser& User, const FDataFragment_Transform& Transform) const;
 
 	/**
 	 * Releases a claimed/in-use smart object and update user fragment.
@@ -115,10 +115,10 @@ struct MASSSMARTOBJECTS_API FMassSmartObjectHandler
 	 * @param User Fragment of the user claiming
 	 * @param Status The new status for in-progress interaction
 	 */
-	void ReleaseSmartObject(const FLWEntity Entity, FDataFragment_SmartObjectUser& User, const EMassSmartObjectInteractionStatus Status = EMassSmartObjectInteractionStatus::Unset) const;
+	void ReleaseSmartObject(const FMassEntityHandle Entity, FDataFragment_SmartObjectUser& User, const EMassSmartObjectInteractionStatus Status = EMassSmartObjectInteractionStatus::Unset) const;
 
 private:
 	UMassEntitySubsystem& EntitySubsystem;
-	FLWComponentSystemExecutionContext& ExecutionContext;
+	FMassExecutionContext& ExecutionContext;
 	USmartObjectSubsystem& SmartObjectSubsystem;
 };
diff --git a/Engine/Plugins/Runtime/MassGameplay/Source/MassSmartObjects/Public/MassSmartObjectProcessor.h b/Engine/Plugins/Runtime/MassGameplay/Source/MassSmartObjects/Public/MassSmartObjectProcessor.h
--- a/Engine/Plugins/Runtime/MassGameplay/Source/MassSmartObjects/Public/MassSmartObjectProcessor.h
+++ b/Engine/Plugins/Runtime/MassGameplay/Source/MassSmartObjects/Public/MassSmartObjectProcessor.h
@@ -15,7 +15,7 @@ class UZoneGraphSubsystem;
 class UZoneGraphAnnotationSubsystem;
 
 USTRUCT()
-struct MASSSMARTOBJECTS_API FDataFragment_SmartObjectUser : public FLWComponentData
+struct MASSSMARTOBJECTS_API FDataFragment_SmartObjectUser : public FMassFragment
 {
 	GENERATED_BODY()
 
@@ -51,7 +51,7 @@ private:
  * Base class for smart object processors that takes care of caching the smart object manager.
  */
 UCLASS(Abstract)
-class MASSSMARTOBJECTS_API UMassProcessor_SmartObjectBase : public UPipeProcessor
+class MASSSMARTOBJECTS_API UMassProcessor_SmartObjectBase : public UMassProcessor
 {
 	GENERATED_BODY()
 
@@ -77,7 +77,7 @@ public:
 protected:
 	virtual void Initialize(UObject& Owner) override;
 	virtual void ConfigureQueries() override;
-	virtual void Execute(UMassEntitySubsystem& EntitySubsystem, FLWComponentSystemExecutionContext& Context) override;
+	virtual void Execute(UMassEntitySubsystem& EntitySubsystem, FMassExecutionContext& Context) override;
 
 	/** Extents used to perform the spatial query in the octree for world location queries. */
 	UPROPERTY(EditDefaultsOnly, Category = SmartObject, config)
@@ -90,14 +90,14 @@ protected:
 	UZoneGraphAnnotationSubsystem* AnnotationSubsystem;
 
 	/** Query to fetch and process requests to find smart objects using spacial query around a given world location. */
-	FLWComponentQuery WorldRequestQuery;
+	FMassEntityQuery WorldRequestQuery;
 
 	/** Query to fetch and process requests to find smart objects on zone graph lanes. */
-	FLWComponentQuery LaneRequestQuery;
+	FMassEntityQuery LaneRequestQuery;
 };
 
 USTRUCT()
-struct MASSSMARTOBJECTS_API FMassSmartObjectTimedBehaviorTag : public FComponentTag
+struct MASSSMARTOBJECTS_API FMassSmartObjectTimedBehaviorTag : public FMassTag
 {
 	GENERATED_BODY()
 };
@@ -112,8 +112,8 @@ public:
 	UMassProcessor_SmartObjectTimedBehavior();
 
 protected:
-	virtual void Execute(UMassEntitySubsystem& EntitySubsystem, FLWComponentSystemExecutionContext& Context) override;
+	virtual void Execute(UMassEntitySubsystem& EntitySubsystem, FMassExecutionContext& Context) override;
 	virtual void ConfigureQueries() override;
 
-	FLWComponentQuery EntityQuery;
+	FMassEntityQuery EntityQuery;
 };
\ No newline at end of file
diff --git a/Engine/Plugins/Runtime/MassGameplay/Source/MassSmartObjects/Public/MassSmartObjectRequest.h b/Engine/Plugins/Runtime/MassGameplay/Source/MassSmartObjects/Public/MassSmartObjectRequest.h
--- a/Engine/Plugins/Runtime/MassGameplay/Source/MassSmartObjects/Public/MassSmartObjectRequest.h
+++ b/Engine/Plugins/Runtime/MassGameplay/Source/MassSmartObjects/Public/MassSmartObjectRequest.h
@@ -48,7 +48,7 @@ struct MASSSMARTOBJECTS_API FMassSmartObjectRequestResult
  * Fragment that holds the result of a request to find candidates.
  */
 USTRUCT()
-struct MASSSMARTOBJECTS_API FMassSmartObjectRequestResultFragment : public FLWComponentData
+struct MASSSMARTOBJECTS_API FMassSmartObjectRequestResultFragment : public FMassFragment
 {
 	GENERATED_BODY()
 
@@ -62,7 +62,7 @@ struct MASSSMARTOBJECTS_API FMassSmartObjectRequestResultFragment : public FLWCo
  * fragment that could then be processed by the reservation processor
  */
 USTRUCT()
-struct MASSSMARTOBJECTS_API FMassSmartObjectWorldLocationRequestFragment : public FLWComponentData
+struct MASSSMARTOBJECTS_API FMassSmartObjectWorldLocationRequestFragment : public FMassFragment
 {
 	GENERATED_BODY()
 
@@ -70,7 +70,7 @@ struct MASSSMARTOBJECTS_API FMassSmartObjectWorldLocationRequestFragment : publi
 	FVector SearchOrigin = FVector::ZeroVector;
 
 	UPROPERTY(Transient)
-	FLWEntity RequestingEntity;
+	FMassEntityHandle RequestingEntity;
 };
 
 /**
@@ -79,21 +79,21 @@ struct MASSSMARTOBJECTS_API FMassSmartObjectWorldLocationRequestFragment : publi
  * fragment that could then be processed by the reservation processor
  */
 USTRUCT()
-struct MASSSMARTOBJECTS_API FMassSmartObjectLaneLocationRequestFragment : public FLWComponentData
+struct MASSSMARTOBJECTS_API FMassSmartObjectLaneLocationRequestFragment : public FMassFragment
 {
 	GENERATED_BODY()
 
 	FZoneGraphCompactLaneLocation CompactLaneLocation;
 
 	UPROPERTY(Transient)
-	FLWEntity RequestingEntity;
+	FMassEntityHandle RequestingEntity;
 };
 
 /**
  * Special tag to mark processed requests
  */
 USTRUCT()
-struct MASSSMARTOBJECTS_API FMassSmartObjectCompletedRequestTag : public FComponentTag
+struct MASSSMARTOBJECTS_API FMassSmartObjectCompletedRequestTag : public FMassTag
 {
 	GENERATED_BODY()
 };
diff --git a/Engine/Plugins/Runtime/MassGameplay/Source/MassSmartObjects/Public/MassSmartObjectTypes.h b/Engine/Plugins/Runtime/MassGameplay/Source/MassSmartObjects/Public/MassSmartObjectTypes.h
--- a/Engine/Plugins/Runtime/MassGameplay/Source/MassSmartObjects/Public/MassSmartObjectTypes.h
+++ b/Engine/Plugins/Runtime/MassGameplay/Source/MassSmartObjects/Public/MassSmartObjectTypes.h
@@ -22,16 +22,16 @@ struct MASSSMARTOBJECTS_API FMassSmartObjectRequestID
 	GENERATED_BODY()
 
 	FMassSmartObjectRequestID() = default;
-	FMassSmartObjectRequestID(const FLWEntity& InEntity) : Entity(InEntity) {}
+	FMassSmartObjectRequestID(const FMassEntityHandle& InEntity) : Entity(InEntity) {}
 	
 	bool IsSet() const { return Entity.IsSet(); }
 	void Reset() { Entity.Reset(); }
 	
-	explicit operator FLWEntity() const { return Entity; }
+	explicit operator FMassEntityHandle() const { return Entity; }
 
 private:
 	UPROPERTY(Transient)
-	FLWEntity Entity;
+	FMassEntityHandle Entity;
 };
 
 UENUM()
diff --git a/Engine/Plugins/Runtime/MassGameplay/Source/MassSpawner/Private/MassAssortedFragmentsTrait.cpp b/Engine/Plugins/Runtime/MassGameplay/Source/MassSpawner/Private/MassAssortedFragmentsTrait.cpp
--- a/Engine/Plugins/Runtime/MassGameplay/Source/MassSpawner/Private/MassAssortedFragmentsTrait.cpp
+++ b/Engine/Plugins/Runtime/MassGameplay/Source/MassSpawner/Private/MassAssortedFragmentsTrait.cpp
@@ -12,11 +12,11 @@ void UMassAssortedFragmentsTrait::BuildTemplate(FMassEntityTemplateBuildContext&
 		{
 			const UScriptStruct* Type = Fragment.GetScriptStruct();
 			CA_ASSUME(Type);
-			if (Type->IsChildOf(FLWComponentData::StaticStruct()))
+			if (Type->IsChildOf(FMassFragment::StaticStruct()))
 			{
 				BuildContext.AddFragmentWithDefaultInitializer(Fragment);
 			}
-			else if (Type->IsChildOf(FComponentTag::StaticStruct()))
+			else if (Type->IsChildOf(FMassTag::StaticStruct()))
 			{
 				BuildContext.AddTag(*Type);
 			}
diff --git a/Engine/Plugins/Runtime/MassGameplay/Source/MassSpawner/Private/MassEntityTemplateRegistry.cpp b/Engine/Plugins/Runtime/MassGameplay/Source/MassSpawner/Private/MassEntityTemplateRegistry.cpp
--- a/Engine/Plugins/Runtime/MassGameplay/Source/MassSpawner/Private/MassEntityTemplateRegistry.cpp
+++ b/Engine/Plugins/Runtime/MassGameplay/Source/MassSpawner/Private/MassEntityTemplateRegistry.cpp
@@ -203,7 +203,7 @@ const FMassEntityTemplate* UMassEntityTemplateRegistry::BuildClassTemplate(const
 	FClassToTemplateBuilderDelegate* Builder = GetBuilderForClass(Class);
 	if (Builder)
 	{
-		TArray<const UPipeProcessor*> Initializers;
+		TArray<const UMassProcessor*> Initializers;
 		
 		// TODO consider removing the need for strings here
 		// Use the class name string for the hash here, so the hash can be deterministic between client and server
@@ -256,7 +256,7 @@ void UMassEntityTemplateRegistry::InitializeEntityTemplate(FMassEntityTemplate&
 	check(EntitySys);
 
 	FTemplateRegistryHelpers::FragmentInstancesToTypes(OutTemplate.GetFragments(), OutTemplate.GetMutableFragmentTypes());
-	const FArchetypeHandle ArchetypeHandle = EntitySys->CreateArchetype(OutTemplate.GetFragmentTypes(), OutTemplate.GetTags(), FLWChunkComponentBitSet(OutTemplate.GetChunkFragments()));
+	const FArchetypeHandle ArchetypeHandle = EntitySys->CreateArchetype(OutTemplate.GetFragmentTypes(), OutTemplate.GetTags(), FMassChunkFragmentBitSet(OutTemplate.GetChunkFragments()));
 	OutTemplate.SetArchetype(ArchetypeHandle);
 	for (const FInstancedStruct& InstancedStruct : OutTemplate.GetChunkFragments())
 	{
diff --git a/Engine/Plugins/Runtime/MassGameplay/Source/MassSpawner/Private/MassSpawner.cpp b/Engine/Plugins/Runtime/MassGameplay/Source/MassSpawner/Private/MassSpawner.cpp
--- a/Engine/Plugins/Runtime/MassGameplay/Source/MassSpawner/Private/MassSpawner.cpp
+++ b/Engine/Plugins/Runtime/MassGameplay/Source/MassSpawner/Private/MassSpawner.cpp
@@ -491,7 +491,7 @@ void AMassSpawner::DoDespawning()
 	AllSpawnedEntities.Reset();
 }
 
-bool AMassSpawner::DespawnEntity(const FLWEntity Entity)
+bool AMassSpawner::DespawnEntity(const FMassEntityHandle Entity)
 {
 	UMassSpawnerSubsystem* SpawnerSystem = UWorld::GetSubsystem<UMassSpawnerSubsystem>(GetWorld());
 	if (SpawnerSystem == nullptr)
@@ -535,7 +535,7 @@ void AMassSpawner::UnloadConfig()
 	}
 }
 
-void AMassSpawner::OnAdjustTickSchematics(UWorld* World,TArray<TSoftObjectPtr<UPipeSchematic>>& InOutTickSchematics)
+void AMassSpawner::OnAdjustTickSchematics(UWorld* World,TArray<TSoftObjectPtr<UMassSchematic>>& InOutTickSchematics)
 {
 	// this function gets called from a static multicast delegate, so we need to differentiate based on the World it's being called for.
 	if (World != GetWorld())
diff --git a/Engine/Plugins/Runtime/MassGameplay/Source/MassSpawner/Private/MassSpawnerSubsystem.cpp b/Engine/Plugins/Runtime/MassGameplay/Source/MassSpawner/Private/MassSpawnerSubsystem.cpp
--- a/Engine/Plugins/Runtime/MassGameplay/Source/MassSpawner/Private/MassSpawnerSubsystem.cpp
+++ b/Engine/Plugins/Runtime/MassGameplay/Source/MassSpawner/Private/MassSpawnerSubsystem.cpp
@@ -31,7 +31,7 @@ namespace UE::Mass::SpawnerSubsystem
 
 		// @Todo: should queue up the entity destruction and apply in batch
 		const UMassTranslatorRegistry& Registry = UMassTranslatorRegistry::Get();
-		TArray<const UPipeProcessor*> Destructors;
+		TArray<const UMassProcessor*> Destructors;
 
 		EntitySystem.ForEachArchetypeComponentType(ArchetypeHandle, [&Registry, &Destructors](const UScriptStruct* FragmentType)
 		{
@@ -45,10 +45,10 @@ namespace UE::Mass::SpawnerSubsystem
 
 		if (Destructors.Num())
 		{
-			FRuntimePipeline DestructionPipeline;
+			FMassRuntimePipeline DestructionPipeline;
 			DestructionPipeline.InitializeFromArray(Destructors, SpawnerSubSystem);
 
-			FPipeContext PipeContext(EntitySystem, /*TimeDelta=*/0.0f);
+			FMassProcessingContext PipeContext(EntitySystem, /*TimeDelta=*/0.0f);
 
 			UE::Pipe::Executor::RunSparse(DestructionPipeline, PipeContext, Chunks);
 		}
@@ -56,16 +56,16 @@ namespace UE::Mass::SpawnerSubsystem
 		return Destructors.Num() > 0;
 	}
 	
-	void CreateSparseChunks(const UMassEntitySubsystem& EntitySystem, const TConstArrayView<FLWEntity> Entities, TArray<FArchetypeChunkCollection>& OutChunkCollections)
+	void CreateSparseChunks(const UMassEntitySubsystem& EntitySystem, const TConstArrayView<FMassEntityHandle> Entities, TArray<FArchetypeChunkCollection>& OutChunkCollections)
 	{
 		TRACE_CPUPROFILER_EVENT_SCOPE_STR("SpawnerSubsystem_CreateSparseChunks");
 
-		TMap<const FArchetypeHandle, TArray<FLWEntity>> ArchetypeToEntities;
+		TMap<const FArchetypeHandle, TArray<FMassEntityHandle>> ArchetypeToEntities;
 
-		for (const FLWEntity& Entity : Entities)
+		for (const FMassEntityHandle& Entity : Entities)
 		{
 			FArchetypeHandle Archetype = EntitySystem.GetArchetypeForEntity(Entity);
-			TArray<FLWEntity>& PerArchetypeEntities = ArchetypeToEntities.FindOrAdd(Archetype);
+			TArray<FMassEntityHandle>& PerArchetypeEntities = ArchetypeToEntities.FindOrAdd(Archetype);
 			PerArchetypeEntities.Add(Entity);
 		}
 
@@ -102,7 +102,7 @@ void UMassSpawnerSubsystem::PostInitialize()
 	TemplateRegistryInstance = NewObject<UMassEntityTemplateRegistry>(this);
 }
 
-bool UMassSpawnerSubsystem::SpawnEntities(const AActor& ActorInstance, const uint32 NumberToSpawn, TArray<FLWEntity>& OutEntities) const
+bool UMassSpawnerSubsystem::SpawnEntities(const AActor& ActorInstance, const uint32 NumberToSpawn, TArray<FMassEntityHandle>& OutEntities) const
 {
 	if (NumberToSpawn == 0)
 	{
@@ -118,7 +118,7 @@ bool UMassSpawnerSubsystem::SpawnEntities(const AActor& ActorInstance, const uin
 	return false;
 }
 
-bool UMassSpawnerSubsystem::SpawnEntities(const FMassEntityTemplate& EntityTemplate, const uint32 NumberToSpawn, TArray<FLWEntity>& OutEntities) const
+bool UMassSpawnerSubsystem::SpawnEntities(const FMassEntityTemplate& EntityTemplate, const uint32 NumberToSpawn, TArray<FMassEntityHandle>& OutEntities) const
 {
 	check(EntitySystem);
 	check(EntityTemplate.IsValid());
@@ -144,7 +144,7 @@ bool UMassSpawnerSubsystem::SpawnEntities(const FMassEntityTemplate& EntityTempl
 		{
 			FMassEntityTemplate& MutableEntityTemplate = const_cast<FMassEntityTemplate&>(EntityTemplate);
 
-			FPipeContext PipeContext(*EntitySystem, /*TimeDelta=*/0.0f);
+			FMassProcessingContext PipeContext(*EntitySystem, /*TimeDelta=*/0.0f);
 			UE::Pipe::Executor::RunSparse(MutableEntityTemplate.GetMutableInitializationPipeline(), PipeContext, ChunkCollection);
 		}
 		return true;
@@ -152,7 +152,7 @@ bool UMassSpawnerSubsystem::SpawnEntities(const FMassEntityTemplate& EntityTempl
 	return false;
 }
 
-void UMassSpawnerSubsystem::SpawnEntities(FMassEntityTemplateID TemplateID, const FMassSpawnConfigBase& SpawnConfig, const FStructView& AuxData, TArray<FLWEntity>& OutEntities)
+void UMassSpawnerSubsystem::SpawnEntities(FMassEntityTemplateID TemplateID, const FMassSpawnConfigBase& SpawnConfig, const FStructView& AuxData, TArray<FMassEntityHandle>& OutEntities)
 {
 	check(TemplateID.IsValid());
 
@@ -185,7 +185,7 @@ void UMassSpawnerSubsystem::SpawnCollection(TArrayView<FInstancedStruct> Collect
 			const FMassEntityTemplate* EntityTemplate = TemplateRegistryInstance->FindOrBuildStructTemplate(Entry);
 			if (EntityTemplate && EntityTemplate->IsValid())
 			{
-				TArray<FLWEntity> Entities;
+				TArray<FMassEntityHandle> Entities;
 
 				DoSpawning(*EntityTemplate, Entry.GetMutable<FMassSpawnConfigBase>(), AuxData, Entities);
 			}
@@ -209,7 +209,7 @@ void UMassSpawnerSubsystem::RegisterCollection(TArrayView<FInstancedStruct> Coll
 	}
 }
 
-void UMassSpawnerSubsystem::DestroyEntities(const FMassEntityTemplateID TemplateID, TConstArrayView<FLWEntity> Entities)
+void UMassSpawnerSubsystem::DestroyEntities(const FMassEntityTemplateID TemplateID, TConstArrayView<FMassEntityHandle> Entities)
 {
 	TRACE_CPUPROFILER_EVENT_SCOPE_STR("MassSpawnerSubsystem_DestroyEntities")
 
@@ -235,7 +235,7 @@ void UMassSpawnerSubsystem::DestroyEntities(const FMassEntityTemplateID Template
 		TArray<FArchetypeChunkCollection> ChunkCollections;
 		UE::Mass::SpawnerSubsystem::CreateSparseChunks(*EntitySystem, Entities, ChunkCollections);
 
-		FPipeContext PipeContext(*EntitySystem, /*TimeDelta=*/0.0f);
+		FMassProcessingContext PipeContext(*EntitySystem, /*TimeDelta=*/0.0f);
 
 		// Run destructors
 		bool bDestructorsRun = false;
@@ -284,7 +284,7 @@ void UMassSpawnerSubsystem::DestroyEntities(const FMassEntityTemplateID Template
 	}
 }
 
-void UMassSpawnerSubsystem::DoSpawning(const FMassEntityTemplate& EntityTemplate, const FMassSpawnConfigBase& Data, const FStructView& AuxData, TArray<FLWEntity>& OutEntities) const
+void UMassSpawnerSubsystem::DoSpawning(const FMassEntityTemplate& EntityTemplate, const FMassSpawnConfigBase& Data, const FStructView& AuxData, TArray<FMassEntityHandle>& OutEntities) const
 {
 	check(EntitySystem);
 	check(EntityTemplate.GetArchetype().IsValid());
@@ -298,7 +298,7 @@ void UMassSpawnerSubsystem::DoSpawning(const FMassEntityTemplate& EntityTemplate
 	//		b. @todo could be done via an always-first translator in EntityTemplate.Translators
 	// 3. Run all EntityTemplate.Translators to configure remaining values (i.e. not set via FMassEntityTemplate.Fragments).
 
-	TArray<FLWEntity> SpawnedEntities;
+	TArray<FMassEntityHandle> SpawnedEntities;
 	EntitySystem->BatchCreateEntities(EntityTemplate.GetArchetype(), Data.MaxNumber, SpawnedEntities);
 
 	const TArrayView<const FInstancedStruct> FragmentInstances = EntityTemplate.GetFragments();
@@ -309,7 +309,7 @@ void UMassSpawnerSubsystem::DoSpawning(const FMassEntityTemplate& EntityTemplate
 	{
 		FMassEntityTemplate& MutableEntityTemplate = const_cast<FMassEntityTemplate&>(EntityTemplate);
 
-		FPipeContext PipeContext(*EntitySystem, /*TimeDelta=*/0.0f);
+		FMassProcessingContext PipeContext(*EntitySystem, /*TimeDelta=*/0.0f);
 		PipeContext.AuxData = AuxData;
 		UE::Pipe::Executor::RunSparse(MutableEntityTemplate.GetMutableInitializationPipeline(), PipeContext, ChunkCollection);
 	}
diff --git a/Engine/Plugins/Runtime/MassGameplay/Source/MassSpawner/Private/MassTranslator.cpp b/Engine/Plugins/Runtime/MassGameplay/Source/MassSpawner/Private/MassTranslator.cpp
--- a/Engine/Plugins/Runtime/MassGameplay/Source/MassSpawner/Private/MassTranslator.cpp
+++ b/Engine/Plugins/Runtime/MassGameplay/Source/MassSpawner/Private/MassTranslator.cpp
@@ -13,9 +13,9 @@ UMassTranslator::UMassTranslator()
 	ExecutionFlags = (int32)EProcessorExecutionFlags::All;
 }
 
-void UMassTranslator::AddRequiredTagsToQuery(FLWComponentQuery& EntityQuery)
+void UMassTranslator::AddRequiredTagsToQuery(FMassEntityQuery& EntityQuery)
 {
-	EntityQuery.AddTagRequirements<ELWComponentPresence::All>(RequiredTags);
+	EntityQuery.AddTagRequirements<EMassFragmentPresence::All>(RequiredTags);
 }
 
 //----------------------------------------------------------------------//
diff --git a/Engine/Plugins/Runtime/MassGameplay/Source/MassSpawner/Public/MassAssortedFragmentsTrait.h b/Engine/Plugins/Runtime/MassGameplay/Source/MassSpawner/Public/MassAssortedFragmentsTrait.h
--- a/Engine/Plugins/Runtime/MassGameplay/Source/MassSpawner/Public/MassAssortedFragmentsTrait.h
+++ b/Engine/Plugins/Runtime/MassGameplay/Source/MassSpawner/Public/MassAssortedFragmentsTrait.h
@@ -19,6 +19,6 @@ protected:
 
 	virtual void BuildTemplate(FMassEntityTemplateBuildContext& BuildContext, UWorld& World) const override;
 
-	UPROPERTY(Category="Fragments", EditAnywhere, meta = (BaseStruct = "LWComponentData", ExcludeBaseStruct))
+	UPROPERTY(Category="Fragments", EditAnywhere, meta = (BaseStruct = "MassFragment", ExcludeBaseStruct))
 	TArray<FInstancedStruct> Fragments;
 };
diff --git a/Engine/Plugins/Runtime/MassGameplay/Source/MassSpawner/Public/MassEntityTemplate.h b/Engine/Plugins/Runtime/MassGameplay/Source/MassSpawner/Public/MassEntityTemplate.h
--- a/Engine/Plugins/Runtime/MassGameplay/Source/MassSpawner/Public/MassEntityTemplate.h
+++ b/Engine/Plugins/Runtime/MassGameplay/Source/MassSpawner/Public/MassEntityTemplate.h
@@ -10,7 +10,7 @@
 #include "MassEntityTemplate.generated.h"
 
 
-struct FEntityView;
+struct FMassEntityView;
 
 
 UENUM()
@@ -27,9 +27,9 @@ struct MASSSPAWNER_API FMassObjectHandlers
 {
 	GENERATED_BODY()
 	UPROPERTY()
-	TArray<const UPipeProcessor*> Initializers;
+	TArray<const UMassProcessor*> Initializers;
 	UPROPERTY()
-	TArray<const UPipeProcessor*> Deinitializers;
+	TArray<const UMassProcessor*> Deinitializers;
 };
 
 //ID of the template an entity is using
@@ -79,7 +79,7 @@ struct MASSSPAWNER_API FMassEntityTemplate
 {
 	GENERATED_BODY()
 
-	typedef TFunction<void(UObject& /*Owner*/, FEntityView& /*EntityView*/, const EMassTranslationDirection /*CurrentDirection*/)> FObjectFragmentInitializerFunction;
+	typedef TFunction<void(UObject& /*Owner*/, FMassEntityView& /*EntityView*/, const EMassTranslationDirection /*CurrentDirection*/)> FObjectFragmentInitializerFunction;
 
 	FMassEntityTemplate() = default;
 	/** InArchetype is expected to be valid. The function will crash-check it. */
@@ -88,11 +88,11 @@ struct MASSSPAWNER_API FMassEntityTemplate
 	/** Configures initialization and deinitialization pipeline from ObjectHandlers */
 	void SetUpProcessors(const FMassObjectHandlers& ObjectHandlers, UObject& PipelineOwner);
 	
-	const FRuntimePipeline& GetInitializationPipeline() const { return InitializationPipeline; }
-	FRuntimePipeline& GetMutableInitializationPipeline() { return InitializationPipeline; }
+	const FMassRuntimePipeline& GetInitializationPipeline() const { return InitializationPipeline; }
+	FMassRuntimePipeline& GetMutableInitializationPipeline() { return InitializationPipeline; }
 
-	const FRuntimePipeline& GetDeinitializationPipeline() const { return DeinitializationPipeline; }
-	FRuntimePipeline& GetMutableDeinitializationPipeline() { return DeinitializationPipeline; }
+	const FMassRuntimePipeline& GetDeinitializationPipeline() const { return DeinitializationPipeline; }
+	FMassRuntimePipeline& GetMutableDeinitializationPipeline() { return DeinitializationPipeline; }
 
 	TConstArrayView<FObjectFragmentInitializerFunction> GetObjectFragmentInitializers() const { return ObjectInitializers; }
 	TArray<FObjectFragmentInitializerFunction>& GetMutableObjectFragmentInitializers() { return ObjectInitializers; }
@@ -109,7 +109,7 @@ struct MASSSPAWNER_API FMassEntityTemplate
 	FMassUniqueFragmentCollection& GetMutableFragmentCollection() { return FragmentCollection; }
 	TConstArrayView<FInstancedStruct> GetChunkFragments() const { return ChunkFragments; }
 
-	FLWCompositionDescriptor GetCompositionDescriptor() const { return FLWCompositionDescriptor(FragmentCollection.GetFragmentBitSet(), TagBitSet, ChunkFragmentBitSet); }
+	FMassCompositionDescriptor GetCompositionDescriptor() const { return FMassCompositionDescriptor(FragmentCollection.GetFragmentBitSet(), TagBitSet, ChunkFragmentBitSet); }
 
 	template<typename T>
 	void AddTag()
@@ -122,13 +122,13 @@ struct MASSSPAWNER_API FMassEntityTemplate
 		TagBitSet.Add(TagType);
 	}
 
-	const FLWTagBitSet& GetTags() const { return TagBitSet; }	
-	FLWTagBitSet& GetMutableTags() { return TagBitSet; }
+	const FMassTagBitSet& GetTags() const { return TagBitSet; }	
+	FMassTagBitSet& GetMutableTags() { return TagBitSet; }
 
 	template<typename T>
 	void AddChunkFragment()
 	{
-		static_assert(TIsDerivedFrom<T, FLWChunkComponent>::IsDerived, "Given struct doesn't represent a valid chunk fragment type. Make sure to inherit from FLWChunkComonent or one of its child-types.");
+		static_assert(TIsDerivedFrom<T, FMassChunkFragment>::IsDerived, "Given struct doesn't represent a valid chunk fragment type. Make sure to inherit from FLWChunkComonent or one of its child-types.");
 		ChunkFragments.AddUnique(FInstancedStruct(T::StaticStruct()));
 		ChunkFragmentBitSet.Add<T>();
 	}
@@ -136,7 +136,7 @@ struct MASSSPAWNER_API FMassEntityTemplate
 	template<typename T>
 	T& AddChunkFragment_GetRef()
 	{
-		static_assert(TIsDerivedFrom<T, FLWChunkComponent>::IsDerived, "Given struct doesn't represent a valid chunk fragment type. Make sure to inherit from FLWChunkComonent or one of its child-types.");
+		static_assert(TIsDerivedFrom<T, FMassChunkFragment>::IsDerived, "Given struct doesn't represent a valid chunk fragment type. Make sure to inherit from FLWChunkComonent or one of its child-types.");
 		ChunkFragmentBitSet.Add<T>();
 		const int32 Index = ChunkFragments.AddUnique(FInstancedStruct(T::StaticStruct()));
 		return ChunkFragments[Index].template GetMutable<T>();
@@ -149,14 +149,14 @@ private:
 	FArchetypeHandle Archetype;
 
 	FMassUniqueFragmentCollection FragmentCollection;
-	FLWTagBitSet TagBitSet;
-	FLWChunkComponentBitSet ChunkFragmentBitSet;
+	FMassTagBitSet TagBitSet;
+	FMassChunkFragmentBitSet ChunkFragmentBitSet;
 	
 	UPROPERTY()
-	FRuntimePipeline InitializationPipeline;
+	FMassRuntimePipeline InitializationPipeline;
 	
 	UPROPERTY()
-	FRuntimePipeline DeinitializationPipeline;
+	FMassRuntimePipeline DeinitializationPipeline;
 
 	/** An mirror of FMassUniqueFragmentCollection.Fragments' types, used for adding fragment types instead of fragment instances */
 	UPROPERTY()
diff --git a/Engine/Plugins/Runtime/MassGameplay/Source/MassSpawner/Public/MassEntityTemplateRegistry.h b/Engine/Plugins/Runtime/MassGameplay/Source/MassSpawner/Public/MassEntityTemplateRegistry.h
--- a/Engine/Plugins/Runtime/MassGameplay/Source/MassSpawner/Public/MassEntityTemplateRegistry.h
+++ b/Engine/Plugins/Runtime/MassGameplay/Source/MassSpawner/Public/MassEntityTemplateRegistry.h
@@ -44,7 +44,7 @@ struct FMassEntityTemplateBuildContext
 	}
 
 	template<typename T>
-	T& AddFragmentWithInitializer_GetRef(const UPipeProcessor& Initializer)
+	T& AddFragmentWithInitializer_GetRef(const UMassProcessor& Initializer)
 	{
 		AddInitializer(Initializer);
 		return AddFragment_GetRef<T>();
@@ -53,22 +53,22 @@ struct FMassEntityTemplateBuildContext
 	template<typename T>
 	void AddFragment()
 	{
-		static_assert(TIsDerivedFrom<T, FLWComponentData>::IsDerived, "Given struct doesn't represent a valid fragment type. Make sure to inherit from FLWComponentData or one of its child-types.");
+		static_assert(TIsDerivedFrom<T, FMassFragment>::IsDerived, "Given struct doesn't represent a valid fragment type. Make sure to inherit from FMassFragment or one of its child-types.");
 		Template.GetMutableFragmentCollection().Add(T::StaticStruct());
 	}
 
 	template<typename T>
 	void AddFragmentWithDefaultInitializer()
 	{
-		static_assert(TIsDerivedFrom<T, FLWComponentData>::IsDerived, "Given struct doesn't represent a valid fragment type. Make sure to inherit from FLWComponentData or one of its child-types.");
+		static_assert(TIsDerivedFrom<T, FMassFragment>::IsDerived, "Given struct doesn't represent a valid fragment type. Make sure to inherit from FMassFragment or one of its child-types.");
 		AddFragment<T>();
 		AddDefaultInitializer<T>();
 	}
 
 	template<typename T>
-	void AddFragmentWithInitializer(const UPipeProcessor& Initializer)
+	void AddFragmentWithInitializer(const UMassProcessor& Initializer)
 	{
-		static_assert(TIsDerivedFrom<T, FLWComponentData>::IsDerived, "Given struct doesn't represent a valid fragment type. Make sure to inherit from FLWComponentData or one of its child-types.");
+		static_assert(TIsDerivedFrom<T, FMassFragment>::IsDerived, "Given struct doesn't represent a valid fragment type. Make sure to inherit from FMassFragment or one of its child-types.");
 		AddFragment<T>();
 		AddInitializer(Initializer);
 	}
@@ -87,7 +87,7 @@ struct FMassEntityTemplateBuildContext
 		}
 	}
 
-	void AddFragmentWithInitializer(FStructView InFragment, const UPipeProcessor& Initializer)
+	void AddFragmentWithInitializer(FStructView InFragment, const UMassProcessor& Initializer)
 	{
 		AddFragment(InFragment);
 		AddInitializer(Initializer);
@@ -119,7 +119,7 @@ struct FMassEntityTemplateBuildContext
 	//----------------------------------------------------------------------//
 	// Initializers
 	//----------------------------------------------------------------------//
-	void AddInitializer(const UPipeProcessor& Initializer) { Handlers.Initializers.AddUnique(&Initializer); }
+	void AddInitializer(const UMassProcessor& Initializer) { Handlers.Initializers.AddUnique(&Initializer); }
 	
 	void AddDefaultInitializer(const UScriptStruct& FragmentType)
 	{
@@ -143,7 +143,7 @@ struct FMassEntityTemplateBuildContext
 	//----------------------------------------------------------------------//
 	// Deinitializers
 	//----------------------------------------------------------------------//
-	void AddDeinitializer(const UPipeProcessor& Deinitializer) { Handlers.Deinitializers.AddUnique(&Deinitializer); }
+	void AddDeinitializer(const UMassProcessor& Deinitializer) { Handlers.Deinitializers.AddUnique(&Deinitializer); }
 
 	void AddDefaultDeinitializer(const UScriptStruct& FragmentType)
 	{
diff --git a/Engine/Plugins/Runtime/MassGameplay/Source/MassSpawner/Public/MassSpawner.h b/Engine/Plugins/Runtime/MassGameplay/Source/MassSpawner/Public/MassSpawner.h
--- a/Engine/Plugins/Runtime/MassGameplay/Source/MassSpawner/Public/MassSpawner.h
+++ b/Engine/Plugins/Runtime/MassGameplay/Source/MassSpawner/Public/MassSpawner.h
@@ -10,7 +10,7 @@
 #include "MassSpawner.generated.h"
 
 class UMassEntitySpawnPointsGeneratorConfigAsset;
-class UPipeSchematic;
+class UMassSchematic;
 
 /** A spawner you can put on a map and configure it to spawn different things */
 UCLASS(hidecategories = (Object, Actor, Input, Rendering, LOD, Cooking, Collision, HLOD, Partition))
@@ -71,7 +71,7 @@ public:
 	 * updated. Otherwise the call has no effect.
 	 * @return true if the entity got removed. False otherwise.
 	 */
-	bool DespawnEntity(const FLWEntity Entity);
+	bool DespawnEntity(const FMassEntityHandle Entity);
 
 	/**
 	 * Scales the spawning counts (this only works when bUseEntityConfig is set)
@@ -82,7 +82,7 @@ public:
 
 protected:
 	void OnEQSQueryFinished(TSharedPtr<FEnvQueryResult> Result);
-	void OnAdjustTickSchematics(UWorld* World, TArray<TSoftObjectPtr<UPipeSchematic>>& InOutTickSchematics);
+	void OnAdjustTickSchematics(UWorld* World, TArray<TSoftObjectPtr<UMassSchematic>>& InOutTickSchematics);
 	void OnPostWorldInit(UWorld* World, const UWorld::InitializationValues);
 	void SpawnAtLocations(const TArray<FVector>& Locations);
 	void OnSpawnPointGenerationFinished(const TArray<FVector>& Locations, FMassSpawnPointGenerator* FinishedGenerator);
@@ -94,7 +94,7 @@ protected:
 	struct FSpawnedEntities
 	{
 		FMassEntityTemplateID TemplateID;
-		TArray<FLWEntity> Entities;
+		TArray<FMassEntityHandle> Entities;
 	};
 
 	// @TODO CONSIDER HAVING ONE PER ENTRY IN SpawnSets
@@ -129,7 +129,7 @@ protected:
 	uint32 bOverrideSchematics : 1;
 
 	UPROPERTY(Category = "Mass|Simulation", EditAnywhere)
-	TArray<TSoftObjectPtr<UPipeSchematic>> TickSchematics;
+	TArray<TSoftObjectPtr<UMassSchematic>> TickSchematics;
 
 	/** Scale of the spawning count */
 	float SpawningCountScale = 1.0f;
diff --git a/Engine/Plugins/Runtime/MassGameplay/Source/MassSpawner/Public/MassSpawnerSubsystem.h b/Engine/Plugins/Runtime/MassGameplay/Source/MassSpawner/Public/MassSpawnerSubsystem.h
--- a/Engine/Plugins/Runtime/MassGameplay/Source/MassSpawner/Public/MassSpawnerSubsystem.h
+++ b/Engine/Plugins/Runtime/MassGameplay/Source/MassSpawner/Public/MassSpawnerSubsystem.h
@@ -21,7 +21,7 @@ class MASSSPAWNER_API UMassSpawnerSubsystem : public UWorldSubsystem
 	GENERATED_BODY()
 public:
 
-	void SpawnEntities(FMassEntityTemplateID TemplateID, const FMassSpawnConfigBase& SpawnConfig, const FStructView& AuxData, TArray<FLWEntity>& OutEntities);
+	void SpawnEntities(FMassEntityTemplateID TemplateID, const FMassSpawnConfigBase& SpawnConfig, const FStructView& AuxData, TArray<FMassEntityHandle>& OutEntities);
 
 	/** Finds the appropriate entity template for ActorInstance and calls the other SpawnEntities implementation 
 	 *  @return true if spawning was successful, false otherwise. The failure might come from there being no EntityTemplate 
@@ -29,7 +29,7 @@ public:
 	 *  @param NumberToSpawn number of entities to spawn
 	 *  @param OutEntities where the IDs of created entities get added. Note that the contents of OutEntities get overridden by the function.
 	 *  for ActorInstance. See logs for more details. */
-	bool SpawnEntities(const AActor& ActorInstance, const uint32 NumberToSpawn, TArray<FLWEntity>& OutEntities) const;
+	bool SpawnEntities(const AActor& ActorInstance, const uint32 NumberToSpawn, TArray<FMassEntityHandle>& OutEntities) const;
 
 	/** Spawns entities of the kind described by the given EntityTemplate. The spawned entities are fully initialized
 	 *  meaning the EntityTemplate.InitializationPipeline gets run for all spawned entities.
@@ -37,11 +37,11 @@ public:
 	 *  @param NumberToSpawn number of entities to spawn
 	 *  @param OutEntities where the IDs of created entities get added. Note that the contents of OutEntities get overridden by the function.
 	 *  @return true if spawning was successful, false otherwise. In case of failure see logs for more details. */
-	bool SpawnEntities(const FMassEntityTemplate& EntityTemplate, const uint32 NumberToSpawn, TArray<FLWEntity>& OutEntities) const;
+	bool SpawnEntities(const FMassEntityTemplate& EntityTemplate, const uint32 NumberToSpawn, TArray<FMassEntityHandle>& OutEntities) const;
 
 	void SpawnCollection(TArrayView<FInstancedStruct> Collection, const int32 Count, const FStructView& AuxData = FStructView());
 
-	void DestroyEntities(const FMassEntityTemplateID TemplateID, TConstArrayView<FLWEntity> Entities);
+	void DestroyEntities(const FMassEntityTemplateID TemplateID, TConstArrayView<FMassEntityHandle> Entities);
 
 	UMassEntityTemplateRegistry& GetTemplateRegistryInstance() const { check(TemplateRegistryInstance); return *TemplateRegistryInstance; }
 
@@ -53,7 +53,7 @@ protected:
 	virtual void Initialize(FSubsystemCollectionBase& Collection) override;
 	virtual void PostInitialize() override;
 
-	void DoSpawning(const FMassEntityTemplate& EntityTemplate, const FMassSpawnConfigBase& Data, const FStructView& AuxData, TArray<FLWEntity>& OutEntities) const;
+	void DoSpawning(const FMassEntityTemplate& EntityTemplate, const FMassSpawnConfigBase& Data, const FStructView& AuxData, TArray<FMassEntityHandle>& OutEntities) const;
 
 	UPROPERTY()
 	UMassEntitySubsystem* EntitySystem;
diff --git a/Engine/Plugins/Runtime/MassGameplay/Source/MassSpawner/Public/MassSpawnerTypes.h b/Engine/Plugins/Runtime/MassGameplay/Source/MassSpawner/Public/MassSpawnerTypes.h
--- a/Engine/Plugins/Runtime/MassGameplay/Source/MassSpawner/Public/MassSpawnerTypes.h
+++ b/Engine/Plugins/Runtime/MassGameplay/Source/MassSpawner/Public/MassSpawnerTypes.h
@@ -31,7 +31,7 @@ struct FMassSpawnConfigBase
 {
 	GENERATED_BODY()
 
-	UPROPERTY(EditAnywhere, Category = "Mass|Spawn", meta = (BaseStruct = "LWComponentData"))
+	UPROPERTY(EditAnywhere, Category = "Mass|Spawn", meta = (BaseStruct = "MassFragment"))
 	TArray<FInstancedStruct> AdditionalDataFragments;
 
 	UPROPERTY(EditAnywhere, Category = "Mass|Spawn", meta = (ClampMin = 0, UIMin = 0))
@@ -97,7 +97,7 @@ struct FMassSpawnPointGenerator
 };
 
 USTRUCT()
-struct MASSSPAWNER_API FDataFragment_ReplicationTemplateID : public FLWComponentData
+struct MASSSPAWNER_API FDataFragment_ReplicationTemplateID : public FMassFragment
 {
 	GENERATED_BODY()
 
diff --git a/Engine/Plugins/Runtime/MassGameplay/Source/MassSpawner/Public/MassTranslator.h b/Engine/Plugins/Runtime/MassGameplay/Source/MassSpawner/Public/MassTranslator.h
--- a/Engine/Plugins/Runtime/MassGameplay/Source/MassSpawner/Public/MassTranslator.h
+++ b/Engine/Plugins/Runtime/MassGameplay/Source/MassSpawner/Public/MassTranslator.h
@@ -27,32 +27,32 @@ ENUM_CLASS_FLAGS(EMassTranslationDirection);
  *  copy values from UObjects to fragments and back.
  */
 UCLASS(abstract)
-class MASSSPAWNER_API UMassTranslator : public UPipeProcessor
+class MASSSPAWNER_API UMassTranslator : public UMassProcessor
 {
 	GENERATED_BODY()
 
 protected:
 	UMassTranslator();
 public:
-	/** Fetches the FComponentTag-derived types required by this Translator. And entity needs these tags to be 
+	/** Fetches the FMassTag-derived types required by this Translator. And entity needs these tags to be 
 	 *  processed by this Translator instance. 
 	 * 
-	 *  @todo might want this function on the PipeProcessor level. TBD
+	 *  @todo might want this function on the MassProcessor level. TBD
 	 * 
 	 *  @param OutTagTypes tag types will be appended to this array. No uniqueness checks are performed. */
-	void AppendRequiredTags(FLWTagBitSet& InOutTags) const { InOutTags += RequiredTags; }
+	void AppendRequiredTags(FMassTagBitSet& InOutTags) const { InOutTags += RequiredTags; }
 
 protected:
-	void AddRequiredTagsToQuery(FLWComponentQuery& EntityQuery);
+	void AddRequiredTagsToQuery(FMassEntityQuery& EntityQuery);
 
 protected:
 	/** These are the tag fragments expected by this translator that let other code (like entity traits) hint to 
 	 *  the system which translators they'd want their entity to be processed by. */
-	FLWTagBitSet RequiredTags;
+	FMassTagBitSet RequiredTags;
 };
 
 UCLASS(abstract)
-class MASSSPAWNER_API UMassFragmentBuilder : public UPipeProcessor
+class MASSSPAWNER_API UMassFragmentBuilder : public UMassProcessor
 {
 	GENERATED_BODY()
 
diff --git a/Engine/Plugins/Runtime/MassGameplay/Source/MassSpawner/Public/MassTranslatorRegistry.h b/Engine/Plugins/Runtime/MassGameplay/Source/MassSpawner/Public/MassTranslatorRegistry.h
--- a/Engine/Plugins/Runtime/MassGameplay/Source/MassSpawner/Public/MassTranslatorRegistry.h
+++ b/Engine/Plugins/Runtime/MassGameplay/Source/MassSpawner/Public/MassTranslatorRegistry.h
@@ -11,7 +11,7 @@
 class UScriptStruct;
 class AActor;
 class UActorComponent;
-class UPipeProcessor;
+class UMassProcessor;
 class UMassTranslator;
 class UMassObjectFragmentCreator;
 class UMassFragmentInitializer;
@@ -32,7 +32,7 @@ public:
 	static const UMassTranslatorRegistry& Get() { return *GetDefault<UMassTranslatorRegistry>(); }
 
 	/** UMassFragmentInitializer - given a UScriptStruct sets its initial values to some defaults or randomized 
-	 *  values(implementation - specific). It's a kind of a PipeProcessor */
+	 *  values(implementation - specific). It's a kind of a MassProcessor */
 	void RegisterFragmentInitializer(const UScriptStruct& FragmentType, TSubclassOf<UMassFragmentInitializer> FragmentInitializerClass);
 	const UMassFragmentInitializer* GetFragmentInitializer(const UScriptStruct& FragmentType) const;
 
