Removed FMassHandle and replace its uses with FMassEnityHandle

diff --git a/Engine/Plugins/AI/MassAI/Source/MassAIBehavior/Private/MassComponentHitSubsystem.cpp b/Engine/Plugins/AI/MassAI/Source/MassAIBehavior/Private/MassComponentHitSubsystem.cpp
--- a/Engine/Plugins/AI/MassAI/Source/MassAIBehavior/Private/MassComponentHitSubsystem.cpp
+++ b/Engine/Plugins/AI/MassAI/Source/MassAIBehavior/Private/MassComponentHitSubsystem.cpp
@@ -39,7 +39,7 @@ void UMassComponentHitSubsystem::Initialize(FSubsystemCollectionBase& Collection
 	{
 		if (UCapsuleComponent* CapsuleComponent = AgentComponent.GetOwner()->FindComponentByClass<UCapsuleComponent>())
 		{
-			RegisterForComponentHit(AgentComponent.GetEntityHandle().GetLWEntity(), *CapsuleComponent);
+			RegisterForComponentHit(AgentComponent.GetEntityHandle(), *CapsuleComponent);
 		}
 	});
 
@@ -47,7 +47,7 @@ void UMassComponentHitSubsystem::Initialize(FSubsystemCollectionBase& Collection
 	{
 		if (UCapsuleComponent* CapsuleComponent = AgentComponent.GetOwner()->FindComponentByClass<UCapsuleComponent>())
 		{
-			UnregisterForComponentHit(AgentComponent.GetEntityHandle().GetLWEntity(), *CapsuleComponent);
+			UnregisterForComponentHit(AgentComponent.GetEntityHandle(), *CapsuleComponent);
 		}
 	});
 }
@@ -79,7 +79,7 @@ void UMassComponentHitSubsystem::OnHitCallback(UPrimitiveComponent* HitComp, AAc
 {
 	const UWorld* World = GetWorld();
 	check(World);
-	const FMassEntityHandle& Entity = ComponentToEntityMap.FindChecked(HitComp);
+	const FMassEntityHandle Entity = ComponentToEntityMap.FindChecked(HitComp);
 	FMassEntityHandle* OtherEntity = ComponentToEntityMap.Find(OtherComp);
 
 	bool bProcessHit = (OtherEntity != nullptr && OtherEntity->IsSet());
diff --git a/Engine/Plugins/AI/MassAI/Source/MassAIBehavior/Public/MassComponentHitTypes.h b/Engine/Plugins/AI/MassAI/Source/MassAIBehavior/Public/MassComponentHitTypes.h
--- a/Engine/Plugins/AI/MassAI/Source/MassAIBehavior/Public/MassComponentHitTypes.h
+++ b/Engine/Plugins/AI/MassAI/Source/MassAIBehavior/Public/MassComponentHitTypes.h
@@ -20,7 +20,7 @@ struct FMassHitResult
 
 	FMassHitResult() = default;
 
-	FMassHitResult(const FMassEntityHandle& OtherEntity, const float Time)
+	FMassHitResult(const FMassEntityHandle OtherEntity, const float Time)
 		: OtherEntity(OtherEntity)
 		, HitTime(Time)
 		, LastFilteredHitTime(Time)
diff --git a/Engine/Plugins/AI/MassAI/Source/MassAIBehavior/Public/MassStateTreeExecutionContext.h b/Engine/Plugins/AI/MassAI/Source/MassAIBehavior/Public/MassStateTreeExecutionContext.h
--- a/Engine/Plugins/AI/MassAI/Source/MassAIBehavior/Public/MassStateTreeExecutionContext.h
+++ b/Engine/Plugins/AI/MassAI/Source/MassAIBehavior/Public/MassStateTreeExecutionContext.h
@@ -26,7 +26,7 @@ public:
 	FMassExecutionContext& GetEntitySubsystemExecutionContext() const { return *EntitySubsystemExecutionContext; }
 
 	FMassEntityHandle GetEntity() const { return Entity; }
-	void SetEntity(const FMassEntityHandle& InEntity) { Entity = InEntity; }
+	void SetEntity(const FMassEntityHandle InEntity) { Entity = InEntity; }
 
 	int32 GetEntityIndex() const { return EntityIndex; }
 	void SetEntityIndex(const int32 InIndex) { EntityIndex = InIndex; }
diff --git a/Engine/Plugins/AI/MassAI/Source/MassAIDebug/Private/GameplayDebuggerCategory_Mass.cpp b/Engine/Plugins/AI/MassAI/Source/MassAIDebug/Private/GameplayDebuggerCategory_Mass.cpp
--- a/Engine/Plugins/AI/MassAI/Source/MassAIDebug/Private/GameplayDebuggerCategory_Mass.cpp
+++ b/Engine/Plugins/AI/MassAI/Source/MassAIDebug/Private/GameplayDebuggerCategory_Mass.cpp
@@ -101,7 +101,7 @@ void FGameplayDebuggerCategory_Mass::PickEntity(const APlayerController& OwnerPC
 		UMassActorSubsystem* ActorManager = World.GetSubsystem<UMassActorSubsystem>();
 		if (ActorManager != nullptr)
 		{
-			BestActor = ActorManager->GetActorFromHandle(FMassHandle(CachedEntity));
+			BestActor = ActorManager->GetActorFromHandle(FMassEntityHandle(CachedEntity));
 			CachedDebugActor = BestActor;
 		}
 	}
@@ -139,7 +139,7 @@ void FGameplayDebuggerCategory_Mass::CollectData(APlayerController* OwnerPC, AAc
 		UMassAgentComponent* AgentComp = Actor.FindComponentByClass<UMassAgentComponent>();
 		if (AgentComp)
 		{
-			LWEntity = AgentComp->GetEntityHandle().GetLWEntity();
+			LWEntity = AgentComp->GetEntityHandle();
 			if (OutMassAgentComponent)
 			{
 				*OutMassAgentComponent = AgentComp;
@@ -150,7 +150,7 @@ void FGameplayDebuggerCategory_Mass::CollectData(APlayerController* OwnerPC, AAc
 			UMassActorSubsystem* ActorManager = UWorld::GetSubsystem<UMassActorSubsystem>(Actor.GetWorld());
 			if (ActorManager != nullptr)
 			{
-				LWEntity = ActorManager->GetHandleFromActor(&Actor).GetLWEntity();
+				LWEntity = ActorManager->GetEntityHandleFromActor(&Actor);
 			}
 		}
 		return LWEntity;
diff --git a/Engine/Plugins/AI/MassAI/Source/MassAIMovement/Private/MassAvoidanceProcessors.cpp b/Engine/Plugins/AI/MassAI/Source/MassAIMovement/Private/MassAvoidanceProcessors.cpp
--- a/Engine/Plugins/AI/MassAI/Source/MassAIMovement/Private/MassAvoidanceProcessors.cpp
+++ b/Engine/Plugins/AI/MassAI/Source/MassAIMovement/Private/MassAvoidanceProcessors.cpp
@@ -225,7 +225,7 @@ namespace UE::MassAvoidance
 		const FMassEntityHandle Entity;
 	};
 
-	static bool DebugIsSelected(const FMassEntityHandle& Entity)
+	static bool DebugIsSelected(const FMassEntityHandle Entity)
 	{
 		FColor Color;
 		return UE::MassDebug::IsDebuggingEntity(Entity, &Color);
diff --git a/Engine/Plugins/AI/MassAI/Source/MassAIMovement/Private/MovementProcessor.cpp b/Engine/Plugins/AI/MassAI/Source/MassAIMovement/Private/MovementProcessor.cpp
--- a/Engine/Plugins/AI/MassAI/Source/MassAIMovement/Private/MovementProcessor.cpp
+++ b/Engine/Plugins/AI/MassAI/Source/MassAIMovement/Private/MovementProcessor.cpp
@@ -58,7 +58,7 @@ namespace UE::MassMovement
 		FAutoConsoleVariableRef(TEXT("ai.mass.AvoidDistance"), Tweakables::AvoidDistanceCutOff, TEXT(""), ECVF_Default), FAutoConsoleVariableRef(TEXT("ai.mass.AvoidRadius"), Tweakables::AvoidAgentRadius, TEXT(""), ECVF_Default), FAutoConsoleVariableRef(TEXT("ai.mass.AvoidMaxForce"), Tweakables::AvoidMaxForce, TEXT(""), ECVF_Default), FAutoConsoleVariableRef(TEXT("ai.mass.AvoidTimeHorizon"), Tweakables::AvoidTimeHorizon, TEXT(""), ECVF_Default)
 	};
 
- 	static bool DebugIsSelected(const FMassEntityHandle& Entity)
+ 	static bool DebugIsSelected(const FMassEntityHandle Entity)
  	{
  #if WITH_MASSGAMEPLAY_DEBUG
  		return UE::MassDebug::IsDebuggingEntity(Entity);
diff --git a/Engine/Plugins/AI/MassCrowd/Source/MassCrowd/Private/MassCrowdRepresentationProcessor.cpp b/Engine/Plugins/AI/MassCrowd/Source/MassCrowd/Private/MassCrowdRepresentationProcessor.cpp
--- a/Engine/Plugins/AI/MassCrowd/Source/MassCrowd/Private/MassCrowdRepresentationProcessor.cpp
+++ b/Engine/Plugins/AI/MassCrowd/Source/MassCrowd/Private/MassCrowdRepresentationProcessor.cpp
@@ -125,7 +125,7 @@ void UMassCrowdRepresentationProcessor::SetActorEnabled(const EActorEnabledType
 	}
 }
 
-AActor* UMassCrowdRepresentationProcessor::GetOrSpawnActor(const FMassHandle MassAgent, FDataFragment_Actor& ActorInfo, const FTransform& Transform, const int16 TemplateActorIndex, FMassHandle_ActorSpawnRequest& SpawnRequestHandle, const float Priority)
+AActor* UMassCrowdRepresentationProcessor::GetOrSpawnActor(const FMassEntityHandle MassAgent, FDataFragment_Actor& ActorInfo, const FTransform& Transform, const int16 TemplateActorIndex, FMassActorSpawnRequestHandle& SpawnRequestHandle, const float Priority)
 {
 	FTransform RootTransform = Transform;
 	
diff --git a/Engine/Plugins/AI/MassCrowd/Source/MassCrowd/Private/MassCrowdSubsystem.cpp b/Engine/Plugins/AI/MassCrowd/Source/MassCrowd/Private/MassCrowdSubsystem.cpp
--- a/Engine/Plugins/AI/MassCrowd/Source/MassCrowd/Private/MassCrowdSubsystem.cpp
+++ b/Engine/Plugins/AI/MassCrowd/Source/MassCrowd/Private/MassCrowdSubsystem.cpp
@@ -364,7 +364,7 @@ const FRegisteredCrowdLaneData* UMassCrowdSubsystem::GetCrowdData(const FZoneGra
 	return &LanesData;
 }
 
-void UMassCrowdSubsystem::OnEntityLaneChanged(const FMassEntityHandle& Entity, const FZoneGraphLaneHandle PreviousLaneHandle, const FZoneGraphLaneHandle CurrentLaneHandle)
+void UMassCrowdSubsystem::OnEntityLaneChanged(const FMassEntityHandle Entity, const FZoneGraphLaneHandle PreviousLaneHandle, const FZoneGraphLaneHandle CurrentLaneHandle)
 {
 	const bool bPreviousLocationValid = PreviousLaneHandle.IsValid();
 	const bool bCurrentLocationValid = CurrentLaneHandle.IsValid();
@@ -544,14 +544,14 @@ void UMassCrowdSubsystem::UpdateDensityMask()
 	}
 }
 
-void UMassCrowdSubsystem::OnEnterTrackedLane(const FMassEntityHandle& Entity, const int32 LaneIndex, FCrowdTrackingLaneData& TrackingData)
+void UMassCrowdSubsystem::OnEnterTrackedLane(const FMassEntityHandle Entity, const int32 LaneIndex, FCrowdTrackingLaneData& TrackingData)
 {
 	++TrackingData.NumEntitiesOnLane;
 
 	UE_VLOG_UELOG(this, LogMassNavigation, Verbose, TEXT("[%s] enters lane %d. Num entities on lane: %d"), *Entity.DebugGetDescription(), LaneIndex, TrackingData.NumEntitiesOnLane);
 }
 
-void UMassCrowdSubsystem::OnExitTrackedLane(const FMassEntityHandle& Entity, const int32 LaneIndex, FCrowdTrackingLaneData& TrackingData)
+void UMassCrowdSubsystem::OnExitTrackedLane(const FMassEntityHandle Entity, const int32 LaneIndex, FCrowdTrackingLaneData& TrackingData)
 {
 	ensureMsgf(TrackingData.NumEntitiesOnLane >= 1, TEXT("OnExitTrackedLane should not be called more often than OnEnterTrackedLane"));
 	--TrackingData.NumEntitiesOnLane;
diff --git a/Engine/Plugins/AI/MassCrowd/Source/MassCrowd/Public/MassCrowdRepresentationProcessor.h b/Engine/Plugins/AI/MassCrowd/Source/MassCrowd/Public/MassCrowdRepresentationProcessor.h
--- a/Engine/Plugins/AI/MassCrowd/Source/MassCrowd/Public/MassCrowdRepresentationProcessor.h
+++ b/Engine/Plugins/AI/MassCrowd/Source/MassCrowd/Public/MassCrowdRepresentationProcessor.h
@@ -49,7 +49,7 @@ protected:
 	 * @param Context of the execution from the entity sub system
 	 * @return the actor spawned
 	 */
-	virtual AActor* GetOrSpawnActor(const FMassHandle MassAgent, FDataFragment_Actor& ActorInfo, const FTransform& Transform, const int16 TemplateActorIndex, FMassHandle_ActorSpawnRequest& SpawnRequestHandle, const float Priority) override;
+	virtual AActor* GetOrSpawnActor(const FMassEntityHandle MassAgent, FDataFragment_Actor& ActorInfo, const FTransform& Transform, const int16 TemplateActorIndex, FMassActorSpawnRequestHandle& SpawnRequestHandle, const float Priority) override;
 
 	/**
 	 * Teleports the actor at the specified transform by preserving its velocity and without collision.
diff --git a/Engine/Plugins/AI/MassCrowd/Source/MassCrowd/Public/MassCrowdSubsystem.h b/Engine/Plugins/AI/MassCrowd/Source/MassCrowd/Public/MassCrowdSubsystem.h
--- a/Engine/Plugins/AI/MassCrowd/Source/MassCrowd/Public/MassCrowdSubsystem.h
+++ b/Engine/Plugins/AI/MassCrowd/Source/MassCrowd/Public/MassCrowdSubsystem.h
@@ -155,7 +155,7 @@ public:
 	 * @param PreviousLaneHandle Last frame lane handle (can be invalid)
 	 * @param CurrentLaneHandle Current frame lane handle (can be invalid)
 	 */
-	void OnEntityLaneChanged(const FMassEntityHandle& Entity, const FZoneGraphLaneHandle PreviousLaneHandle, const FZoneGraphLaneHandle CurrentLaneHandle);
+	void OnEntityLaneChanged(const FMassEntityHandle Entity, const FZoneGraphLaneHandle PreviousLaneHandle, const FZoneGraphLaneHandle CurrentLaneHandle);
 
 	/** Returns the weight for lane selection that is associated to the given lane based on its density tag. */
 	float GetDensityWeight(const FZoneGraphLaneHandle LaneHandle, const FZoneGraphTagMask LaneTagMask) const;
@@ -193,7 +193,7 @@ protected:
 	 * @param LaneIndex Index of the lane
 	 * @param TrackingData Runtime state associated to the lane
 	 */
-	void OnEnterTrackedLane(const FMassEntityHandle& Entity, const int32 LaneIndex, FCrowdTrackingLaneData& TrackingData);
+	void OnEnterTrackedLane(const FMassEntityHandle Entity, const int32 LaneIndex, FCrowdTrackingLaneData& TrackingData);
 
 	/**
 	 * Callback to keep count of entities currently on a given lane.
@@ -201,7 +201,7 @@ protected:
 	 * @param LaneIndex Index of the lane
 	 * @param TrackingData Runtime state associated to the lane
 	 */
-	void OnExitTrackedLane(const FMassEntityHandle& Entity, const int32 LaneIndex, FCrowdTrackingLaneData& TrackingData);
+	void OnExitTrackedLane(const FMassEntityHandle Entity, const int32 LaneIndex, FCrowdTrackingLaneData& TrackingData);
 
 	/**
 	 * Creates and initializes the occupancy data of a lane.
diff --git a/Engine/Plugins/Runtime/MassEntity/Source/MassEntity/Public/MassEntityTypes.h b/Engine/Plugins/Runtime/MassEntity/Source/MassEntity/Public/MassEntityTypes.h
--- a/Engine/Plugins/Runtime/MassEntity/Source/MassEntity/Public/MassEntityTypes.h
+++ b/Engine/Plugins/Runtime/MassEntity/Source/MassEntity/Public/MassEntityTypes.h
@@ -47,12 +47,12 @@ struct FMassEntityHandle
 	UPROPERTY(VisibleAnywhere, Category = "AggregateTicking|Debug", Transient)
 	int32 SerialNumber = 0;
 
-	bool operator==(const FMassEntityHandle& Other) const
+	bool operator==(const FMassEntityHandle Other) const
 	{
 		return Index == Other.Index && SerialNumber == Other.SerialNumber;
 	}
 
-	bool operator!=(const FMassEntityHandle& Other) const
+	bool operator!=(const FMassEntityHandle Other) const
 	{
 		return !operator==(Other);
 	}
@@ -64,12 +64,17 @@ struct FMassEntityHandle
 		return Index != 0 && SerialNumber != 0;
 	}
 
+	FORCEINLINE bool IsValid() const
+	{
+		return IsSet();
+	}
+
 	void Reset()
 	{
 		Index = SerialNumber = 0;
 	}
 
-	friend uint32 GetTypeHash(const FMassEntityHandle& Entity)
+	friend uint32 GetTypeHash(const FMassEntityHandle Entity)
 	{
 		return HashCombine(Entity.Index, Entity.SerialNumber);
 	}
diff --git a/Engine/Plugins/Runtime/MassGameplay/Source/MassActors/Private/MassActorHelper.cpp b/Engine/Plugins/Runtime/MassGameplay/Source/MassActors/Private/MassActorHelper.cpp
--- a/Engine/Plugins/Runtime/MassGameplay/Source/MassActors/Private/MassActorHelper.cpp
+++ b/Engine/Plugins/Runtime/MassGameplay/Source/MassActors/Private/MassActorHelper.cpp
@@ -15,10 +15,10 @@ namespace UE::MassActor
 		UMassActorSubsystem* MassActorSubsystem = UWorld::GetSubsystem<UMassActorSubsystem>(World);
 		if (EntitySystem && MassActorSubsystem)
 		{
-			const FMassHandle AgentHandle = MassActorSubsystem->GetHandleFromActor(&Actor);
+			const FMassEntityHandle AgentHandle = MassActorSubsystem->GetEntityHandleFromActor(&Actor);
 			if (AgentHandle.IsValid())
 			{
-				EntitySystem->AddTagToEntity(AgentHandle.GetLWEntity(), &TagType);
+				EntitySystem->AddTagToEntity(AgentHandle, &TagType);
 				return true;
 			}
 			UE_LOG(LogMassActor, Warning, TEXT("Failed to add tag %s to actor %s due to it not having an associated entity")
@@ -40,10 +40,10 @@ namespace UE::MassActor
 		UMassActorSubsystem* MassActorSubsystem = UWorld::GetSubsystem<UMassActorSubsystem>(World);
 		if (EntitySystem && MassActorSubsystem)
 		{
-			const FMassHandle AgentHandle = MassActorSubsystem->GetHandleFromActor(&Actor);
+			const FMassEntityHandle AgentHandle = MassActorSubsystem->GetEntityHandleFromActor(&Actor);
 			if (AgentHandle.IsValid())
 			{
-				EntitySystem->RemoveTagFromEntity(AgentHandle.GetLWEntity(), &TagType);
+				EntitySystem->RemoveTagFromEntity(AgentHandle, &TagType);
 				return true;
 			}
 			UE_LOG(LogMassActor, Warning, TEXT("Failed to remove tag %s from actor %s due to it not having an associated entity")
diff --git a/Engine/Plugins/Runtime/MassGameplay/Source/MassActors/Private/MassActorSpawnerSubsystem.cpp b/Engine/Plugins/Runtime/MassGameplay/Source/MassActors/Private/MassActorSpawnerSubsystem.cpp
--- a/Engine/Plugins/Runtime/MassGameplay/Source/MassActors/Private/MassActorSpawnerSubsystem.cpp
+++ b/Engine/Plugins/Runtime/MassGameplay/Source/MassActors/Private/MassActorSpawnerSubsystem.cpp
@@ -19,7 +19,7 @@ namespace UE::MassActors
 //----------------------------------------------------------------------//
 // UMassActorSpawnerSubsystem 
 //----------------------------------------------------------------------//
-void UMassActorSpawnerSubsystem::RetryActorSpawnRequest(FMassHandle_ActorSpawnRequest SpawnRequestHandle)
+void UMassActorSpawnerSubsystem::RetryActorSpawnRequest(const FMassActorSpawnRequestHandle SpawnRequestHandle)
 {
 	check(SpawnRequestHandleManager.IsValidHandle(SpawnRequestHandle));
 	const int32 Index = SpawnRequestHandle.GetIndex();
@@ -36,7 +36,7 @@ void UMassActorSpawnerSubsystem::RetryActorSpawnRequest(FMassHandle_ActorSpawnRe
 	}
 }
 
-bool UMassActorSpawnerSubsystem::RemoveActorSpawnRequest(FMassHandle_ActorSpawnRequest& SpawnRequestHandle)
+bool UMassActorSpawnerSubsystem::RemoveActorSpawnRequest(FMassActorSpawnRequestHandle& SpawnRequestHandle)
 {
 	if (!ensureMsgf(SpawnRequestHandleManager.RemoveHandle(SpawnRequestHandle), TEXT("Invalid spawn request handle")))
 	{
@@ -107,10 +107,10 @@ bool UMassActorSpawnerSubsystem::ReleaseActorToPool(AActor* Actor)
 	return false;
 }
 
-FMassHandle_ActorSpawnRequest UMassActorSpawnerSubsystem::RequestActorSpawnInternal(const FConstStructView SpawnRequestView)
+FMassActorSpawnRequestHandle UMassActorSpawnerSubsystem::RequestActorSpawnInternal(const FConstStructView SpawnRequestView)
 {
 	// The handle manager has a freelist of the release indexes, so it can return us a index that we previously used.
-	const FMassHandle_ActorSpawnRequest SpawnRequestHandle = SpawnRequestHandleManager.GetNextHandle();
+	const FMassActorSpawnRequestHandle SpawnRequestHandle = SpawnRequestHandleManager.GetNextHandle();
 	const int32 Index = SpawnRequestHandle.GetIndex();
 
 	// Check if we need to grow the array, otherwise it is a previously released index that was returned.
@@ -136,13 +136,13 @@ FMassHandle_ActorSpawnRequest UMassActorSpawnerSubsystem::RequestActorSpawnInter
 	return SpawnRequestHandle;
 }
 
-FMassHandle_ActorSpawnRequest UMassActorSpawnerSubsystem::GetNextRequestToSpawn() const
+FMassActorSpawnRequestHandle UMassActorSpawnerSubsystem::GetNextRequestToSpawn() const
 {
-	FMassHandle_ActorSpawnRequest BestSpawnRequestHandle;
+	FMassActorSpawnRequestHandle BestSpawnRequestHandle;
 	float BestPriority = MAX_FLT;
 	bool bBestIsPending = false;
 	uint32 BestSerialNumber = MAX_uint32;
-	for (const FMassHandle_ActorSpawnRequest SpawnRequestHandle : SpawnRequestHandleManager.GetHandles())
+	for (const FMassActorSpawnRequestHandle SpawnRequestHandle : SpawnRequestHandleManager.GetHandles())
 	{
 		if (!SpawnRequestHandle.IsValid())
 		{
@@ -234,7 +234,7 @@ AActor* UMassActorSpawnerSubsystem::SpawnActor(const FStructView& SpawnRequestVi
 					SpawnRequest.Template.GetDefaultObject()->GetSimpleCollisionRadius(),
 					SpawnRequest.Transform.GetRotation(),
 					FColor::Red,
-					TEXT("Unable to spawn actor for Mass entity [%s]"), *SpawnRequest.MassAgent.GetLWEntity().DebugGetDescription());
+					TEXT("Unable to spawn actor for Mass entity [%s]"), *SpawnRequest.MassAgent.DebugGetDescription());
 	return nullptr;
 }
 
@@ -247,7 +247,7 @@ void UMassActorSpawnerSubsystem::ProcessPendingSpawningRequest(const double MaxT
 
 	while (FPlatformTime::Seconds() < TimeSliceEnd)
 	{
-		FMassHandle_ActorSpawnRequest SpawnRequestHandle = GetNextRequestToSpawn();
+		FMassActorSpawnRequestHandle SpawnRequestHandle = GetNextRequestToSpawn();
 		if (!SpawnRequestHandle.IsValid() ||
 			!ensureMsgf(SpawnRequestHandleManager.IsValidHandle(SpawnRequestHandle), TEXT("GetNextRequestToSpawn returned an invalid handle, expecting an empty one or a valid one.")))
 		{
diff --git a/Engine/Plugins/Runtime/MassGameplay/Source/MassActors/Private/MassActorSubsystem.cpp b/Engine/Plugins/Runtime/MassGameplay/Source/MassActors/Private/MassActorSubsystem.cpp
--- a/Engine/Plugins/Runtime/MassGameplay/Source/MassActors/Private/MassActorSubsystem.cpp
+++ b/Engine/Plugins/Runtime/MassGameplay/Source/MassActors/Private/MassActorSubsystem.cpp
@@ -11,7 +11,7 @@
 //  FDataFragment_Actor 
 //----------------------------------------------------------------------//
 
-void FDataFragment_Actor::SetAndUpdateHandleMap(const FMassHandle MassAgent, AActor* InActor, const bool bInIsOwnedByMass)
+void FDataFragment_Actor::SetAndUpdateHandleMap(const FMassEntityHandle MassAgent, AActor* InActor, const bool bInIsOwnedByMass)
 {
 	SetNoHandleMapUpdate(MassAgent, InActor, bInIsOwnedByMass);
 
@@ -38,7 +38,7 @@ void FDataFragment_Actor::ResetAndUpdateHandleMap()
 	ResetNoHandleMapUpdate();
 }
 
-void FDataFragment_Actor::SetNoHandleMapUpdate(const FMassHandle MassAgent, AActor* InActor, const bool bInIsOwnedByMass)
+void FDataFragment_Actor::SetNoHandleMapUpdate(const FMassEntityHandle MassAgent, AActor* InActor, const bool bInIsOwnedByMass)
 {
 	check(InActor);
 	check(!Actor.IsValid());
@@ -65,30 +65,30 @@ void UMassActorSubsystem::Initialize(FSubsystemCollectionBase& Collection)
 	EntitySystem = UWorld::GetSubsystem<UMassEntitySubsystem>(GetWorld());
 }
 
-FMassHandle UMassActorSubsystem::GetHandleFromActor(const TObjectKey<const AActor> Actor)
+FMassEntityHandle UMassActorSubsystem::GetEntityHandleFromActor(const TObjectKey<const AActor> Actor)
 {
 	UE_MT_SCOPED_READ_ACCESS(ActorHandleMapDetector);
-	FMassEntityHandle* LWEntity = ActorHandleMap.Find(Actor);
-	if (!LWEntity)
+	FMassEntityHandle* Entity = ActorHandleMap.Find(Actor);
+	if (!Entity)
 	{
-		return FMassHandle::InvalidHandle;
+		return UMassEntitySubsystem::InvalidEntity;
 	}
 
-	check(TObjectKey<const AActor>(GetActorFromHandle(FMassHandle(*LWEntity))) == Actor);
-	return FMassHandle(*LWEntity);
+	check(TObjectKey<const AActor>(GetActorFromHandle(*Entity)) == Actor);
+	return *Entity;
 }
 
-AActor* UMassActorSubsystem::GetActorFromHandle(FMassHandle Handle) const
+AActor* UMassActorSubsystem::GetActorFromHandle(const FMassEntityHandle Handle) const
 {
 	check(EntitySystem);
-	FDataFragment_Actor* Data = EntitySystem->GetComponentDataPtr<FDataFragment_Actor>(Handle.GetLWEntity());
+	FDataFragment_Actor* Data = EntitySystem->GetComponentDataPtr<FDataFragment_Actor>(Handle);
 	return Data != nullptr ? Data->GetMutable() : nullptr;
 }
 
-void UMassActorSubsystem::SetHandleForActor(const TObjectKey<const AActor> Actor, FMassHandle Handle)
+void UMassActorSubsystem::SetHandleForActor(const TObjectKey<const AActor> Actor, const FMassEntityHandle Handle)
 {
 	UE_MT_SCOPED_WRITE_ACCESS(ActorHandleMapDetector);
-	ActorHandleMap.Add(Actor, Handle.GetLWEntity());
+	ActorHandleMap.Add(Actor, Handle);
 }
 
 void UMassActorSubsystem::RemoveHandleForActor(const TObjectKey<const AActor> Actor)
@@ -97,29 +97,29 @@ void UMassActorSubsystem::RemoveHandleForActor(const TObjectKey<const AActor> Ac
 	ActorHandleMap.Remove(Actor);
 }
 
-void UMassActorSubsystem::DisconnectActor(const TObjectKey<const AActor> Actor, FMassHandle Handle)
+void UMassActorSubsystem::DisconnectActor(const TObjectKey<const AActor> Actor, const FMassEntityHandle Handle)
 {
 	if (Handle.IsValid() == false)
 	{
 		return;
 	}
 
-	FMassEntityHandle LWEntity;
+	FMassEntityHandle FoundEntity;
 	{
 		UE_MT_SCOPED_WRITE_ACCESS(ActorHandleMapDetector);
 		// We're assuming the Handle does match Actor, so we're RemoveAndCopyValue. If if doesn't we'll add it back.
 		// The expectation is that this won't happen on a regular basis..
-		if (ActorHandleMap.RemoveAndCopyValue(Actor, LWEntity) == false)
+		if (ActorHandleMap.RemoveAndCopyValue(Actor, FoundEntity) == false)
 		{
 			// the entity doesn't match the actor
 			return;
 		}
 	}
 
-	if (LWEntity == Handle.GetLWEntity())
+	if (FoundEntity == Handle)
 	{
 		check(EntitySystem);
-		if (FDataFragment_Actor* Data = EntitySystem->GetComponentDataPtr<FDataFragment_Actor>(Handle.GetLWEntity()))
+		if (FDataFragment_Actor* Data = EntitySystem->GetComponentDataPtr<FDataFragment_Actor>(Handle))
 		{
 			Data->ResetAndUpdateHandleMap();
 		}
diff --git a/Engine/Plugins/Runtime/MassGameplay/Source/MassActors/Private/MassAgentComponent.cpp b/Engine/Plugins/Runtime/MassGameplay/Source/MassActors/Private/MassAgentComponent.cpp
--- a/Engine/Plugins/Runtime/MassGameplay/Source/MassActors/Private/MassAgentComponent.cpp
+++ b/Engine/Plugins/Runtime/MassGameplay/Source/MassActors/Private/MassAgentComponent.cpp
@@ -146,7 +146,7 @@ void UMassAgentComponent::OnUnregister()
 	Super::OnUnregister();
 }
 
-void UMassAgentComponent::SetEntityHandle(FMassHandle NewHandle)
+void UMassAgentComponent::SetEntityHandle(const FMassEntityHandle NewHandle)
 {
 	MASSAGENT_CHECK(State == EAgentComponentState::EntityPendingCreation,
 		TEXT("%s is expecting to be in state[EntityPendingCreation] but is in %s"), ANSI_TO_TCHAR(__FUNCTION__), *UEnum::GetValueAsString(State));
@@ -155,7 +155,7 @@ void UMassAgentComponent::SetEntityHandle(FMassHandle NewHandle)
 	SwitchToState(EAgentComponentState::EntityCreated);
 }
 
-void UMassAgentComponent::SetEntityHandleInternal(FMassHandle NewHandle)
+void UMassAgentComponent::SetEntityHandleInternal(const FMassEntityHandle NewHandle)
 {
 	ensureMsgf((AgentHandle.IsValid() && NewHandle.IsValid()) == false, TEXT("Overriding an existing entity ID might result in a dangling entity still affecting the simulation"));
 	AgentHandle = NewHandle;
@@ -164,7 +164,7 @@ void UMassAgentComponent::SetEntityHandleInternal(FMassHandle NewHandle)
 	// Fetch NetID if it exist
 	if (const UMassEntitySubsystem* EntitySubsystem = UWorld::GetSubsystem<UMassEntitySubsystem>(GetWorld()))
 	{
-		if (const FMassNetworkIDFragment* NetIDFragment = EntitySubsystem->GetComponentDataPtr<FMassNetworkIDFragment>(AgentHandle.GetLWEntity()))
+		if (const FMassNetworkIDFragment* NetIDFragment = EntitySubsystem->GetComponentDataPtr<FMassNetworkIDFragment>(AgentHandle))
 		{
 			if (!IsNetSimulating())
 			{
@@ -178,10 +178,6 @@ void UMassAgentComponent::SetEntityHandleInternal(FMassHandle NewHandle)
 	}
 #endif // UE_REPLICATION_COMPILE_SERVER_CODE
 
-#if WITH_EDITORONLY_DATA
-	DebugEntity = AgentHandle.GetLWEntity();
-#endif // WITH_EDITORONLY_DATA
-
 	if (const UMassAgentSubsystem* AgentSubsystem = UWorld::GetSubsystem<UMassAgentSubsystem>(GetWorld()))
 	{
 		AgentSubsystem->NotifyMassAgentComponentEntityAssociated(*this);
@@ -192,7 +188,7 @@ void UMassAgentComponent::SetEntityHandleInternal(FMassHandle NewHandle)
 	{
 		if (IsNetSimulating())
 		{
-			const FMassEntityView EntityView(*EntitySubsystem, AgentHandle.GetLWEntity());
+			const FMassEntityView EntityView(*EntitySubsystem, AgentHandle);
 
 			// @todo Find a way to add these initialization into either translator initializer or adding new fragments
 			// Make sure to fetch the fragment after any release, as that action can move the entity around into new archetype and 
@@ -221,7 +217,7 @@ void UMassAgentComponent::SetEntityHandleInternal(FMassHandle NewHandle)
 	}
 }
 
-void UMassAgentComponent::SetPuppetHandle(FMassHandle NewHandle)
+void UMassAgentComponent::SetPuppetHandle(const FMassEntityHandle NewHandle)
 {
 	MASSAGENT_CHECK(State == EAgentComponentState::EntityPendingCreation,
 		TEXT("%s is expecting to be in state[EntityPendingCreation] but is in %s"), ANSI_TO_TCHAR(__FUNCTION__), *UEnum::GetValueAsString(State));
@@ -281,7 +277,7 @@ void UMassAgentComponent::ClearEntityHandleInternal()
 	{
 		if (IsNetSimulating())
 		{
-			const FMassEntityView EntityView(*EntitySubsystem, AgentHandle.GetLWEntity());
+			const FMassEntityView EntityView(*EntitySubsystem, AgentHandle);
 			if (FDataFragment_Actor* ActorInfo = EntityView.GetComponentDataPtr<FDataFragment_Actor>())
 			{
 				checkf(!ActorInfo->IsValid() || ActorInfo->Get() == GetOwner(), TEXT("Expecting actor pointer to be the component owner"));
@@ -289,11 +285,8 @@ void UMassAgentComponent::ClearEntityHandleInternal()
 			}
 		}
 	}
-
-#if WITH_EDITORONLY_DATA
-	DebugEntity = UMassEntitySubsystem::InvalidEntity;
-#endif // WITH_EDITORONLY_DATA
-	AgentHandle = FMassHandle::InvalidHandle;
+	
+	AgentHandle = UMassEntitySubsystem::InvalidEntity;
 }
 
 void UMassAgentComponent::PuppetUnregistrationDone()
@@ -327,7 +320,7 @@ void UMassAgentComponent::EntityCreationAborted()
 
 void UMassAgentComponent::SwitchToState(EAgentComponentState NewState)
 {
-	UE_VLOG(GetOwner(), LogMass, Verbose, TEXT("Entity[%s] %s From:%s To:%s"), *AgentHandle.GetLWEntity().DebugGetDescription(), ANSI_TO_TCHAR(__FUNCTION__), *UEnum::GetValueAsString(State), *UEnum::GetValueAsString(NewState));
+	UE_VLOG(GetOwner(), LogMass, Verbose, TEXT("Entity[%s] %s From:%s To:%s"), *AgentHandle.DebugGetDescription(), ANSI_TO_TCHAR(__FUNCTION__), *UEnum::GetValueAsString(State), *UEnum::GetValueAsString(NewState));
 	State = NewState;
 	DebugCheckStateConsistency();
 }
@@ -354,16 +347,16 @@ void UMassAgentComponent::DebugCheckStateConsistency()
 			{
 				if (UMassEntitySubsystem* EntitySubsystem = UWorld::GetSubsystem<UMassEntitySubsystem>(GetWorld()))
 				{
-					const bool bIsValidEntity = EntitySubsystem->IsEntityValid(AgentHandle.GetLWEntity());
+					const bool bIsValidEntity = EntitySubsystem->IsEntityValid(AgentHandle);
 					MASSAGENT_CHECK(bIsValidEntity, TEXT("Exepecting a valid entity in state"), *UEnum::GetValueAsString(State))
 						if (bIsValidEntity)
 						{
-							const bool bIsBuiltEntity = EntitySubsystem->IsEntityBuilt(AgentHandle.GetLWEntity());
+							const bool bIsBuiltEntity = EntitySubsystem->IsEntityBuilt(AgentHandle);
 							MASSAGENT_CHECK(bIsBuiltEntity, TEXT("Expecting a fully built entity in state %s"), *UEnum::GetValueAsString(State));
 							if (bIsBuiltEntity)
 							{
 								AActor* Owner = GetOwner();
-								const AActor* Actor = EntitySubsystem->GetComponentDataChecked<FDataFragment_Actor>(AgentHandle.GetLWEntity()).Get();
+								const AActor* Actor = EntitySubsystem->GetComponentDataChecked<FDataFragment_Actor>(AgentHandle).Get();
 								MASSAGENT_CHECK(Actor == nullptr || Actor == Owner, TEXT("Mass Actor and Owner mismatched in state %s"), *UEnum::GetValueAsString(State));
 							}
 						}
@@ -425,7 +418,7 @@ void UMassAgentComponent::KillEntity(const bool bDestroyActor)
 	}
 
 	// Caching the entity and if we have to disconnect actor as the next operation will invalidate that information
-	const FMassHandle EntityHandleToDespawn = AgentHandle;
+	const FMassEntityHandle EntityHandleToDespawn = AgentHandle;
 	const bool bDisconnectActor = IsPuppet() && bDestroyActor == false;
 	if (State != EAgentComponentState::None)
 	{
@@ -458,7 +451,7 @@ void UMassAgentComponent::KillEntity(const bool bDestroyActor)
 	{
 		for (TActorIterator<AMassSpawner> It(World); It; ++It)
 		{
-			(*It)->DespawnEntity(EntityHandleToDespawn.GetLWEntity());
+			(*It)->DespawnEntity(EntityHandleToDespawn);
 		}
 	}
 }
@@ -499,7 +492,7 @@ void UMassAgentComponent::PuppetReplicationPending()
 	SwitchToState(EAgentComponentState::PuppetPendingReplication);
 }
 
-void UMassAgentComponent::SetReplicatedPuppetHandle(FMassHandle NewHandle)
+void UMassAgentComponent::SetReplicatedPuppetHandle(FMassEntityHandle NewHandle)
 {
 	checkf(IsNetSimulating(), TEXT("Expecting a replicated pupet"));
 	MASSAGENT_CHECK(State == EAgentComponentState::PuppetPendingReplication ||
diff --git a/Engine/Plugins/Runtime/MassGameplay/Source/MassActors/Private/MassAgentSubsystem.cpp b/Engine/Plugins/Runtime/MassGameplay/Source/MassActors/Private/MassAgentSubsystem.cpp
--- a/Engine/Plugins/Runtime/MassGameplay/Source/MassActors/Private/MassAgentSubsystem.cpp
+++ b/Engine/Plugins/Runtime/MassGameplay/Source/MassActors/Private/MassAgentSubsystem.cpp
@@ -136,7 +136,7 @@ void UMassAgentSubsystem::UpdateAgentComponent(const UMassAgentComponent& AgentC
 		return;
 	}
 
-	const FMassEntityHandle Entity = AgentComp.GetEntityHandle().GetLWEntity();
+	const FMassEntityHandle Entity = AgentComp.GetEntityHandle();
 	const FArchetypeHandle CurrentArchetypeHandle = EntitySystem->GetArchetypeForEntity(Entity);
 	if (CurrentArchetypeHandle == EntityTemplate->GetArchetype())
 	{
@@ -178,7 +178,7 @@ void UMassAgentSubsystem::UnregisterAgentComponent(UMassAgentComponent& AgentCom
 		{
 			if (AgentComp.GetPuppetSpecificAddition().IsEmpty() == false)
 			{
-				const FMassEntityHandle Entity = AgentComp.GetEntityHandle().GetLWEntity();
+				const FMassEntityHandle Entity = AgentComp.GetEntityHandle();
 
 			    // remove fragments that have been added for the puppet agent
 			    if (EntitySystem->IsProcessing())
@@ -207,7 +207,7 @@ void UMassAgentSubsystem::UnregisterAgentComponent(UMassAgentComponent& AgentCom
 				EntityTemplate = &EntityConfig.GetEntityTemplateChecked(*AgentActor, AgentComp);
 			}
 
-			FMassEntityHandle Entity = AgentComp.GetEntityHandle().GetLWEntity();
+			FMassEntityHandle Entity = AgentComp.GetEntityHandle();
 			// Clearing the entity before it become invalid as the clear contains notifications
 			AgentComp.ClearEntityHandle();
 
@@ -320,7 +320,7 @@ void UMassAgentSubsystem::HandlePendingInitialization()
 
 		for (int AgentIndex = 0; AgentIndex < Entities.Num(); ++AgentIndex)
 		{		
-			AgentComponents[AgentIndex]->SetEntityHandle(FMassHandle(Entities[AgentIndex]));
+			AgentComponents[AgentIndex]->SetEntityHandle(Entities[AgentIndex]);
 		}
 	}
 
@@ -342,7 +342,7 @@ void UMassAgentSubsystem::HandlePendingInitialization()
 
 		for (UMassAgentComponent* AgentComp : AgentComponents)
 		{
-			const FMassEntityHandle PuppetEntity = AgentComp->GetEntityHandle().GetLWEntity();
+			const FMassEntityHandle PuppetEntity = AgentComp->GetEntityHandle();
 			if (!ensureMsgf(PuppetEntity.IsSet(), TEXT("Trying to initialize puppet's fragments while the pupped doesn't have a corresponding Entity identifier set. This should not happen.")))
 			{
 				continue;
@@ -381,7 +381,7 @@ void UMassAgentSubsystem::NotifyMassAgentComponentReplicated(UMassAgentComponent
 		// If not found, the NotifyMassAgentAddedToReplication will link it later once replicated.
 		if (Entity.IsSet())
 		{
-			AgentComp.SetReplicatedPuppetHandle(FMassHandle(Entity));
+			AgentComp.SetReplicatedPuppetHandle(FMassEntityHandle(Entity));
 			MakePuppet(AgentComp);
 		}
 		else
@@ -410,7 +410,7 @@ void UMassAgentSubsystem::OnMassAgentAddedToReplication(FMassNetworkID NetID, FM
 	{
 		if (UMassAgentComponent** AgentComp = ReplicatedAgentComponents.Find(NetID))
 		{
-			(*AgentComp)->SetReplicatedPuppetHandle(FMassHandle(Entity));
+			(*AgentComp)->SetReplicatedPuppetHandle(FMassEntityHandle(Entity));
 			MakePuppet(**AgentComp);
 		}
 	}
diff --git a/Engine/Plugins/Runtime/MassGameplay/Source/MassActors/Public/MassActorSpawnerSubsystem.h b/Engine/Plugins/Runtime/MassGameplay/Source/MassActors/Public/MassActorSpawnerSubsystem.h
--- a/Engine/Plugins/Runtime/MassGameplay/Source/MassActors/Public/MassActorSpawnerSubsystem.h
+++ b/Engine/Plugins/Runtime/MassGameplay/Source/MassActors/Public/MassActorSpawnerSubsystem.h
@@ -14,29 +14,29 @@
 
 // Handle for an actor spawning request
 USTRUCT()
-struct MASSACTORS_API FMassHandle_ActorSpawnRequest : public FIndexedHandleBase
+struct MASSACTORS_API FMassActorSpawnRequestHandle : public FIndexedHandleBase
 {
 	GENERATED_BODY()
 
-	FMassHandle_ActorSpawnRequest() = default;
+	FMassActorSpawnRequestHandle() = default;
 
 	/** @note passing INDEX_NONE as index will make this handle Invalid */
-	FMassHandle_ActorSpawnRequest(const int32 InIndex, const uint32 InSerialNumber) : FIndexedHandleBase(InIndex, InSerialNumber)
+	FMassActorSpawnRequestHandle(const int32 InIndex, const uint32 InSerialNumber) : FIndexedHandleBase(InIndex, InSerialNumber)
 	{
 	}
 };
 
 // Managing class of spawning requests handles
-typedef FIndexedHandleManager<FMassHandle_ActorSpawnRequest, true/*bOptimizeHandleReuse*/> FMassHandleManager_ActorSpawnRequest;
+typedef FIndexedHandleManager<FMassActorSpawnRequestHandle, true/*bOptimizeHandleReuse*/> FMassEntityHandleManager_ActorSpawnRequest;
 
 
-DECLARE_DELEGATE_TwoParams(FMassActorPreSpawnDelegate, const FMassHandle_ActorSpawnRequest& SpawnRequestHandle, const FStructView& SpawnRequest);
+DECLARE_DELEGATE_TwoParams(FMassActorPreSpawnDelegate, const FMassActorSpawnRequestHandle& SpawnRequestHandle, const FStructView& SpawnRequest);
 enum EMassActorSpawnRequestAction
 {
 	Keep, // Will leave spawning request in the queue and it will be users job to call RemoveActorSpawnRequest
 	Remove, // Will remove the spawning request from the queue once the callback ends
 };
-DECLARE_DELEGATE_RetVal_TwoParams(EMassActorSpawnRequestAction, FMassActorPostSpawnDelegate, const FMassHandle_ActorSpawnRequest& SpawnRequestHandle, const FStructView& SpawnRequest);
+DECLARE_DELEGATE_RetVal_TwoParams(EMassActorSpawnRequestAction, FMassActorPostSpawnDelegate, const FMassActorSpawnRequestHandle& SpawnRequestHandle, const FStructView& SpawnRequest);
 
 UENUM()
 enum class ESpawnRequestStatus : uint8
@@ -58,7 +58,7 @@ struct MASSACTORS_API FMassActorSpawnRequest
 	GENERATED_BODY()
 public:
 	/** The actual mass agent handle corresponding to the actor to spawn */
-	FMassHandle MassAgent;
+	FMassEntityHandle MassAgent;
 
 	/** The template BP actor to spawn */
 	UPROPERTY(Transient)
@@ -91,7 +91,7 @@ public:
 
 	void Reset()
 	{
-		MassAgent = FMassHandle();
+		MassAgent = FMassEntityHandle();
 		Template = nullptr;
 		Priority = MAX_FLT;
 		ActorPreSpawnDelegate.Unbind();
@@ -124,7 +124,7 @@ public:
 	 * @param InSpawnRequest the spawn request parameters, You can provide any type of UStruct as long at it derives from FMassActorSpawnRequest. This let you add more spawning information.
 	 */
 	template< typename T, typename = typename TEnableIf<TIsDerivedFrom<typename TRemoveReference<T>::Type, FMassActorSpawnRequest>::IsDerived, void>::Type >
-	FMassHandle_ActorSpawnRequest RequestActorSpawn(const T& InSpawnRequest)
+	FMassActorSpawnRequestHandle RequestActorSpawn(const T& InSpawnRequest)
 	{
 		return RequestActorSpawnInternal(FConstStructView::Make(InSpawnRequest));
 	}
@@ -132,7 +132,7 @@ public:
 	/** Retries a failed spawn request
 	 * @param SpawnRequestHandle the spawn request handle to retry
 	 */
-	void RetryActorSpawnRequest(FMassHandle_ActorSpawnRequest SpawnRequestHandle);
+	void RetryActorSpawnRequest(const FMassActorSpawnRequestHandle SpawnRequestHandle);
 
 	/**
 	 * Removes a spawn request
@@ -141,7 +141,7 @@ public:
 	 * @param SpawnRequestHandle [IN/OUT] the spawn request handle to remove
 	 * @return true if successfully removed the request
 	 */
-	 bool RemoveActorSpawnRequest(FMassHandle_ActorSpawnRequest& SpawnRequestHandle);
+	 bool RemoveActorSpawnRequest(FMassActorSpawnRequestHandle& SpawnRequestHandle);
 
 	/**
 	 * Returns the stored spawn request from the handle, useful to update the transform
@@ -149,7 +149,7 @@ public:
 	 * @return The spawn request cast in the provided template argument
 	 */
 	template<typename T>
-	const T& GetSpawnRequest(FMassHandle_ActorSpawnRequest SpawnRequestHandle) const
+	const T& GetSpawnRequest(const FMassActorSpawnRequestHandle SpawnRequestHandle) const
 	{
 		check(SpawnRequestHandleManager.IsValidHandle(SpawnRequestHandle));
 		check(SpawnRequests.IsValidIndex(SpawnRequestHandle.GetIndex()));
@@ -162,7 +162,7 @@ public:
 	 * @return The spawn request cast in the provided template argument
 	 */
 	template<typename T> 
-	T& GetMutableSpawnRequest(FMassHandle_ActorSpawnRequest SpawnRequestHandle)
+	T& GetMutableSpawnRequest(const FMassActorSpawnRequestHandle SpawnRequestHandle)
 	{
 		check(SpawnRequestHandleManager.IsValidHandle(SpawnRequestHandle));
 		check(SpawnRequests.IsValidIndex(SpawnRequestHandle.GetIndex()));
@@ -192,7 +192,7 @@ protected:
 	/** Retrieve what would be the next best spawning request to spawn, can be overridden to have different logic
 	 *  Default implementation is the first valid request in the list, no interesting logic yet
 	 *  @return the next best handle to spawn. */
-	virtual FMassHandle_ActorSpawnRequest GetNextRequestToSpawn() const;
+	virtual FMassActorSpawnRequestHandle GetNextRequestToSpawn() const;
 
 	virtual AActor* SpawnOrRetrieveFromPool(const FStructView& SpawnRequest);
 
@@ -215,7 +215,7 @@ protected:
 
 	/** Internal generic request actor spawn to make sure the request derives from FMassActorSpawnRequest 
 	 *  @param SpawnRequest the spawn request parameters, We are allowing any type of spawn request, let's store it internally as a FInstancedStruct. This parameter is the FStructView over provide user struct */
-	FMassHandle_ActorSpawnRequest RequestActorSpawnInternal(const FConstStructView SpawnRequest);
+	FMassActorSpawnRequestHandle RequestActorSpawnInternal(const FConstStructView SpawnRequest);
 
 protected:
 
@@ -232,7 +232,7 @@ protected:
 
 	TMap<TSubclassOf<AActor>, TArray<AActor*>> PooledActors;
 
-	FMassHandleManager_ActorSpawnRequest SpawnRequestHandleManager;
+	FMassEntityHandleManager_ActorSpawnRequest SpawnRequestHandleManager;
 
 	std::atomic<uint32> RequestSerialNumberCounter;
 };
diff --git a/Engine/Plugins/Runtime/MassGameplay/Source/MassActors/Public/MassActorSubsystem.h b/Engine/Plugins/Runtime/MassGameplay/Source/MassActors/Public/MassActorSubsystem.h
--- a/Engine/Plugins/Runtime/MassGameplay/Source/MassActors/Public/MassActorSubsystem.h
+++ b/Engine/Plugins/Runtime/MassGameplay/Source/MassActors/Public/MassActorSubsystem.h
@@ -10,7 +10,7 @@
 #include "Subsystems/WorldSubsystem.h"
 #include "MassActorSubsystem.generated.h"
 
-struct FMassHandle;
+struct FMassEntityHandle;
 class AActor;
 class UMassEntitySubsystem;
 
@@ -28,7 +28,7 @@ struct MASSACTORS_API FDataFragment_Actor : public FDataFragment_ObjectWrapper
 	 * @param InActor the actor associated with the mass agent
 	 * @param bInIsOwnedByMass tell whether the actors was spawned by mass(MassVisualization) or externally(ReplicatedActors)
 	 */
-	void SetAndUpdateHandleMap(const FMassHandle MassAgent, AActor* InActor, const bool bInIsOwnedByMass);
+	void SetAndUpdateHandleMap(const FMassEntityHandle MassAgent, AActor* InActor, const bool bInIsOwnedByMass);
 
 	/** Resets the actor pointed by this fragment, will also keep the map back in MassActorSubsystem up to date */
 	void ResetAndUpdateHandleMap();
@@ -41,7 +41,7 @@ struct MASSACTORS_API FDataFragment_Actor : public FDataFragment_ObjectWrapper
 	 * @param InActor the actor associated with the mass agent
 	 * @param bInIsOwnedByMass tell whether the actors was spawned by mass(MassVisualization) or externally(ReplicatedActors)
 	 */
-	void SetNoHandleMapUpdate(const FMassHandle MassAgent, AActor* InActor, const bool bInIsOwnedByMass);
+	void SetNoHandleMapUpdate(const FMassEntityHandle MassAgent, AActor* InActor, const bool bInIsOwnedByMass);
 
 	/** Resets the actor pointed by this fragment, will NOT keep map back in MassActorSubsystem up to date.
 	 * The user needs to do the map update by hand.
@@ -82,22 +82,22 @@ class MASSACTORS_API UMassActorSubsystem : public UWorldSubsystem
 
 public:
 	/** Get mass handle from an actor */
-	FMassHandle GetHandleFromActor(const TObjectKey<const AActor> Actor);
+	FMassEntityHandle GetEntityHandleFromActor(const TObjectKey<const AActor> Actor);
 
 	/** Set the mass handle associated to an actor */
-	void SetHandleForActor(const TObjectKey<const AActor> Actor, const FMassHandle Handle);
+	void SetHandleForActor(const TObjectKey<const AActor> Actor, const FMassEntityHandle Handle);
 
 	/** Remove entry associated to an actor */
 	void RemoveHandleForActor(const TObjectKey<const AActor> Actor);
 
 	/** Get an actor pointer from a mass handle */
-	AActor* GetActorFromHandle(const FMassHandle Handle) const;
+	AActor* GetActorFromHandle(const FMassEntityHandle Handle) const;
 
 	/** 
 	 *  Removes the connection between Actor and the given entity. Does all the required book keeping 
 	 *  (as opposed to straight up RemoveHandleForActor call). If the Handle doesn't match Actor no action is taken.
 	 */
-	void DisconnectActor(const TObjectKey<const AActor> Actor, const FMassHandle Handle);
+	void DisconnectActor(const TObjectKey<const AActor> Actor, const FMassEntityHandle Handle);
 
 protected:
 	// USubsystem BEGIN
diff --git a/Engine/Plugins/Runtime/MassGameplay/Source/MassActors/Public/MassAgentComponent.h b/Engine/Plugins/Runtime/MassGameplay/Source/MassActors/Public/MassAgentComponent.h
--- a/Engine/Plugins/Runtime/MassGameplay/Source/MassActors/Public/MassAgentComponent.h
+++ b/Engine/Plugins/Runtime/MassGameplay/Source/MassActors/Public/MassAgentComponent.h
@@ -64,15 +64,15 @@ public:
 	/** @todo to enforce encapsulation we could make this protected and have a UMassAgentSubsystem friend. I'm not sure I like it though. */
 
 	/** Methods handling the state for actors that needs mass entity creation*/
-	FMassHandle GetEntityHandle() const { return AgentHandle; }
-	void SetEntityHandle(FMassHandle NewHandle);
+	FMassEntityHandle GetEntityHandle() const { return AgentHandle; }
+	void SetEntityHandle(const FMassEntityHandle NewHandle);
 	void ClearEntityHandle();
 	void EntityCreationPending();
 	void EntityCreationAborted();
 	bool IsEntityPendingCreation() const { return (State == EAgentComponentState::EntityPendingCreation); }
 
 	/** Methods handling the state of puppet actors that needs fragments initialization */
-	void SetPuppetHandle(FMassHandle NewHandle);
+	void SetPuppetHandle(const FMassEntityHandle NewHandle);
 	void PuppetInitializationPending();
 	void PuppetInitializationDone();
 	void PuppetInitializationAborted();
@@ -91,7 +91,7 @@ public:
 
 	/** Methods handling the state of a server replicated puppet */
 	void PuppetReplicationPending();
-	void SetReplicatedPuppetHandle(FMassHandle NewHandle);
+	void SetReplicatedPuppetHandle(FMassEntityHandle NewHandle);
 	void ClearReplicatedPuppetHandle();
 	void MakePuppetAReplicatedOrphan();
 
@@ -129,7 +129,7 @@ public:
 protected:
 
 	void SwitchToState(EAgentComponentState NewState);
-	void SetEntityHandleInternal(FMassHandle NewHandle);
+	void SetEntityHandleInternal(const FMassEntityHandle NewHandle);
 	void ClearEntityHandleInternal();
 	void DebugCheckStateConsistency();
 
@@ -141,14 +141,8 @@ protected:
 
 	UPROPERTY(EditAnywhere, Category = "Mass")
 	FMassEntityConfig EntityConfig;
-
-#if WITH_EDITORONLY_DATA
-	/** for debugging purposes only. The Entity used at runtime is stored in AgentHandle */
-	UPROPERTY(VisibleAnywhere, Category = "Mass")
-	FMassEntityHandle DebugEntity;
-#endif // WITH_EDITORONLY_DATA
-
-	FMassHandle AgentHandle;
+	
+	FMassEntityHandle AgentHandle;
 	FMassEntityTemplateID TemplateID;
 
 	EAgentComponentState State;
diff --git a/Engine/Plugins/Runtime/MassGameplay/Source/MassActors/Public/MassAgentSubsystem.h b/Engine/Plugins/Runtime/MassGameplay/Source/MassActors/Public/MassAgentSubsystem.h
--- a/Engine/Plugins/Runtime/MassGameplay/Source/MassActors/Public/MassAgentSubsystem.h
+++ b/Engine/Plugins/Runtime/MassGameplay/Source/MassActors/Public/MassAgentSubsystem.h
@@ -8,7 +8,6 @@
 #include "Subsystems/WorldSubsystem.h"
 #include "MassAgentSubsystem.generated.h"
 
-struct FMassHandle;
 class AActor;
 class UMassEntitySubsystem;
 class UMassSpawnerSubsystem;
diff --git a/Engine/Plugins/Runtime/MassGameplay/Source/MassCommon/Private/MassCommonTypes.cpp b/Engine/Plugins/Runtime/MassGameplay/Source/MassCommon/Private/MassCommonTypes.cpp
--- a/Engine/Plugins/Runtime/MassGameplay/Source/MassCommon/Private/MassCommonTypes.cpp
+++ b/Engine/Plugins/Runtime/MassGameplay/Source/MassCommon/Private/MassCommonTypes.cpp
@@ -103,8 +103,6 @@ FColor GetEntityDebugColor(FMassEntityHandle Entity)
 }} // namespace UE::MassDebug
 #endif // WITH_MASSGAMEPLAY_DEBUG
 
-const FMassHandle FMassHandle::InvalidHandle(UMassEntitySubsystem::InvalidEntity);
-
 //----------------------------------------------------------------------//
 // FMassUniqueFragmentCollection 
 //----------------------------------------------------------------------//
diff --git a/Engine/Plugins/Runtime/MassGameplay/Source/MassCommon/Public/MassCommonTypes.h b/Engine/Plugins/Runtime/MassGameplay/Source/MassCommon/Public/MassCommonTypes.h
--- a/Engine/Plugins/Runtime/MassGameplay/Source/MassCommon/Public/MassCommonTypes.h
+++ b/Engine/Plugins/Runtime/MassGameplay/Source/MassCommon/Public/MassCommonTypes.h
@@ -32,58 +32,6 @@ namespace UE::Mass::ProcessorGroupNames
 	const FName Tasks = FName(TEXT("Tasks"));
 }
 
-/**
- * Generic handle on mass agent which is associated to the specific Mass EntitySubsystem 
- */
-struct MASSCOMMON_API FMassHandle
-{
-public:
-	FMassHandle() = default;
-
-	explicit FMassHandle(FMassEntityHandle InLWEntity)
-		: LWEntity{ InLWEntity }
-	{}
-
-	FMassEntityHandle GetLWEntity() const
-	{
-		return LWEntity;
-	}
-
-	bool IsValid() const
-	{
-		return LWEntity.IsSet();
-	}
-
-	FString DebugGetDescription() const 
-	{
-#if WITH_MASSGAMEPLAY_DEBUG
-		return LWEntity.DebugGetDescription();
-#else
-		return FString();
-#endif // WITH_MASSGAMEPLAY_DEBUG
-	}
-
-	const static FMassHandle InvalidHandle;
-
-	bool operator==(const FMassHandle& Other) const
-	{
-		return LWEntity == Other.LWEntity;
-	}
-
-	bool operator!=(const FMassHandle& Other) const
-	{
-		return !operator==(Other);
-	}
-
-	friend uint32 GetTypeHash(const FMassHandle& Handle)
-	{
-		return GetTypeHash(Handle.GetLWEntity());
-	}
-
-private:
-	FMassEntityHandle LWEntity;
-};
-
 USTRUCT()
 struct FMassNetworkID : public FSequentialIDBase
 {
diff --git a/Engine/Plugins/Runtime/MassGameplay/Source/MassGameplayDebug/Private/MassDebuggerSubsystem.cpp b/Engine/Plugins/Runtime/MassGameplay/Source/MassGameplayDebug/Private/MassDebuggerSubsystem.cpp
--- a/Engine/Plugins/Runtime/MassGameplay/Source/MassGameplayDebug/Private/MassDebuggerSubsystem.cpp
+++ b/Engine/Plugins/Runtime/MassGameplay/Source/MassGameplayDebug/Private/MassDebuggerSubsystem.cpp
@@ -43,7 +43,7 @@ void UMassDebuggerSubsystem::OnProcessingPhaseStarted(const float DeltaSeconds)
 	PreTickProcessors();
 }
 
-void UMassDebuggerSubsystem::SetSelectedEntity(const FMassEntityHandle& InSelectedEntity)
+void UMassDebuggerSubsystem::SetSelectedEntity(const FMassEntityHandle InSelectedEntity)
 {
 	SelectedEntity = InSelectedEntity;
 	SelectedEntityDetails.Empty();
diff --git a/Engine/Plugins/Runtime/MassGameplay/Source/MassGameplayDebug/Public/MassDebuggerSubsystem.h b/Engine/Plugins/Runtime/MassGameplay/Source/MassGameplayDebug/Public/MassDebuggerSubsystem.h
--- a/Engine/Plugins/Runtime/MassGameplay/Source/MassGameplayDebug/Public/MassDebuggerSubsystem.h
+++ b/Engine/Plugins/Runtime/MassGameplay/Source/MassGameplayDebug/Public/MassDebuggerSubsystem.h
@@ -35,7 +35,7 @@ public:
 	TConstArrayView<FVector> GetLocations() const { return Locations; }
 
 	FMassEntityHandle GetSelectedEntity() const { return SelectedEntity; }
-	void SetSelectedEntity(const FMassEntityHandle& InSelectedEntity);
+	void SetSelectedEntity(const FMassEntityHandle InSelectedEntity);
 
 	void AppendSelectedEntityInfo(const FString& Info);
 	const FString& GetSelectedEntityInfo() const { return SelectedEntityDetails; }
diff --git a/Engine/Plugins/Runtime/MassGameplay/Source/MassReplication/Private/MassReplicationManager.cpp b/Engine/Plugins/Runtime/MassGameplay/Source/MassReplication/Private/MassReplicationManager.cpp
--- a/Engine/Plugins/Runtime/MassGameplay/Source/MassReplication/Private/MassReplicationManager.cpp
+++ b/Engine/Plugins/Runtime/MassGameplay/Source/MassReplication/Private/MassReplicationManager.cpp
@@ -50,10 +50,10 @@ void UMassReplicationManager::Deinitialize()
 	EntitySystem = nullptr;
 }
 
-FMassNetworkID UMassReplicationManager::GetNetIDFromHandle(const FMassHandle& Handle) const
+FMassNetworkID UMassReplicationManager::GetNetIDFromHandle(const FMassEntityHandle Handle) const
 {
 	check(EntitySystem);
-	const FMassNetworkIDFragment& Data = EntitySystem->GetComponentDataChecked<FMassNetworkIDFragment>(Handle.GetLWEntity());
+	const FMassNetworkIDFragment& Data = EntitySystem->GetComponentDataChecked<FMassNetworkIDFragment>(Handle);
 	return Data.NetID;
 }
 
@@ -238,9 +238,9 @@ void UMassReplicationManager::SynchronizeClientViewers(const TArray<FViewerInfo>
 {
 	check(MassLODManager);
 
-	struct FMassHandleData
+	struct FMassClientViewerHandle
 	{
-		FMassHandleData(FMassClientHandle InClientHandle, FMassViewerHandle InViewerHandle)
+		FMassClientViewerHandle(FMassClientHandle InClientHandle, FMassViewerHandle InViewerHandle)
 			: ClientHandle(InClientHandle)
 			, ViewerHandle(InViewerHandle)
 		{}
@@ -250,7 +250,7 @@ void UMassReplicationManager::SynchronizeClientViewers(const TArray<FViewerInfo>
 	};
 
 	// go through the ClientsViewerHandles and check validity and store the valid ones into a map
-	typedef TMap<APlayerController*, FMassHandleData> FViewerMap;
+	typedef TMap<APlayerController*, FMassClientViewerHandle> FViewerMap;
 
 	FViewerMap ClientViewerMap;
 	{
@@ -280,7 +280,7 @@ void UMassReplicationManager::SynchronizeClientViewers(const TArray<FViewerInfo>
 						// we don't verify or remove Parent Connections here as that was done when we synchronized the client bubbles
 						if (ConnectionType == UE::Mass::Replication::ENetConnectionType::Child)
 						{
-							ClientViewerMap.Add(Controller, FMassHandleData(ClientHandle, ClientViewer));
+							ClientViewerMap.Add(Controller, FMassClientViewerHandle(ClientHandle, ClientViewer));
 						}
 					}
 					else //remove invalid ClientViewer, but dont increment the ViewerIdx
@@ -331,7 +331,7 @@ void UMassReplicationManager::SynchronizeClientViewers(const TArray<FViewerInfo>
 		// anything left in the map needs to be removed from the list
 		for (FViewerMap::TIterator Itr = ClientViewerMap.CreateIterator(); Itr; ++Itr)
 		{
-			const FMassHandleData& HandleData = Itr->Value;
+			const FMassClientViewerHandle& HandleData = Itr->Value;
 
 			FClientViewerHandles& ClientViewers = ClientsViewerHandles[HandleData.ClientHandle.GetIndex()];
 
diff --git a/Engine/Plugins/Runtime/MassGameplay/Source/MassReplication/Public/MassReplicationManager.h b/Engine/Plugins/Runtime/MassGameplay/Source/MassReplication/Public/MassReplicationManager.h
--- a/Engine/Plugins/Runtime/MassGameplay/Source/MassReplication/Public/MassReplicationManager.h
+++ b/Engine/Plugins/Runtime/MassGameplay/Source/MassReplication/Public/MassReplicationManager.h
@@ -89,7 +89,7 @@ public:
 	}
 
 	/** Retrieve network id from Mass handle. */
-	FMassNetworkID GetNetIDFromHandle(const FMassHandle& Handle) const;
+	FMassNetworkID GetNetIDFromHandle(const FMassEntityHandle Handle) const;
 
 	static inline FMassNetworkID GetNextAvailableMassNetID() { return FMassNetworkID(++CurrentNetMassCounter); }
 
@@ -286,6 +286,9 @@ protected:
 	UE::MassReplication::FMassAgentDelegate OnRemovingMassAgent;
 #endif // UE_REPLICATION_COMPILE_CLIENT_CODE
 
-	/** Used to make sure the FMassHandle_ClientReplications are synchronized immediately before they are needed*/
+	/**
+	 * Used to make sure the FMassEntityHandle_ClientReplications are synchronized immediately before they are needed
+	 * @todo this comment is no longer accurate, needs fixing
+	 */
 	uint64 LastSynchronizedFrame = 0;
 };
\ No newline at end of file
diff --git a/Engine/Plugins/Runtime/MassGameplay/Source/MassRepresentation/Private/MassRepresentationProcessor.cpp b/Engine/Plugins/Runtime/MassGameplay/Source/MassRepresentation/Private/MassRepresentationProcessor.cpp
--- a/Engine/Plugins/Runtime/MassGameplay/Source/MassRepresentation/Private/MassRepresentationProcessor.cpp
+++ b/Engine/Plugins/Runtime/MassGameplay/Source/MassRepresentation/Private/MassRepresentationProcessor.cpp
@@ -69,8 +69,7 @@ void UMassRepresentationProcessor::UpdateRepresentation(FMassExecutionContext& C
 	const int32 NumEntities = Context.GetEntitiesNum();
 	for (int32 EntityIdx = 0; EntityIdx < NumEntities; EntityIdx++)
 	{
-		const FMassEntityHandle Entity = Context.GetEntity(EntityIdx);
-		const FMassHandle MassAgent(Entity);
+		const FMassEntityHandle MassAgent = Context.GetEntity(EntityIdx);
 		const FDataFragment_Transform& TransformFragment = TransformList[EntityIdx];
 		const FMassRepresentationLODFragment& RepresentationLOD = RepresentationLODList[EntityIdx];
 		FMassRepresentationFragment& Representation = RepresentationList[EntityIdx];
@@ -218,7 +217,7 @@ void UMassRepresentationProcessor::Execute(UMassEntitySubsystem& EntitySubsystem
 	});
 }
 
-AActor* UMassRepresentationProcessor::GetOrSpawnActor(const FMassHandle MassAgent, FDataFragment_Actor& ActorInfo, const FTransform& Transform, const int16 TemplateActorIndex, FMassHandle_ActorSpawnRequest& SpawnRequestHandle, const float Priority)
+AActor* UMassRepresentationProcessor::GetOrSpawnActor(const FMassEntityHandle MassAgent, FDataFragment_Actor& ActorInfo, const FTransform& Transform, const int16 TemplateActorIndex, FMassActorSpawnRequestHandle& SpawnRequestHandle, const float Priority)
 {
 	check(RepresentationSubsystem);
 
@@ -227,7 +226,7 @@ AActor* UMassRepresentationProcessor::GetOrSpawnActor(const FMassHandle MassAgen
 		FMassActorPostSpawnDelegate::CreateUObject(this, &UMassRepresentationProcessor::OnActorPostSpawn));
 }
 
-bool UMassRepresentationProcessor::ReleaseActorOrCancelSpawning(const FMassHandle MassAgent, FDataFragment_Actor& ActorInfo, const int16 TemplateActorIndex, FMassHandle_ActorSpawnRequest& SpawnRequestHandle, FMassExecutionContext& Context, bool bCancelSpawningOnly /*= false*/)
+bool UMassRepresentationProcessor::ReleaseActorOrCancelSpawning(const FMassEntityHandle MassAgent, FDataFragment_Actor& ActorInfo, const int16 TemplateActorIndex, FMassActorSpawnRequestHandle& SpawnRequestHandle, FMassExecutionContext& Context, bool bCancelSpawningOnly /*= false*/)
 {
 	check(RepresentationSubsystem);
 	check(!ActorInfo.IsValid() || ActorInfo.IsOwnedByMass());
@@ -286,13 +285,13 @@ void UMassRepresentationProcessor::TeleportActor(const FTransform& Transform, AA
 	}
 }
 
-void UMassRepresentationProcessor::ReleaseAnyActorOrCancelAnySpawning(UMassRepresentationSubsystem& RepresentationSubsystem, const FMassHandle MassAgent)
+void UMassRepresentationProcessor::ReleaseAnyActorOrCancelAnySpawning(UMassRepresentationSubsystem& RepresentationSubsystem, const FMassEntityHandle MassAgent)
 {
 	if (const UMassEntitySubsystem* EntitySubsystem = UWorld::GetSubsystem<UMassEntitySubsystem>(RepresentationSubsystem.GetWorld()))
 	{
-		if (FDataFragment_Actor* ActorInfo = EntitySubsystem->GetComponentDataPtr<FDataFragment_Actor>(MassAgent.GetLWEntity()))
+		if (FDataFragment_Actor* ActorInfo = EntitySubsystem->GetComponentDataPtr<FDataFragment_Actor>(MassAgent))
 		{
-			if (FMassRepresentationFragment* Representation = EntitySubsystem->GetComponentDataPtr<FMassRepresentationFragment>(MassAgent.GetLWEntity()))
+			if (FMassRepresentationFragment* Representation = EntitySubsystem->GetComponentDataPtr<FMassRepresentationFragment>(MassAgent))
 			{
 				ReleaseAnyActorOrCancelAnySpawning(RepresentationSubsystem, MassAgent, *ActorInfo, *Representation);
 			}
@@ -300,7 +299,7 @@ void UMassRepresentationProcessor::ReleaseAnyActorOrCancelAnySpawning(UMassRepre
 	}
 }
 
-void UMassRepresentationProcessor::ReleaseAnyActorOrCancelAnySpawning(UMassRepresentationSubsystem& RepresentationSubsystem, const FMassHandle MassAgent, FDataFragment_Actor& ActorInfo, FMassRepresentationFragment& Representation)
+void UMassRepresentationProcessor::ReleaseAnyActorOrCancelAnySpawning(UMassRepresentationSubsystem& RepresentationSubsystem, const FMassEntityHandle MassAgent, FDataFragment_Actor& ActorInfo, FMassRepresentationFragment& Representation)
 {
 	// This method can only release owned by mass actors
 	AActor* Actor = ActorInfo.GetOwnedByMassMutable();
@@ -387,13 +386,13 @@ void UMassRepresentationProcessor::UpdateEntityVisibility(const FMassEntityHandl
 	}
 }
 
-void UMassRepresentationProcessor::OnActorPreSpawn(const FMassHandle_ActorSpawnRequest& SpawnRequestHandle, const FStructView& SpawnRequest)
+void UMassRepresentationProcessor::OnActorPreSpawn(const FMassActorSpawnRequestHandle& SpawnRequestHandle, const FStructView& SpawnRequest)
 {
 	check(PipeEntitySubsystem);
 	check(RepresentationSubsystem);
 
 	const FMassActorSpawnRequest& MassActorSpawnRequest = SpawnRequest.Get<FMassActorSpawnRequest>();
-	const FMassEntityView View(*PipeEntitySubsystem, MassActorSpawnRequest.MassAgent.GetLWEntity());
+	const FMassEntityView View(*PipeEntitySubsystem, MassActorSpawnRequest.MassAgent);
 
 	if (FDataFragment_Actor* ActorInfo = View.GetComponentDataPtr<FDataFragment_Actor>())
 	{
@@ -421,7 +420,7 @@ void UMassRepresentationProcessor::OnActorPreSpawn(const FMassHandle_ActorSpawnR
 	}
 }
 
-EMassActorSpawnRequestAction UMassRepresentationProcessor::OnActorPostSpawn(const FMassHandle_ActorSpawnRequest& SpawnRequestHandle, const FStructView& SpawnRequest)
+EMassActorSpawnRequestAction UMassRepresentationProcessor::OnActorPostSpawn(const FMassActorSpawnRequestHandle& SpawnRequestHandle, const FStructView& SpawnRequest)
 {
 	check(PipeEntitySubsystem);
 
@@ -429,7 +428,7 @@ EMassActorSpawnRequestAction UMassRepresentationProcessor::OnActorPostSpawn(cons
 	checkf(MassActorSpawnRequest.SpawnedActor, TEXT("Expecting valid spawned actor"));
 
 	// Might be already done if the actor has a MassAgentComponent via the callback OnMassAgentComponentEntityAssociated on the MassRepresentationSubsystem
-	FDataFragment_Actor& ActorInfo = PipeEntitySubsystem->GetComponentDataChecked<FDataFragment_Actor>(MassActorSpawnRequest.MassAgent.GetLWEntity());
+	FDataFragment_Actor& ActorInfo = PipeEntitySubsystem->GetComponentDataChecked<FDataFragment_Actor>(MassActorSpawnRequest.MassAgent);
 	if (ActorInfo.IsValid())
 	{
 		// If already set, make sure it is pointing to the same actor.
@@ -484,8 +483,7 @@ void UMassRepresentationFragmentDestructor::Execute(UMassEntitySubsystem& Entity
 			FMassRepresentationFragment& Representation = RepresentationList[i];
 			FDataFragment_Actor& ActorInfo = ActorList[i];
 
-			const FMassEntityHandle Entity = Context.GetEntity(i);
-			const FMassHandle MassAgent(Entity);
+			const FMassEntityHandle MassAgent = Context.GetEntity(i);
 
 			UMassRepresentationProcessor::ReleaseAnyActorOrCancelAnySpawning(*RepresentationSubsystem, MassAgent, ActorInfo, Representation);
 		}
diff --git a/Engine/Plugins/Runtime/MassGameplay/Source/MassRepresentation/Private/MassRepresentationSubsystem.cpp b/Engine/Plugins/Runtime/MassGameplay/Source/MassRepresentation/Private/MassRepresentationSubsystem.cpp
--- a/Engine/Plugins/Runtime/MassGameplay/Source/MassRepresentation/Private/MassRepresentationSubsystem.cpp
+++ b/Engine/Plugins/Runtime/MassGameplay/Source/MassRepresentation/Private/MassRepresentationSubsystem.cpp
@@ -49,7 +49,7 @@ int16 UMassRepresentationSubsystem::FindOrAddTemplateActor(const TSubclassOf<AAc
 	return (int16)VisualIndex;
 }
 
-AActor* UMassRepresentationSubsystem::GetOrSpawnActorFromTemplate(const FMassHandle MassAgent, const FTransform& Transform, const int16 TemplateActorIndex, FMassHandle_ActorSpawnRequest& SpawnRequestHandle, float Priority /*= MAX_FLT*/, 
+AActor* UMassRepresentationSubsystem::GetOrSpawnActorFromTemplate(const FMassEntityHandle MassAgent, const FTransform& Transform, const int16 TemplateActorIndex, FMassActorSpawnRequestHandle& SpawnRequestHandle, float Priority /*= MAX_FLT*/, 
 	FMassActorPreSpawnDelegate ActorPreSpawnDelegate /*= FMassActorPreSpawnDelegate()*/, FMassActorPostSpawnDelegate ActorPostSpawnDelegate /*= FMassActorPostSpawnDelegate()*/)
 {
 	UE_MT_SCOPED_READ_ACCESS(TemplateActorsMTAccessDetector);
@@ -176,7 +176,7 @@ void UMassRepresentationSubsystem::ReleaseAllResources()
 	VisualizationComponent->ClearAllVisualInstances();
 }
 
-bool UMassRepresentationSubsystem::ReleaseTemplateActorOrCancelSpawning(const FMassHandle MassAgent, const int16 TemplateActorIndex, AActor* ActorToRelease, FMassHandle_ActorSpawnRequest& SpawnRequestHandle)
+bool UMassRepresentationSubsystem::ReleaseTemplateActorOrCancelSpawning(const FMassEntityHandle MassAgent, const int16 TemplateActorIndex, AActor* ActorToRelease, FMassActorSpawnRequestHandle& SpawnRequestHandle)
 {
 	// First try to cancel the spawning request, then try to release the actor
 	if( CancelSpawningInternal(TemplateActorIndex, SpawnRequestHandle) ||
@@ -194,7 +194,7 @@ bool UMassRepresentationSubsystem::ReleaseTemplateActorOrCancelSpawning(const FM
 	return false;
 }
 
-bool UMassRepresentationSubsystem::ReleaseTemplateActor(const FMassHandle MassAgent, const int16 TemplateActorIndex, AActor* ActorToRelease, bool bImmediate)
+bool UMassRepresentationSubsystem::ReleaseTemplateActor(const FMassEntityHandle MassAgent, const int16 TemplateActorIndex, AActor* ActorToRelease, bool bImmediate)
 {
 	if (ReleaseTemplateActorInternal(TemplateActorIndex, ActorToRelease, bImmediate))
 	{
@@ -210,7 +210,7 @@ bool UMassRepresentationSubsystem::ReleaseTemplateActor(const FMassHandle MassAg
 	return false;
 }
 
-bool UMassRepresentationSubsystem::CancelSpawning(const FMassHandle MassAgent, const int16 TemplateActorIndex, FMassHandle_ActorSpawnRequest& SpawnRequestHandle)
+bool UMassRepresentationSubsystem::CancelSpawning(const FMassEntityHandle MassAgent, const int16 TemplateActorIndex, FMassActorSpawnRequestHandle& SpawnRequestHandle)
 {
 	if (CancelSpawningInternal(TemplateActorIndex, SpawnRequestHandle))
 	{
@@ -244,7 +244,7 @@ bool UMassRepresentationSubsystem::ReleaseTemplateActorInternal(const int16 Temp
 	return true;
 }
 
-bool UMassRepresentationSubsystem::CancelSpawningInternal(const int16 TemplateActorIndex, FMassHandle_ActorSpawnRequest& SpawnRequestHandle)
+bool UMassRepresentationSubsystem::CancelSpawningInternal(const int16 TemplateActorIndex, FMassActorSpawnRequestHandle& SpawnRequestHandle)
 {
 	UE_MT_SCOPED_READ_ACCESS(TemplateActorsMTAccessDetector);
 	if (!TemplateActors.IsValidIndex(TemplateActorIndex))
@@ -372,9 +372,9 @@ void UMassRepresentationSubsystem::OnMassAgentComponentEntityAssociated(const UM
 {
 	check(EntitySubsystem);
 
-	const FMassHandle MassAgent = AgentComponent.GetEntityHandle();
-	checkf(EntitySubsystem->IsEntityValid(MassAgent.GetLWEntity()), TEXT("Expecting a valid mass entity"));
-	if (EntitySubsystem->IsEntityValid(MassAgent.GetLWEntity()) && AgentComponent.IsNetSimulating())
+	const FMassEntityHandle MassAgent = AgentComponent.GetEntityHandle();
+	checkf(EntitySubsystem->IsEntityValid(MassAgent), TEXT("Expecting a valid mass entity"));
+	if (EntitySubsystem->IsEntityValid(MassAgent) && AgentComponent.IsNetSimulating())
 	{
 		// Check if this mass agent already handled by this sub system, if yes than release any local spawned actor or cancel any spawn requests
 		if (HandledMassAgents.Find(MassAgent))
@@ -391,11 +391,11 @@ void UMassRepresentationSubsystem::OnMassAgentComponentEntityDetaching(const UMa
 	AActor* ComponentOwner = AgentComponent.GetOwner();
 	check(ComponentOwner);
 
-	const FMassHandle MassAgent = AgentComponent.GetEntityHandle();
-	checkf(EntitySubsystem->IsEntityValid(MassAgent.GetLWEntity()), TEXT("Expecting a valid mass entity"));
-	if (EntitySubsystem->IsEntityValid(MassAgent.GetLWEntity()) && AgentComponent.IsNetSimulating())
+	const FMassEntityHandle MassAgent = AgentComponent.GetEntityHandle();
+	checkf(EntitySubsystem->IsEntityValid(MassAgent), TEXT("Expecting a valid mass entity"));
+	if (EntitySubsystem->IsEntityValid(MassAgent) && AgentComponent.IsNetSimulating())
 	{
-		const FMassEntityView EntityView(*EntitySubsystem, MassAgent.GetLWEntity());
+		const FMassEntityView EntityView(*EntitySubsystem, MassAgent);
 		if (FMassRepresentationFragment* Representation = EntityView.GetComponentDataPtr<FMassRepresentationFragment>())
 		{
 			// Force a reevaluate of the current representation
diff --git a/Engine/Plugins/Runtime/MassGameplay/Source/MassRepresentation/Public/MassRepresentationFragments.h b/Engine/Plugins/Runtime/MassGameplay/Source/MassRepresentation/Public/MassRepresentationFragments.h
--- a/Engine/Plugins/Runtime/MassGameplay/Source/MassRepresentation/Public/MassRepresentationFragments.h
+++ b/Engine/Plugins/Runtime/MassGameplay/Source/MassRepresentation/Public/MassRepresentationFragments.h
@@ -47,7 +47,7 @@ struct MASSREPRESENTATION_API FMassRepresentationFragment : public FMassFragment
 
 	int16 StaticMeshDescIndex = INDEX_NONE;
 
-	FMassHandle_ActorSpawnRequest ActorSpawnRequestHandle;
+	FMassActorSpawnRequestHandle ActorSpawnRequestHandle;
 
 	FTransform PrevTransform;
 
diff --git a/Engine/Plugins/Runtime/MassGameplay/Source/MassRepresentation/Public/MassRepresentationProcessor.h b/Engine/Plugins/Runtime/MassGameplay/Source/MassRepresentation/Public/MassRepresentationProcessor.h
--- a/Engine/Plugins/Runtime/MassGameplay/Source/MassRepresentation/Public/MassRepresentationProcessor.h
+++ b/Engine/Plugins/Runtime/MassGameplay/Source/MassRepresentation/Public/MassRepresentationProcessor.h
@@ -54,7 +54,7 @@ protected:
 	 * @param Priority of this spawn request in comparison with the others, lower value means higher priority
 	 * @return the actor spawned
 	 */
-	virtual AActor* GetOrSpawnActor(const FMassHandle MassAgent, FDataFragment_Actor& ActorInfo, const FTransform& Transform, const int16 TemplateActorIndex, FMassHandle_ActorSpawnRequest& SpawnRequestHandle, const float Priority);
+	virtual AActor* GetOrSpawnActor(const FMassEntityHandle MassAgent, FDataFragment_Actor& ActorInfo, const FTransform& Transform, const int16 TemplateActorIndex, FMassActorSpawnRequestHandle& SpawnRequestHandle, const float Priority);
 
 	/**
 	 * Release the actor to the subsystem, will only release it the actor or spawn request matches the template actor
@@ -66,7 +66,7 @@ protected:
 	 * @param bCancelSpawningOnly tell to only cancel the existing spawning request and to not release the associated actor it any.
 	 * @return if the actor was release or the spawning was canceled.
 	 */
-	virtual bool ReleaseActorOrCancelSpawning(const FMassHandle MassAgent, FDataFragment_Actor& ActorInfo, const int16 TemplateActorIndex, FMassHandle_ActorSpawnRequest& SpawnRequestHandle, FMassExecutionContext& Context, const bool bCancelSpawningOnly = false);
+	virtual bool ReleaseActorOrCancelSpawning(const FMassEntityHandle MassAgent, FDataFragment_Actor& ActorInfo, const int16 TemplateActorIndex, FMassActorSpawnRequestHandle& SpawnRequestHandle, FMassExecutionContext& Context, const bool bCancelSpawningOnly = false);
 
 	/** 
 	 * Enable/disable a spawned actor
@@ -91,7 +91,7 @@ protected:
 	 * @param SpawnRequestHandle the handle of the spawn request that is about to spawn
 	 * @param SpawnRequest of the actor that is about to spawn
 	 */
-	void OnActorPreSpawn(const FMassHandle_ActorSpawnRequest& SpawnRequestHandle, const FStructView& SpawnRequest);
+	void OnActorPreSpawn(const FMassActorSpawnRequestHandle& SpawnRequestHandle, const FStructView& SpawnRequest);
 
 	/**
 	 * Method that will be bound to a delegate used post-spawn to notify and let the requester configure the actor
@@ -99,14 +99,14 @@ protected:
 	 * @param SpawnRequest of the actor that just spawned
 	 * @return The action to take on the spawn request, either keep it there or remove it.
 	 */
-	EMassActorSpawnRequestAction OnActorPostSpawn(const FMassHandle_ActorSpawnRequest& SpawnRequestHandle, const FStructView& SpawnRequest);
+	EMassActorSpawnRequestAction OnActorPostSpawn(const FMassActorSpawnRequestHandle& SpawnRequestHandle, const FStructView& SpawnRequest);
 
 	/**
 	 * Method called once the spawning of an actor just occurred, give the processors a chance to configure the actor during the spawn time slicing
 	 * @param MassAgent mass entity to the spawn actor
 	 * @param SpawnedActor that was just spawned
 	 */
-	virtual void PostActorSpawned(const FMassHandle MassAgent, AActor& SpawnedActor) {}
+	virtual void PostActorSpawned(const FMassEntityHandle MassAgent, AActor& SpawnedActor) {}
 
 	/**
 	 * Updates chunk visibility info for later chunk logic optimization
@@ -134,7 +134,7 @@ public:
 	 * @param MassAgent is the handle to the associated mass agent
 	 * @return True if actor was release or spawning request was canceled
 	 */
-	static void ReleaseAnyActorOrCancelAnySpawning(UMassRepresentationSubsystem& RepresentationSubsystem, const FMassHandle MassAgent);
+	static void ReleaseAnyActorOrCancelAnySpawning(UMassRepresentationSubsystem& RepresentationSubsystem, const FMassEntityHandle MassAgent);
 
 	/**
 	 * Release an actor or cancel its spawning
@@ -145,7 +145,7 @@ public:
 	 * @param ActorInfo is the fragment where we are going to store the actor pointer
 	 * @param Representation fragment containing the current and previous visual state
 	 */
-	static void ReleaseAnyActorOrCancelAnySpawning(UMassRepresentationSubsystem& RepresentationSubsystem, const FMassHandle MassAgent, FDataFragment_Actor& ActorInfo, FMassRepresentationFragment& Representation);
+	static void ReleaseAnyActorOrCancelAnySpawning(UMassRepresentationSubsystem& RepresentationSubsystem, const FMassEntityHandle MassAgent, FDataFragment_Actor& ActorInfo, FMassRepresentationFragment& Representation);
 
 protected:
 
diff --git a/Engine/Plugins/Runtime/MassGameplay/Source/MassRepresentation/Public/MassRepresentationSubsystem.h b/Engine/Plugins/Runtime/MassGameplay/Source/MassRepresentation/Public/MassRepresentationSubsystem.h
--- a/Engine/Plugins/Runtime/MassGameplay/Source/MassRepresentation/Public/MassRepresentationSubsystem.h
+++ b/Engine/Plugins/Runtime/MassGameplay/Source/MassRepresentation/Public/MassRepresentationSubsystem.h
@@ -15,7 +15,7 @@ class AMassVisualizer;
 class UHierarchicalInstancedStaticMeshComponent;
 struct FStaticMeshInstanceVisualizationDesc;
 struct FMassInstancedStaticMeshInfo;
-struct FMassHandle_ActorSpawnRequest;
+struct FMassActorSpawnRequestHandle;
 class UMassActorSpawnerSubsystem;
 class UMassAgentComponent;
 class UMassEntitySubsystem;
@@ -66,7 +66,7 @@ public:
 	 * @param ActorPostSpawnDelegate is an optional delegate called once the actor is spawned
 	 * @return The spawned actor from the template actor type if ready
 	 */
-	AActor* GetOrSpawnActorFromTemplate(const FMassHandle MassAgent, const FTransform& Transform, const int16 TemplateActorIndex, FMassHandle_ActorSpawnRequest& SpawnRequestHandle, float Priority = MAX_FLT,
+	AActor* GetOrSpawnActorFromTemplate(const FMassEntityHandle MassAgent, const FTransform& Transform, const int16 TemplateActorIndex, FMassActorSpawnRequestHandle& SpawnRequestHandle, float Priority = MAX_FLT,
 		FMassActorPreSpawnDelegate ActorPreSpawnDelegate = FMassActorPreSpawnDelegate(), FMassActorPostSpawnDelegate ActorPostSpawnDelegate = FMassActorPostSpawnDelegate());
 
 	/**
@@ -76,7 +76,7 @@ public:
 	 * @param SpawnRequestHandle [IN/OUT] previously requested spawn, gets invalidated as a result of this call.
 	 * @return True if spawning request was canceled
 	 */
-	bool CancelSpawning(const FMassHandle MassAgent, const int16 TemplateActorIndex, FMassHandle_ActorSpawnRequest & SpawnRequestHandle);
+	bool CancelSpawning(const FMassEntityHandle MassAgent, const int16 TemplateActorIndex, FMassActorSpawnRequestHandle & SpawnRequestHandle);
 
 	/**
 	 * Release an actor that is matching the TemplateActorIndex
@@ -86,7 +86,7 @@ public:
 	 * @param bImmediate means it needs to be done immediately and not queue for later
 	 * @return True if actor was released
 	 */
-	bool ReleaseTemplateActor(const FMassHandle MassAgent, const int16 TemplateActorIndex, AActor* ActorToRelease, bool bImmediate);
+	bool ReleaseTemplateActor(const FMassEntityHandle MassAgent, const int16 TemplateActorIndex, AActor* ActorToRelease, bool bImmediate);
 
 	/**
 	 * Release an actor or cancel its spawning if it is matching the TemplateActorIndex
@@ -96,7 +96,7 @@ public:
 	 * @param SpawnRequestHandle [IN/OUT] previously requested spawn, gets invalidated as a result of this call.
 	 * @return True if actor was released or spawning request was canceled
 	 */
-	bool ReleaseTemplateActorOrCancelSpawning(const FMassHandle MassAgent, const int16 TemplateActorIndex, AActor* ActorToRelease, FMassHandle_ActorSpawnRequest& SpawnRequestHandle);
+	bool ReleaseTemplateActorOrCancelSpawning(const FMassEntityHandle MassAgent, const int16 TemplateActorIndex, AActor* ActorToRelease, FMassActorSpawnRequestHandle& SpawnRequestHandle);
 
 
 	/**
@@ -129,7 +129,7 @@ protected:
 	void OnMassAgentComponentEntityDetaching(const UMassAgentComponent& AgentComponent);
 
 	bool ReleaseTemplateActorInternal(const int16 TemplateActorIndex, AActor* ActorToRelease, bool bImmediate);
-	bool CancelSpawningInternal(const int16 TemplateActorIndex, FMassHandle_ActorSpawnRequest& SpawnRequestHandle);
+	bool CancelSpawningInternal(const int16 TemplateActorIndex, FMassActorSpawnRequestHandle& SpawnRequestHandle);
 
 protected:
 
@@ -162,6 +162,6 @@ protected:
 	float RetryTimeInterval = 10.0f;
 
 	/** Keeping track of all the mass agent this subsystem is responsible for spawning actors */
-	TMap<FMassHandle, int32> HandledMassAgents;
+	TMap<FMassEntityHandle, int32> HandledMassAgents;
 };
 
diff --git a/Engine/Plugins/Runtime/MassGameplay/Source/MassSmartObjects/Public/MassSmartObjectBehaviorConfig.h b/Engine/Plugins/Runtime/MassGameplay/Source/MassSmartObjects/Public/MassSmartObjectBehaviorConfig.h
--- a/Engine/Plugins/Runtime/MassGameplay/Source/MassSmartObjects/Public/MassSmartObjectBehaviorConfig.h
+++ b/Engine/Plugins/Runtime/MassGameplay/Source/MassSmartObjects/Public/MassSmartObjectBehaviorConfig.h
@@ -3,11 +3,11 @@
 #pragma once
 
 #include "SmartObjectConfig.h"
+#include "MassEntityTypes.h"
 #include "MassSmartObjectBehaviorConfig.generated.h"
 
 class UMassEntitySubsystem;
 class USmartObjectSubsystem;
-struct FMassEntityHandle;
 struct FMassExecutionContext;
 struct FDataFragment_Transform;
 struct FDataFragment_SmartObjectUser;
@@ -19,11 +19,11 @@ struct MASSSMARTOBJECTS_API FMassBehaviorEntityContext
 {
 	FMassBehaviorEntityContext() = delete;
 
-	FMassBehaviorEntityContext(const FMassEntityHandle& InEntity, const FDataFragment_Transform& InTransformFragment, FDataFragment_SmartObjectUser& InSOUser, USmartObjectSubsystem& InSubsystem)
+	FMassBehaviorEntityContext(const FMassEntityHandle InEntity, const FDataFragment_Transform& InTransformFragment, FDataFragment_SmartObjectUser& InSOUser, USmartObjectSubsystem& InSubsystem)
 		: Entity(InEntity), TransformFragment(InTransformFragment), SOUser(InSOUser), Subsystem(InSubsystem)
 	{}
 
-	const FMassEntityHandle& Entity;
+	const FMassEntityHandle Entity;
 	const FDataFragment_Transform& TransformFragment;
 	FDataFragment_SmartObjectUser& SOUser;
 	USmartObjectSubsystem& Subsystem;
diff --git a/Engine/Plugins/Runtime/MassGameplay/Source/MassSmartObjects/Public/MassSmartObjectTypes.h b/Engine/Plugins/Runtime/MassGameplay/Source/MassSmartObjects/Public/MassSmartObjectTypes.h
--- a/Engine/Plugins/Runtime/MassGameplay/Source/MassSmartObjects/Public/MassSmartObjectTypes.h
+++ b/Engine/Plugins/Runtime/MassGameplay/Source/MassSmartObjects/Public/MassSmartObjectTypes.h
@@ -22,7 +22,7 @@ struct MASSSMARTOBJECTS_API FMassSmartObjectRequestID
 	GENERATED_BODY()
 
 	FMassSmartObjectRequestID() = default;
-	FMassSmartObjectRequestID(const FMassEntityHandle& InEntity) : Entity(InEntity) {}
+	FMassSmartObjectRequestID(const FMassEntityHandle InEntity) : Entity(InEntity) {}
 	
 	bool IsSet() const { return Entity.IsSet(); }
 	void Reset() { Entity.Reset(); }
