Renamed all Component-containing functions and member variable names in Mass, replaced with "Fragment". Updated a lot of comments and logged text as well.

diff --git a/Engine/Plugins/Runtime/MassGameplay/Source/MassSmartObjects/Private/MassSmartObjectProcessor.cpp b/Engine/Plugins/Runtime/MassGameplay/Source/MassSmartObjects/Private/MassSmartObjectProcessor.cpp
--- a/Engine/Plugins/Runtime/MassGameplay/Source/MassSmartObjects/Private/MassSmartObjectProcessor.cpp
+++ b/Engine/Plugins/Runtime/MassGameplay/Source/MassSmartObjects/Private/MassSmartObjectProcessor.cpp
@@ -98,8 +98,8 @@ void UMassProcessor_SmartObjectCandidatesFinder::Execute(UMassEntitySubsystem& E
 		const int32 NumEntities = Context.GetNumEntities();
 		EntitiesToSignal.Reserve(EntitiesToSignal.Num() + NumEntities);
 
-		const TConstArrayView<FMassSmartObjectWorldLocationRequestFragment> RequestList = Context.GetComponentView<FMassSmartObjectWorldLocationRequestFragment>();
-		const TArrayView<FMassSmartObjectRequestResultFragment> ResultList = Context.GetMutableComponentView<FMassSmartObjectRequestResultFragment>();
+		const TConstArrayView<FMassSmartObjectWorldLocationRequestFragment> RequestList = Context.GetFragmentView<FMassSmartObjectWorldLocationRequestFragment>();
+		const TArrayView<FMassSmartObjectRequestResultFragment> ResultList = Context.GetMutableFragmentView<FMassSmartObjectRequestResultFragment>();
 
 		for (int32 i = 0; i < NumEntities; ++i)
 		{
@@ -157,8 +157,8 @@ void UMassProcessor_SmartObjectCandidatesFinder::Execute(UMassEntitySubsystem& E
 			const int32 NumEntities = Context.GetNumEntities();
 			EntitiesToSignal.Reserve(EntitiesToSignal.Num() + NumEntities);
 
-			TConstArrayView<FMassSmartObjectLaneLocationRequestFragment> RequestList = Context.GetComponentView<FMassSmartObjectLaneLocationRequestFragment>();
-			TArrayView<FMassSmartObjectRequestResultFragment> ResultList = Context.GetMutableComponentView<FMassSmartObjectRequestResultFragment>();
+			TConstArrayView<FMassSmartObjectLaneLocationRequestFragment> RequestList = Context.GetFragmentView<FMassSmartObjectLaneLocationRequestFragment>();
+			TArrayView<FMassSmartObjectRequestResultFragment> ResultList = Context.GetMutableFragmentView<FMassSmartObjectRequestResultFragment>();
 
 			// Cache latest used data since request are most of the time on the same zone graph
 			FZoneGraphDataHandle LastUsedDataHandle;
@@ -290,7 +290,7 @@ void UMassProcessor_SmartObjectTimedBehavior::Execute(UMassEntitySubsystem& Enti
 	EntityQuery.ForEachEntityChunk(EntitySubsystem, Context, [this, &ToRelease](FMassExecutionContext& Context)
 	{
 		const int32 NumEntities = Context.GetNumEntities();
-		const TArrayView<FDataFragment_SmartObjectUser> UserList = Context.GetMutableComponentView<FDataFragment_SmartObjectUser>();
+		const TArrayView<FDataFragment_SmartObjectUser> UserList = Context.GetMutableFragmentView<FDataFragment_SmartObjectUser>();
 
 		for (int32 i = 0; i < NumEntities; ++i)
 		{
diff --git a/Engine/Plugins/Runtime/MassGameplay/Source/MassSpawner/Private/MassEntityTemplateRegistry.cpp b/Engine/Plugins/Runtime/MassGameplay/Source/MassSpawner/Private/MassEntityTemplateRegistry.cpp
--- a/Engine/Plugins/Runtime/MassGameplay/Source/MassSpawner/Private/MassEntityTemplateRegistry.cpp
+++ b/Engine/Plugins/Runtime/MassGameplay/Source/MassSpawner/Private/MassEntityTemplateRegistry.cpp
@@ -260,7 +260,7 @@ void UMassEntityTemplateRegistry::InitializeEntityTemplate(FMassEntityTemplate&
 	OutTemplate.SetArchetype(ArchetypeHandle);
 	for (const FInstancedStruct& InstancedStruct : OutTemplate.GetChunkFragments())
 	{
-		EntitySys->SetDefaultChunkComponentValue(ArchetypeHandle, InstancedStruct);
+		EntitySys->SetDefaultChunkFragmentValue(ArchetypeHandle, InstancedStruct);
 	}
 }
 
diff --git a/Engine/Plugins/Runtime/MassGameplay/Source/MassSpawner/Private/MassSpawnerSubsystem.cpp b/Engine/Plugins/Runtime/MassGameplay/Source/MassSpawner/Private/MassSpawnerSubsystem.cpp
--- a/Engine/Plugins/Runtime/MassGameplay/Source/MassSpawner/Private/MassSpawnerSubsystem.cpp
+++ b/Engine/Plugins/Runtime/MassGameplay/Source/MassSpawner/Private/MassSpawnerSubsystem.cpp
@@ -33,7 +33,7 @@ namespace UE::Mass::SpawnerSubsystem
 		const UMassTranslatorRegistry& Registry = UMassTranslatorRegistry::Get();
 		TArray<const UMassProcessor*> Destructors;
 
-		EntitySystem.ForEachArchetypeComponentType(ArchetypeHandle, [&Registry, &Destructors](const UScriptStruct* FragmentType)
+		EntitySystem.ForEachArchetypeFragmentType(ArchetypeHandle, [&Registry, &Destructors](const UScriptStruct* FragmentType)
 		{
 			check(FragmentType);
 			const UMassFragmentDestructor* Destructor = Registry.GetFragmentDestructor(*FragmentType);
@@ -138,7 +138,7 @@ bool UMassSpawnerSubsystem::SpawnEntities(const FMassEntityTemplate& EntityTempl
 
 		const TArrayView<const FInstancedStruct> FragmentInstances = EntityTemplate.GetFragments();
 		const FArchetypeChunkCollection ChunkCollection(EntityTemplate.GetArchetype(), OutEntities);
-		EntitySystem->BatchSetEntityComponentsValues(ChunkCollection, FragmentInstances);
+		EntitySystem->BatchSetEntityFragmentsValues(ChunkCollection, FragmentInstances);
 
 		if (EntityTemplate.GetInitializationPipeline().Processors.Num())
 		{
@@ -303,7 +303,7 @@ void UMassSpawnerSubsystem::DoSpawning(const FMassEntityTemplate& EntityTemplate
 
 	const TArrayView<const FInstancedStruct> FragmentInstances = EntityTemplate.GetFragments();
 	const FArchetypeChunkCollection ChunkCollection(EntityTemplate.GetArchetype(), SpawnedEntities);
-	EntitySystem->BatchSetEntityComponentsValues(ChunkCollection, FragmentInstances);
+	EntitySystem->BatchSetEntityFragmentsValues(ChunkCollection, FragmentInstances);
 	
 	if (EntityTemplate.GetInitializationPipeline().Processors.Num())
 	{
diff --git a/Engine/Plugins/Runtime/MassGameplay/Source/MassSpawner/Public/MassEntityTemplateRegistry.h b/Engine/Plugins/Runtime/MassGameplay/Source/MassSpawner/Public/MassEntityTemplateRegistry.h
--- a/Engine/Plugins/Runtime/MassGameplay/Source/MassSpawner/Public/MassEntityTemplateRegistry.h
+++ b/Engine/Plugins/Runtime/MassGameplay/Source/MassSpawner/Public/MassEntityTemplateRegistry.h
@@ -188,7 +188,7 @@ protected:
 
 /** @todo document 
  *  Note: the Object-Instance entity templates do not support having multiple instances of the same class that differ in
- *  composition (like with dynamically added components). The first registered instance will determine the composition 
+ *  composition (like with dynamically added fragments). The first registered instance will determine the composition 
  *  of the entity template.
  */
 UCLASS()
