Added MassGameplayEditor module for details customization.

diff --git a/Engine/Plugins/Runtime/MassGameplay/MassGameplay.uplugin b/Engine/Plugins/Runtime/MassGameplay/MassGameplay.uplugin
--- a/Engine/Plugins/Runtime/MassGameplay/MassGameplay.uplugin
+++ b/Engine/Plugins/Runtime/MassGameplay/MassGameplay.uplugin
@@ -59,6 +59,10 @@
 		{
 			"Name": "MassGameplayDebug",
 			"Type": "UncookedOnly"
+		},
+		{
+			"Name": "MassGameplayEditor",
+			"Type": "Editor"
 		}
 	],
 	"Plugins": [
diff --git a/Engine/Plugins/Runtime/MassGameplay/Source/MassGameplayEditor/MassGameplayEditor.Build.cs b/Engine/Plugins/Runtime/MassGameplay/Source/MassGameplayEditor/MassGameplayEditor.Build.cs
new file mode 100644
--- /dev/null
+++ b/Engine/Plugins/Runtime/MassGameplay/Source/MassGameplayEditor/MassGameplayEditor.Build.cs
@@ -0,0 +1,49 @@
+// Copyright Epic Games, Inc. All Rights Reserved.
+
+namespace UnrealBuildTool.Rules
+{
+	public class MassGameplayEditor : ModuleRules
+	{
+		public MassGameplayEditor(ReadOnlyTargetRules Target) : base(Target)
+		{
+			PCHUsage = ModuleRules.PCHUsageMode.UseExplicitOrSharedPCHs;
+
+			PublicIncludePaths.AddRange(
+			new string[] {
+			}
+			);
+
+			PublicDependencyModuleNames.AddRange(
+			new string[] {
+				"Core",
+				"CoreUObject",
+				"Engine",
+				"InputCore",
+				"AssetTools",
+				"UnrealEd",
+				"Slate",
+				"SlateCore",
+				"EditorStyle",
+				"PropertyEditor",
+				"MassSpawner",
+				"MassEntity",
+				"DetailCustomizations",
+				"ComponentVisualizers",
+				"Projects",
+			}
+			);
+
+			PrivateDependencyModuleNames.AddRange(
+			new string[] {
+				"RenderCore",
+				"GraphEditor",
+				"KismetWidgets",
+				"PropertyEditor",
+				"AIGraph",
+				"ToolMenus",
+			}
+			);
+		}
+
+	}
+}
diff --git a/Engine/Plugins/Runtime/MassGameplay/Source/MassGameplayEditor/Private/MassGameplayEditorModule.cpp b/Engine/Plugins/Runtime/MassGameplay/Source/MassGameplayEditor/Private/MassGameplayEditorModule.cpp
new file mode 100644
--- /dev/null
+++ b/Engine/Plugins/Runtime/MassGameplay/Source/MassGameplayEditor/Private/MassGameplayEditorModule.cpp
@@ -0,0 +1,43 @@
+// Copyright Epic Games, Inc. All Rights Reserved.
+
+#include "MassGameplayEditorModule.h"
+#include "Modules/ModuleManager.h"
+//#include "Developer/AssetTools/Public/AssetToolsModule.h"
+#include "MassProcessor.h"
+
+
+#define LOCTEXT_NAMESPACE "MassGameplayEditor"
+
+IMPLEMENT_MODULE(FMassGameplayEditorModule, MassGameplayEditor)
+
+void FMassGameplayEditorModule::StartupModule()
+{
+	// Register the details customizers
+	FPropertyEditorModule& PropertyModule = FModuleManager::LoadModuleChecked<FPropertyEditorModule>("PropertyEditor");
+
+	RegisterSectionMappings();
+	
+	PropertyModule.NotifyCustomizationModuleChanged();
+}
+
+void FMassGameplayEditorModule::ShutdownModule()
+{
+}
+
+void FMassGameplayEditorModule::RegisterSectionMappings()
+{
+	static const FName PropertyEditor("PropertyEditor");
+	FPropertyEditorModule& PropertyModule = FModuleManager::GetModuleChecked<FPropertyEditorModule>(PropertyEditor);
+
+	{
+		TSharedRef<FPropertySection> Section = PropertyModule.FindOrCreateSection("MassSpawner", "Mass", LOCTEXT("Mass", "Mass"));
+		Section->AddCategory("Mass");
+		Section->AddCategory("Debug");
+	}
+
+	{
+		TSharedRef<FPropertySection> Section = PropertyModule.FindOrCreateSection("MassAgentComponent", "Mass", LOCTEXT("Mass", "Mass"));
+		Section->AddCategory("Mass");
+	}
+}
+#undef LOCTEXT_NAMESPACE
diff --git a/Engine/Plugins/Runtime/MassGameplay/Source/MassGameplayEditor/Public/IMassGameplayEditor.h b/Engine/Plugins/Runtime/MassGameplay/Source/MassGameplayEditor/Public/IMassGameplayEditor.h
new file mode 100644
--- /dev/null
+++ b/Engine/Plugins/Runtime/MassGameplay/Source/MassGameplayEditor/Public/IMassGameplayEditor.h
@@ -0,0 +1,13 @@
+// Copyright Epic Games, Inc. All Rights Reserved.
+
+#pragma once
+
+#include "CoreMinimal.h"
+#include "Toolkits/AssetEditorToolkit.h"
+
+
+/** Mass Editor public interface */
+class MASSGAMEPLAYEDITOR_API IMassGameplayEditor : public FAssetEditorToolkit
+{
+public:
+};
diff --git a/Engine/Plugins/Runtime/MassGameplay/Source/MassGameplayEditor/Public/MassGameplayEditorModule.h b/Engine/Plugins/Runtime/MassGameplay/Source/MassGameplayEditor/Public/MassGameplayEditorModule.h
new file mode 100644
--- /dev/null
+++ b/Engine/Plugins/Runtime/MassGameplay/Source/MassGameplayEditor/Public/MassGameplayEditorModule.h
@@ -0,0 +1,27 @@
+// Copyright Epic Games, Inc. All Rights Reserved.
+
+#pragma  once
+
+#include "CoreMinimal.h"
+#include "Toolkits/AssetEditorToolkit.h"
+
+
+class IMassGameplayEditor;
+class FAssetTypeActions_Base;
+struct FGraphPanelNodeFactory;
+class UMassSchematic;
+
+/**
+* The public interface to this module
+*/
+class MASSGAMEPLAYEDITOR_API FMassGameplayEditorModule : public IModuleInterface
+{
+public:
+	DECLARE_MULTICAST_DELEGATE_TwoParams(FOnAssetPropertiesChanged, class UMassSchematic* /*MassSchematic*/, const FPropertyChangedEvent& /*PropertyChangedEvent*/);
+
+	virtual void StartupModule() override;
+	virtual void ShutdownModule() override;
+
+protected:
+	void RegisterSectionMappings();
+};
