Massaged Mass' debug processors and trait a bit to be more independent

diff --git a/Engine/Plugins/Runtime/MassGameplay/Source/MassGameplayDebug/Private/DebugVisLocationProcessor.cpp b/Engine/Plugins/Runtime/MassGameplay/Source/MassGameplayDebug/Private/DebugVisLocationProcessor.cpp
--- a/Engine/Plugins/Runtime/MassGameplay/Source/MassGameplayDebug/Private/DebugVisLocationProcessor.cpp
+++ b/Engine/Plugins/Runtime/MassGameplay/Source/MassGameplayDebug/Private/DebugVisLocationProcessor.cpp
@@ -20,7 +20,7 @@ void UDebugVisLocationProcessor::ConfigureQueries()
 {
 	EntityQuery.AddRequirement<FDataFragment_Transform>(EMassFragmentAccess::ReadOnly);
 	EntityQuery.AddRequirement<FSimDebugVisComponent>(EMassFragmentAccess::ReadOnly);
-	EntityQuery.AddRequirement<FMassVelocityFragment>(EMassFragmentAccess::ReadOnly);
+	EntityQuery.AddTagRequirement<FMassDebuggableTag>(EMassFragmentPresence::All);
 }
 
 void UDebugVisLocationProcessor::Execute(UMassEntitySubsystem& EntitySubsystem, FMassExecutionContext& Context)
@@ -90,6 +90,7 @@ void UMassProcessor_UpdateDebugVis::ConfigureQueries()
 	EntityQuery.AddRequirement<FDataFragment_Transform>(EMassFragmentAccess::ReadOnly);
 	EntityQuery.AddRequirement<FDataFragment_DebugVis>(EMassFragmentAccess::ReadWrite);
 	EntityQuery.AddRequirement<FDataFragment_AgentRadius>(EMassFragmentAccess::ReadWrite);
+	EntityQuery.AddTagRequirement<FMassDebuggableTag>(EMassFragmentPresence::All);
 }
 
 void UMassProcessor_UpdateDebugVis::Execute(UMassEntitySubsystem& EntitySubsystem, FMassExecutionContext& Context)
diff --git a/Engine/Plugins/Runtime/MassGameplay/Source/MassGameplayDebug/Private/MassDebugVisualizationTrait.cpp b/Engine/Plugins/Runtime/MassGameplay/Source/MassGameplayDebug/Private/MassDebugVisualizationTrait.cpp
--- a/Engine/Plugins/Runtime/MassGameplay/Source/MassGameplayDebug/Private/MassDebugVisualizationTrait.cpp
+++ b/Engine/Plugins/Runtime/MassGameplay/Source/MassGameplayDebug/Private/MassDebugVisualizationTrait.cpp
@@ -15,6 +15,8 @@ void UMassDebugVisualizationTrait::BuildTemplate(FMassEntityTemplateBuildContext
 	const UStaticMesh* const DebugMesh = nullptr;
 #endif
 
+	BuildContext.AddTag<FMassDebuggableTag>();
+	
 	if (DebugMesh)
 	{
 #if WITH_EDITORONLY_DATA
@@ -45,6 +47,7 @@ void UMassDebugVisualizationTrait::BuildTemplate(FMassEntityTemplateBuildContext
 #endif // WITH_EDITORONLY_DATA
 		BuildContext.AddFragmentWithDefaultInitializer<FDataFragment_AgentRadius>();
 	}
+	BuildContext.AddFragmentWithDefaultInitializer<FDataFragment_Transform>();
 #endif // if !(UE_BUILD_SHIPPING || UE_BUILD_TEST)
 
 }
\ No newline at end of file
diff --git a/Engine/Plugins/Runtime/MassGameplay/Source/MassGameplayDebug/Public/MassGameplayDebugTypes.h b/Engine/Plugins/Runtime/MassGameplay/Source/MassGameplayDebug/Public/MassGameplayDebugTypes.h
--- a/Engine/Plugins/Runtime/MassGameplay/Source/MassGameplayDebug/Public/MassGameplayDebugTypes.h
+++ b/Engine/Plugins/Runtime/MassGameplay/Source/MassGameplayDebug/Public/MassGameplayDebugTypes.h
@@ -63,6 +63,12 @@ struct FDataFragment_DebugVis : public FMassFragment
 	EMassEntityDebugShape Shape = EMassEntityDebugShape::Box;
 };
 
+USTRUCT()
+struct FMassDebuggableTag : public FMassTag
+{
+	GENERATED_BODY()
+};
+
 USTRUCT()
 struct FAgentDebugVisualization : public FTableRowBase
 {
