Added ways of asking the entity template whether it hosts given fragment/tag/chunk fragment

diff --git a/Engine/Plugins/Runtime/MassGameplay/Source/MassSpawner/Public/MassEntityTemplate.h b/Engine/Plugins/Runtime/MassGameplay/Source/MassSpawner/Public/MassEntityTemplate.h
--- a/Engine/Plugins/Runtime/MassGameplay/Source/MassSpawner/Public/MassEntityTemplate.h
+++ b/Engine/Plugins/Runtime/MassGameplay/Source/MassSpawner/Public/MassEntityTemplate.h
@@ -145,6 +145,24 @@ struct MASSSPAWNER_API FMassEntityTemplate
 	FString DebugGetDescription(UMassEntitySubsystem* EntitySubsystem = nullptr) const;
 	FString DebugGetArchetypeDescription(UMassEntitySubsystem& EntitySubsystem) const;
 
+	template<typename T>
+	bool HasFragment() const
+	{
+		return FragmentCollection.GetFragmentBitSet().Contains(*T::StaticStruct());
+	}
+	
+	template<typename T>
+	bool HasTag() const
+	{
+		return TagBitSet.Contains(*T::StaticStruct());
+	}
+
+	template<typename T>
+	bool HasChunkFragment() const
+	{
+		return ChunkFragmentBitSet.Contains(*T::StaticStruct());
+	}
+
 private:
 	FArchetypeHandle Archetype;
 
diff --git a/Engine/Plugins/Runtime/MassGameplay/Source/MassSpawner/Public/MassEntityTemplateRegistry.h b/Engine/Plugins/Runtime/MassGameplay/Source/MassSpawner/Public/MassEntityTemplateRegistry.h
--- a/Engine/Plugins/Runtime/MassGameplay/Source/MassSpawner/Public/MassEntityTemplateRegistry.h
+++ b/Engine/Plugins/Runtime/MassGameplay/Source/MassSpawner/Public/MassEntityTemplateRegistry.h
@@ -116,6 +116,24 @@ struct FMassEntityTemplateBuildContext
 		return Template.AddChunkFragment_GetRef<T>();
 	}
 
+	template<typename T>
+	bool HasFragment() const
+	{
+		return Template.HasFragment<T>();
+	}
+	
+	template<typename T>
+	bool HasTag() const
+	{
+		return Template.HasTag<T>();
+	}
+
+	template<typename T>
+	bool HasChunkFragment() const
+	{
+		return Template.HasChunkFragment<T>();
+	}
+
 	//----------------------------------------------------------------------//
 	// Initializers
 	//----------------------------------------------------------------------//
