Made MassSpawner communicate mis-configuration when using EQS and fall back to using spawner's location.

diff --git a/Engine/Plugins/Runtime/MassGameplay/Source/MassSpawner/Private/MassSpawner.cpp b/Engine/Plugins/Runtime/MassGameplay/Source/MassSpawner/Private/MassSpawner.cpp
--- a/Engine/Plugins/Runtime/MassGameplay/Source/MassSpawner/Private/MassSpawner.cpp
+++ b/Engine/Plugins/Runtime/MassGameplay/Source/MassSpawner/Private/MassSpawner.cpp
@@ -335,14 +335,17 @@ void AMassSpawner::DoSpawning()
 		if (EQSRequest.IsValid() == false)
 		{
 			EQSRequest.InitForOwnerAndBlackboard(*this, /*BBAsset=*/nullptr);
-			if (!ensureMsgf(EQSRequest.IsValid(), TEXT("Query request initialization can fail due to missing parameters. See the runtime log for details")))
-			{
-				return;
-			}
 		}
 
 		FQueryFinishedSignature Delegate = FQueryFinishedSignature::CreateUObject(this, &AMassSpawner::OnEQSQueryFinished);
-		EQSRequest.Execute(*this, /*BlackboardComponent=*/nullptr, Delegate);
+		const int32 SpawnLocationsQueryIndex = EQSRequest.Execute(*this, /*BlackboardComponent=*/nullptr, Delegate);
+		if (SpawnLocationsQueryIndex == INDEX_NONE)
+		{
+			UE_VLOG_UELOG(this, LogMassSpawner, Error, TEXT("Requesting EQS query to generate spawn locations failed. Make sure the query asset is set. Falling back to spawner\'s location."));
+
+			const TArray<FVector> DummyLocations = { GetActorLocation() };
+			SpawnAtLocations(DummyLocations);
+		}
 	}
 }
 
