Made MassEntitySpawnPointsGeneratorBase default to instanced and made MassSpawner utilize that (use instances rather than classes).
This is required to be actually able to use the EQS points generator by setting its EQS query asset.
As a side effect I was able to remove EQS-specific bits out of MassSpawner.

diff --git a/Engine/Plugins/Runtime/MassGameplay/Source/MassSpawner/Private/MassEntityEQSSpawnPointsGenerator.cpp b/Engine/Plugins/Runtime/MassGameplay/Source/MassSpawner/Private/MassEntityEQSSpawnPointsGenerator.cpp
--- a/Engine/Plugins/Runtime/MassGameplay/Source/MassSpawner/Private/MassEntityEQSSpawnPointsGenerator.cpp
+++ b/Engine/Plugins/Runtime/MassGameplay/Source/MassSpawner/Private/MassEntityEQSSpawnPointsGenerator.cpp
@@ -4,6 +4,12 @@
 #include "MassSpawnerTypes.h"
 #include "VisualLogger/VisualLogger.h"
 
+
+UMassEntityEQSSpawnPointsGenerator::UMassEntityEQSSpawnPointsGenerator()
+{
+	EQSRequest.RunMode = EEnvQueryRunMode::AllMatching; 
+}
+
 void UMassEntityEQSSpawnPointsGenerator::GenerateSpawnPoints(UObject& QueryOwner, int32 Count, FFinishedGeneratingSpawnPointsSignature& FinishedGeneratingSpawnPointsDelegate) const
 {
 	// Need to copy the request as it is called inside a CDO and CDO states cannot be changed.
diff --git a/Engine/Plugins/Runtime/MassGameplay/Source/MassSpawner/Private/MassSpawner.cpp b/Engine/Plugins/Runtime/MassGameplay/Source/MassSpawner/Private/MassSpawner.cpp
--- a/Engine/Plugins/Runtime/MassGameplay/Source/MassSpawner/Private/MassSpawner.cpp
+++ b/Engine/Plugins/Runtime/MassGameplay/Source/MassSpawner/Private/MassSpawner.cpp
@@ -116,6 +116,21 @@ void AMassSpawner::PostInitProperties()
 	}
 }
 
+void AMassSpawner::PostLoad()
+{
+	Super::PostLoad();
+
+	for (FMassSpawnPointGenerator& SpawnPointsGenerator : SpawnPointsGenerators)
+	{
+		if (SpawnPointsGenerator.GeneratorClass)
+		{
+			SpawnPointsGenerator.GeneratorInstance = NewObject<UMassEntitySpawnPointsGeneratorBase>(this, SpawnPointsGenerator.GeneratorClass);
+			SpawnPointsGenerator.GeneratorClass = nullptr;
+			MarkPackageDirty();
+		}
+	}
+}
+
 void AMassSpawner::PostRegisterAllComponents()
 {
 	Super::PostRegisterAllComponents();
@@ -221,7 +236,6 @@ void AMassSpawner::EndPlay(const EEndPlayReason::Type EndPlayReason)
 void AMassSpawner::PostEditChangeProperty(FPropertyChangedEvent& PropertyChangedEvent)
 {
 	static const FName SpawnSetsName = GET_MEMBER_NAME_CHECKED(AMassSpawner, SpawnSets);
-	static const FName EQSRequestName = GET_MEMBER_NAME_CHECKED(AMassSpawner, EQSRequest);
 	static const FName EntityTypesName = GET_MEMBER_NAME_CHECKED(AMassSpawner, EntityTypes);
 
 	Super::PostEditChangeProperty(PropertyChangedEvent);
@@ -233,22 +247,12 @@ void AMassSpawner::PostEditChangeProperty(FPropertyChangedEvent& PropertyChanged
 		{
 			ValidateSpawnSets();
 		}
-		else if (PropName == EQSRequestName)
-		{
-			EQSRequest.InitForOwnerAndBlackboard(*this, /*BBAsset=*/nullptr);
-		}
 		else if (PropName == EntityTypesName)
 		{
 			// TODO: Should optimize this, i.e. set a dirty flag and update only when needed.
 			RegisterEntityTemplates();
 		}
 	}
-
-	if (PropertyChangedEvent.MemberProperty &&
-		PropertyChangedEvent.MemberProperty->GetFName() == EQSRequestName)
-	{
-		EQSRequest.PostEditChangeProperty(*this, PropertyChangedEvent);
-	}
 }
 
 void AMassSpawner::ValidateSpawnSets()
@@ -281,15 +285,13 @@ void AMassSpawner::RegisterEntityTemplates()
 
 void AMassSpawner::DoSpawning()
 {
-	if (bUseSpawnPointsGenerators)
-	{
 		const int32 SpawnCount = GetSpawnCount();
 		if (SpawnCount > 0)
 		{
 			float TotalProportion = 0.0f;
 			for (FMassSpawnPointGenerator& SpawnPointsGenerator : SpawnPointsGenerators)
 			{
-				if (SpawnPointsGenerator.GeneratorClass.GetDefaultObject())
+			if (SpawnPointsGenerator.GeneratorInstance)
 				{
 					SpawnPointsGenerator.bPointsGenerated = false;
 					TotalProportion += SpawnPointsGenerator.Proportion;
@@ -315,14 +317,14 @@ void AMassSpawner::DoSpawning()
 					continue;
 				}
 
-				if(const UMassEntitySpawnPointsGeneratorBase* Generator = SpawnPointsGenerator.GeneratorClass.GetDefaultObject())
+			if (SpawnPointsGenerator.GeneratorInstance)
 				{
 					const float ProportionRatio = FMath::Min(SpawnPointsGenerator.Proportion / ProportionRemaining, 1.0f);
 					const int32 SpawnPointCount = FMath::CeilToInt(SpawnPointCountRemaining * ProportionRatio);
 					if(SpawnPointCount > 0)
 					{
 						FFinishedGeneratingSpawnPointsSignature Delegate = FFinishedGeneratingSpawnPointsSignature::CreateUObject(this, &AMassSpawner::OnSpawnPointGenerationFinished, &SpawnPointsGenerator);
-						Generator->GenerateSpawnPoints(*this, SpawnPointCount, Delegate);
+					SpawnPointsGenerator.GeneratorInstance->GenerateSpawnPoints(*this, SpawnPointCount, Delegate);
 						SpawnPointCountRemaining -= SpawnPointCount;
 						ProportionRemaining -= SpawnPointsGenerator.Proportion;
 					}
@@ -330,24 +332,6 @@ void AMassSpawner::DoSpawning()
 			}
 		}
 	}
-	else
-	{
-		if (EQSRequest.IsValid() == false)
-		{
-			EQSRequest.InitForOwnerAndBlackboard(*this, /*BBAsset=*/nullptr);
-		}
-
-		FQueryFinishedSignature Delegate = FQueryFinishedSignature::CreateUObject(this, &AMassSpawner::OnEQSQueryFinished);
-		const int32 SpawnLocationsQueryIndex = EQSRequest.Execute(*this, /*BlackboardComponent=*/nullptr, Delegate);
-		if (SpawnLocationsQueryIndex == INDEX_NONE)
-		{
-			UE_VLOG_UELOG(this, LogMassSpawner, Error, TEXT("Requesting EQS query to generate spawn locations failed. Make sure the query asset is set. Falling back to spawner\'s location."));
-
-			const TArray<FVector> DummyLocations = { GetActorLocation() };
-			SpawnAtLocations(DummyLocations);
-		}
-	}
-}
 
 void AMassSpawner::OnSpawnPointGenerationFinished(const TArray<FVector>& Locations, FMassSpawnPointGenerator* FinishedGenerator)
 {
@@ -394,20 +378,6 @@ int32 AMassSpawner::GetSpawnCount() const
 	}
 }
 
-void AMassSpawner::OnEQSQueryFinished(TSharedPtr<FEnvQueryResult> Result)
-{
-	if (Result.IsValid() == false || Result->IsSuccessful() == false)
-	{
-		UE_VLOG_UELOG(this, LogMassSpawner, Error, TEXT("EQS query failed or result is invalid"));
-		return;
-	}
-	
-	TArray<FVector> Locations;
-	Result->GetAllAsLocations(Locations);
-
-	SpawnAtLocations(Locations);
-}
-
 void AMassSpawner::SpawnAtLocations(const TArray<FVector>& Locations)
 {
 	UMassSpawnerSubsystem* SpawnerSystem = UWorld::GetSubsystem<UMassSpawnerSubsystem>(GetWorld());
diff --git a/Engine/Plugins/Runtime/MassGameplay/Source/MassSpawner/Public/MassEntityEQSSpawnPointsGenerator.h b/Engine/Plugins/Runtime/MassGameplay/Source/MassSpawner/Public/MassEntityEQSSpawnPointsGenerator.h
--- a/Engine/Plugins/Runtime/MassGameplay/Source/MassSpawner/Public/MassEntityEQSSpawnPointsGenerator.h
+++ b/Engine/Plugins/Runtime/MassGameplay/Source/MassSpawner/Public/MassEntityEQSSpawnPointsGenerator.h
@@ -12,10 +12,12 @@
  */
 UCLASS(BlueprintType, meta=(DisplayName="EQS SpawnPoints Generator"))
 class MASSSPAWNER_API UMassEntityEQSSpawnPointsGenerator : public UMassEntitySpawnPointsGeneratorBase
-{
+{	
 	GENERATED_BODY()
 
 public:
+	UMassEntityEQSSpawnPointsGenerator();
+	
 	virtual void GenerateSpawnPoints(UObject& QueryOwner, int32 Count, FFinishedGeneratingSpawnPointsSignature& FinishedGeneratingSpawnPointsDelegate) const override;
 
 protected:
diff --git a/Engine/Plugins/Runtime/MassGameplay/Source/MassSpawner/Public/MassEntitySpawnPointsGeneratorBase.h b/Engine/Plugins/Runtime/MassGameplay/Source/MassSpawner/Public/MassEntitySpawnPointsGeneratorBase.h
--- a/Engine/Plugins/Runtime/MassGameplay/Source/MassSpawner/Public/MassEntitySpawnPointsGeneratorBase.h
+++ b/Engine/Plugins/Runtime/MassGameplay/Source/MassSpawner/Public/MassEntitySpawnPointsGeneratorBase.h
@@ -13,7 +13,7 @@ DECLARE_DELEGATE_OneParam(FFinishedGeneratingSpawnPointsSignature, const TArray<
  * A Mass Spawn Points Generator can be of several type (EQS, ZoneGraph, Volume, Area, etc.)
  * The concept is to override the GenerateSpawnPoints() method and requesting a certain number of Spawn Point Locations to the method.
  */
-UCLASS(Abstract, Blueprintable)
+UCLASS(Abstract, Blueprintable, EditInlineNew)
 class MASSSPAWNER_API UMassEntitySpawnPointsGeneratorBase : public UObject
 {
 	GENERATED_BODY()
diff --git a/Engine/Plugins/Runtime/MassGameplay/Source/MassSpawner/Public/MassSpawner.h b/Engine/Plugins/Runtime/MassGameplay/Source/MassSpawner/Public/MassSpawner.h
--- a/Engine/Plugins/Runtime/MassGameplay/Source/MassSpawner/Public/MassSpawner.h
+++ b/Engine/Plugins/Runtime/MassGameplay/Source/MassSpawner/Public/MassSpawner.h
@@ -24,6 +24,7 @@ protected:
 	virtual void BeginPlay() override;
 	virtual void EndPlay(const EEndPlayReason::Type EndPlayReason) override;
 	virtual void PostInitProperties() override;
+	virtual void PostLoad() override;
 	virtual void PostRegisterAllComponents() override;
 	virtual void BeginDestroy() override;
 
@@ -81,7 +82,6 @@ public:
 	void ScaleSpawningCount(float Scale) { SpawningCountScale = Scale; }
 
 protected:
-	void OnEQSQueryFinished(TSharedPtr<FEnvQueryResult> Result);
 	void OnAdjustTickSchematics(UWorld* World, TArray<TSoftObjectPtr<UMassSchematic>>& InOutTickSchematics);
 	void OnPostWorldInit(UWorld* World, const UWorld::InitializationValues);
 	void SpawnAtLocations(const TArray<FVector>& Locations);
@@ -110,16 +110,10 @@ protected:
 	UPROPERTY(BlueprintReadWrite, EditAnywhere, Category = "Mass|Spawn", meta = (EditCondition = "bUseEntityConfig"))
 	TArray<FMassSpawnedEntityType> EntityTypes;
 
-	UPROPERTY(EditAnywhere, Category = "Mass|Spawn")
-	bool bUseSpawnPointsGenerators = false;
-
 	/** Asset that describes the way we want to generate SpawnPoints of the entity */
-	UPROPERTY(EditAnywhere, Category = "Mass|Spawn", meta = (EditCondition = "bUseSpawnPointsGenerators"))
+	UPROPERTY(EditAnywhere, Category = "Mass|Spawn")
 	TArray<FMassSpawnPointGenerator> SpawnPointsGenerators;
 
-	UPROPERTY(Category = "Mass|Spawn", EditAnywhere, meta = (EditCondition = "!bUseSpawnPointsGenerators"))
-	FEQSParametrizedQueryExecutionRequest EQSRequest;
-
 	UPROPERTY(Category = "Mass|Spawn", EditAnywhere)
 	uint32 bAutoSpawnOnBeginPlay : 1;
 
diff --git a/Engine/Plugins/Runtime/MassGameplay/Source/MassSpawner/Public/MassSpawnerTypes.h b/Engine/Plugins/Runtime/MassGameplay/Source/MassSpawner/Public/MassSpawnerTypes.h
--- a/Engine/Plugins/Runtime/MassGameplay/Source/MassSpawner/Public/MassSpawnerTypes.h
+++ b/Engine/Plugins/Runtime/MassGameplay/Source/MassSpawner/Public/MassSpawnerTypes.h
@@ -85,9 +85,12 @@ struct FMassSpawnPointGenerator
 	GENERATED_BODY()
 
 	/** The Generator to use to generate the spawn the points */
-	UPROPERTY(BlueprintReadWrite, EditAnywhere, Category = "Mass|Generator")
+	UPROPERTY()
 	TSubclassOf<UMassEntitySpawnPointsGeneratorBase> GeneratorClass;
 
+	UPROPERTY(BlueprintReadWrite, EditAnywhere, Instanced, Category = "Mass|Generator")
+	TObjectPtr<UMassEntitySpawnPointsGeneratorBase> GeneratorInstance;
+
 	/** Proportion of the spawn points that should be generated by this generator, (the proportions will be normalized with other sibling generators) */
 	UPROPERTY(BlueprintReadWrite, EditAnywhere, Category = "Mass|Generator", meta = (ClampMin = 0.0f, UIMin = 0.0f))
 	float Proportion = 1.0f;
