Fixed FMassSimulationLODConfig.FMassSimulationLODConfig BaseStruct meta markup

diff --git a/Engine/Plugins/Runtime/MassGameplay/Source/MassLOD/Public/MassSimulationLOD.h b/Engine/Plugins/Runtime/MassGameplay/Source/MassLOD/Public/MassSimulationLOD.h
--- a/Engine/Plugins/Runtime/MassGameplay/Source/MassLOD/Public/MassSimulationLOD.h
+++ b/Engine/Plugins/Runtime/MassGameplay/Source/MassLOD/Public/MassSimulationLOD.h
@@ -95,7 +95,7 @@ struct FMassSimulationLODConfig
 	FMassSimulationLODConfig();
 
 	/** Component Tag that will be used to associate LOD config */
-	UPROPERTY(EditAnywhere, Category = "Mass|LOD", config, meta = (BaseStruct = "ComponentTag"))
+	UPROPERTY(EditAnywhere, Category = "Mass|LOD", config, meta = (BaseStruct = "MassTag"))
 	FInstancedStruct TagFilter;
 
 	/** Distance where each LOD becomes relevant */
