Added a debug-time functionality to Mass to retrieve information on entities being marked for debugging (index range and actual entity handle for a given index)

diff --git a/Engine/Plugins/Runtime/MassEntity/Source/MassEntity/Public/MassEntitySubsystem.h b/Engine/Plugins/Runtime/MassEntity/Source/MassEntity/Public/MassEntitySubsystem.h
--- a/Engine/Plugins/Runtime/MassEntity/Source/MassEntity/Public/MassEntitySubsystem.h
+++ b/Engine/Plugins/Runtime/MassEntity/Source/MassEntity/Public/MassEntitySubsystem.h
@@ -265,6 +265,7 @@ public:
 	void DebugRemoveAllEntities();
 	void DebugForceArchetypeDataVersionBump() { ++ArchetypeDataVersion; }
 	void DebugGetArchetypeStrings(const FArchetypeHandle& Archetype, TArray<FName>& OutFragmentNames, TArray<FName>& OutTagNames);
+	FMassEntityHandle DebugGetEntityIndexHandle(const int32 EntityIndex) const { return Entities.IsValidIndex(EntityIndex) ? FMassEntityHandle({EntityIndex, Entities[EntityIndex].SerialNumber}) : FMassEntityHandle(); }
 #endif // WITH_MASSENTITY_DEBUG
 
 protected:
diff --git a/Engine/Plugins/Runtime/MassEntity/Source/MassEntity/Public/MassEntityTypes.h b/Engine/Plugins/Runtime/MassEntity/Source/MassEntity/Public/MassEntityTypes.h
--- a/Engine/Plugins/Runtime/MassEntity/Source/MassEntity/Public/MassEntityTypes.h
+++ b/Engine/Plugins/Runtime/MassEntity/Source/MassEntity/Public/MassEntityTypes.h
@@ -42,7 +42,7 @@ struct FMassEntityHandle
 {
 	GENERATED_BODY()
 
-	FMassEntityHandle() {}
+	FMassEntityHandle() = default;
 	
 	UPROPERTY(VisibleAnywhere, Category = "Mass|Debug", Transient)
 	int32 Index = 0;
diff --git a/Engine/Plugins/Runtime/MassGameplay/Source/MassCommon/Private/MassCommonTypes.cpp b/Engine/Plugins/Runtime/MassGameplay/Source/MassCommon/Private/MassCommonTypes.cpp
--- a/Engine/Plugins/Runtime/MassGameplay/Source/MassCommon/Private/MassCommonTypes.cpp
+++ b/Engine/Plugins/Runtime/MassGameplay/Source/MassCommon/Private/MassCommonTypes.cpp
@@ -82,6 +82,17 @@ bool HasDebugEntities()
 {
 	return DebugEntityBegin != INDEX_NONE && DebugEntityEnd != INDEX_NONE;
 }
+
+void GetDebugEntitiesRange(int32& OutBegin, int32& OutEnd)
+{
+	OutBegin = DebugEntityBegin;
+	OutEnd = DebugEntityEnd;
+}
+	
+FMassEntityHandle ConvertEntityIndexToHandle(const UMassEntitySubsystem& EntitySystem, const int32 EntityIndex)
+{
+	return EntitySystem.DebugGetEntityIndexHandle(EntityIndex);
+}
 	
 bool IsDebuggingEntity(FMassEntityHandle Entity, FColor* OutEntityColor)
 {
diff --git a/Engine/Plugins/Runtime/MassGameplay/Source/MassCommon/Public/MassCommonTypes.h b/Engine/Plugins/Runtime/MassGameplay/Source/MassCommon/Public/MassCommonTypes.h
--- a/Engine/Plugins/Runtime/MassGameplay/Source/MassCommon/Public/MassCommonTypes.h
+++ b/Engine/Plugins/Runtime/MassGameplay/Source/MassCommon/Public/MassCommonTypes.h
@@ -12,11 +12,14 @@
 
 class UStaticMesh;
 class UMaterialInterface;
+class UMassEntitySubsystem;
 
 #if WITH_MASSGAMEPLAY_DEBUG
 namespace UE::Mass::Debug
 {
 	MASSCOMMON_API extern bool HasDebugEntities();
+	MASSCOMMON_API extern void GetDebugEntitiesRange(int32& OutBegin, int32& OutEnd);
+	MASSCOMMON_API extern FMassEntityHandle ConvertEntityIndexToHandle(const UMassEntitySubsystem& EntitySystem, const int32 EntityIndex);
 	MASSCOMMON_API extern bool IsDebuggingEntity(FMassEntityHandle Entity, FColor* OutEntityColor = nullptr);
 	MASSCOMMON_API extern FColor GetEntityDebugColor(FMassEntityHandle Entity);
 } // namespace UE::Mass::Debug
