Split FMassObjectHandlers.Initializers into two separate collections to support default initializers always getting executed first.

Note that this is a temporary measure as we plan to change how fragment initialization/deinitialization is done.

diff --git a/Engine/Plugins/Runtime/MassGameplay/Source/MassSpawner/Private/MassEntityTemplate.cpp b/Engine/Plugins/Runtime/MassGameplay/Source/MassSpawner/Private/MassEntityTemplate.cpp
--- a/Engine/Plugins/Runtime/MassGameplay/Source/MassSpawner/Private/MassEntityTemplate.cpp
+++ b/Engine/Plugins/Runtime/MassGameplay/Source/MassSpawner/Private/MassEntityTemplate.cpp
@@ -24,8 +24,13 @@ void FMassEntityTemplate::SetArchetype(const FArchetypeHandle& InArchetype)
 
 void FMassEntityTemplate::SetUpProcessors(const FMassObjectHandlers& ObjectHandlers, UObject& PipelineOwner)
 {
+	TArray<const UMassProcessor*> Initializers;
+	Initializers.Reserve(ObjectHandlers.DefaultInitializers.Num() + ObjectHandlers.Initializers.Num());
+	Initializers.Append(ObjectHandlers.DefaultInitializers);
+	Initializers.Append(ObjectHandlers.Initializers);
+
 	// @todo do we want to copy the Processors or just use the CDOs?
-	InitializationPipeline.InitializeFromArray(ObjectHandlers.Initializers, PipelineOwner);
+	InitializationPipeline.InitializeFromArray(Initializers, PipelineOwner);
 	DeinitializationPipeline.InitializeFromArray(ObjectHandlers.Deinitializers, PipelineOwner);
 }
 
diff --git a/Engine/Plugins/Runtime/MassGameplay/Source/MassSpawner/Public/MassEntityTemplate.h b/Engine/Plugins/Runtime/MassGameplay/Source/MassSpawner/Public/MassEntityTemplate.h
--- a/Engine/Plugins/Runtime/MassGameplay/Source/MassSpawner/Public/MassEntityTemplate.h
+++ b/Engine/Plugins/Runtime/MassGameplay/Source/MassSpawner/Public/MassEntityTemplate.h
@@ -27,6 +27,8 @@ struct MASSSPAWNER_API FMassObjectHandlers
 {
 	GENERATED_BODY()
 	UPROPERTY()
+	TArray<const UMassFragmentInitializer*> DefaultInitializers;
+	UPROPERTY()
 	TArray<const UMassProcessor*> Initializers;
 	UPROPERTY()
 	TArray<const UMassProcessor*> Deinitializers;
diff --git a/Engine/Plugins/Runtime/MassGameplay/Source/MassSpawner/Public/MassEntityTemplateRegistry.h b/Engine/Plugins/Runtime/MassGameplay/Source/MassSpawner/Public/MassEntityTemplateRegistry.h
--- a/Engine/Plugins/Runtime/MassGameplay/Source/MassSpawner/Public/MassEntityTemplateRegistry.h
+++ b/Engine/Plugins/Runtime/MassGameplay/Source/MassSpawner/Public/MassEntityTemplateRegistry.h
@@ -144,7 +144,7 @@ struct FMassEntityTemplateBuildContext
 		const UMassTranslatorRegistry& Registry = UMassTranslatorRegistry::Get();
 		if (const UMassFragmentInitializer* Initializer = Registry.GetFragmentInitializer(FragmentType))
 		{
-			Handlers.Initializers.AddUnique(Initializer);
+			Handlers.DefaultInitializers.AddUnique(Initializer);
 		}
 	}
 
@@ -154,7 +154,7 @@ struct FMassEntityTemplateBuildContext
 		const UMassTranslatorRegistry& Registry = UMassTranslatorRegistry::Get();
 		if (const UMassFragmentInitializer* Initializer = Registry.GetFragmentInitializer(*T::StaticStruct()))
 		{
-			Handlers.Initializers.AddUnique(Initializer);
+			Handlers.DefaultInitializers.AddUnique(Initializer);
 		}
 	}
 
