Fixed missing member initialization value in FMassSpawnPointGenerator UStruct

diff --git a/Engine/Plugins/Runtime/MassGameplay/Source/MassSpawner/Public/MassSpawnerTypes.h b/Engine/Plugins/Runtime/MassGameplay/Source/MassSpawner/Public/MassSpawnerTypes.h
--- a/Engine/Plugins/Runtime/MassGameplay/Source/MassSpawner/Public/MassSpawnerTypes.h
+++ b/Engine/Plugins/Runtime/MassGameplay/Source/MassSpawner/Public/MassSpawnerTypes.h
@@ -89,7 +89,7 @@ struct FMassSpawnPointGenerator
 	TSubclassOf<UMassEntitySpawnPointsGeneratorBase> GeneratorClass;
 
 	UPROPERTY(BlueprintReadWrite, EditAnywhere, Instanced, Category = "Mass|Generator")
-	TObjectPtr<UMassEntitySpawnPointsGeneratorBase> GeneratorInstance;
+	TObjectPtr<UMassEntitySpawnPointsGeneratorBase> GeneratorInstance = nullptr;
 
 	/** Proportion of the spawn points that should be generated by this generator, (the proportions will be normalized with other sibling generators) */
 	UPROPERTY(BlueprintReadWrite, EditAnywhere, Category = "Mass|Generator", meta = (ClampMin = 0.0f, UIMin = 0.0f))
