Added a default entry in UMassSimulationLODProcessor.LODConfigs to ensure the processor is functional without additional configuration

diff --git a/Engine/Plugins/Runtime/MassGameplay/Source/MassLOD/Private/MassSimulationLOD.cpp b/Engine/Plugins/Runtime/MassGameplay/Source/MassLOD/Private/MassSimulationLOD.cpp
--- a/Engine/Plugins/Runtime/MassGameplay/Source/MassLOD/Private/MassSimulationLOD.cpp
+++ b/Engine/Plugins/Runtime/MassGameplay/Source/MassLOD/Private/MassSimulationLOD.cpp
@@ -78,10 +78,16 @@ UMassSimulationLODProcessor::UMassSimulationLODProcessor()
 	ExecutionFlags = (int32)EProcessorExecutionFlags::All;
 	ExecutionOrder.ExecuteInGroup = UE::Mass::ProcessorGroupNames::LOD;
 	ExecutionOrder.ExecuteAfter.Add(UE::Mass::ProcessorGroupNames::LODCollector);
+	LODConfigs.AddDefaulted();
 }
 
 void UMassSimulationLODProcessor::ConfigureQueries()
 {
+	if (LODConfigs.Num() == 0)
+	{
+		UE_VLOG_UELOG(this, LogMassLOD, Warning, TEXT("Using MassSimulationLODProcessor while no LOD configs set (see MassSimulationLODProcessor.LODConfigs)."));
+	}
+
 	for (FMassSimulationLODConfig& LODConfig : LODConfigs)
 	{
 		if(LODConfig.TagFilter.GetScriptStruct())
