Incremental compilation fix

diff --git a/Engine/Plugins/Runtime/MassGameplay/Source/MassLOD/Private/MassSimulationLOD.cpp b/Engine/Plugins/Runtime/MassGameplay/Source/MassLOD/Private/MassSimulationLOD.cpp
--- a/Engine/Plugins/Runtime/MassGameplay/Source/MassLOD/Private/MassSimulationLOD.cpp
+++ b/Engine/Plugins/Runtime/MassGameplay/Source/MassLOD/Private/MassSimulationLOD.cpp
@@ -5,6 +5,8 @@
 #include "MassCommonFragments.h"
 #include "GameFramework/PlayerController.h"
 #include "Engine/World.h"
+#include "VisualLogger/VisualLogger.h"
+#include "MassLODUtils.h"
 
 //-----------------------------------------------------------------------------
 // UMassProcessor_MassSimulationLODViewersInfo
