Removed MassSpawner.bUseEntityConfig and MassSpawner.SpawnSets since we now fully rely on EntityConfigs for spawner-induced spawning.

diff --git a/Engine/Plugins/AI/MassAI/Config/BaseMassAI.ini b/Engine/Plugins/AI/MassAI/Config/BaseMassAI.ini
--- a/Engine/Plugins/AI/MassAI/Config/BaseMassAI.ini
+++ b/Engine/Plugins/AI/MassAI/Config/BaseMassAI.ini
@@ -34,7 +34,6 @@
 +PropertyRedirects=(OldName="MassAgentComponent.AgentConfig",NewName="EntityConfig")
 +PropertyRedirects=(OldName="MassSpawnedEntityType.AgentConfig",NewName="EntityConfig")
 +PropertyRedirects=(OldName="MassSpawner.AgentTypes",NewName="EntityTypes")
-+PropertyRedirects=(OldName="MassSpawner.bUseAgentConfig",NewName="bUseEntityConfig")
 +PropertyRedirects=(OldName="MassEntityConfig.Features",NewName="Traits")
 
 +PropertyRedirects=(OldName="MassZoneGraphFindEscapeTarget.FleeTag",NewName="DisturbanceAnnotationTag")
diff --git a/Engine/Plugins/Runtime/MassGameplay/Source/MassSpawner/Private/MassSpawner.cpp b/Engine/Plugins/Runtime/MassGameplay/Source/MassSpawner/Private/MassSpawner.cpp
--- a/Engine/Plugins/Runtime/MassGameplay/Source/MassSpawner/Private/MassSpawner.cpp
+++ b/Engine/Plugins/Runtime/MassGameplay/Source/MassSpawner/Private/MassSpawner.cpp
@@ -135,26 +135,21 @@ void AMassSpawner::PostRegisterAllComponents()
 {
 	Super::PostRegisterAllComponents();
 
-	// This is a temp fix for streaming levels async loading MassSpawners after UMassSpawnerSubsystem::OnPostWorldInit, in the long run we are going to need a better system
-	// for making sure all the spawnsets are registered on the clients before replication of Agents occurs. This is only required to be done for clients
 	if (HasAnyFlags(RF_ClassDefaultObject) == false)
 	{
 		UWorld* World = GetWorld();
 		check(World);
 
+
+		// This is a temp fix for streaming levels async loading MassSpawners after UMassSpawnerSubsystem::OnPostWorldInit, 
+		// in the long run we are going to need a better system for making sure all the entity templates are registered
+		// on the clients before replication of Agents occurs. This is only required to be done for clients.
 		if (GEngine->GetNetMode(GetWorld()) == NM_Client)
 		{
 			UMassSpawnerSubsystem* MassSpawnerSubsystem = UWorld::GetSubsystem<UMassSpawnerSubsystem>(World);
 			if (MassSpawnerSubsystem)
 			{
-				if (bUseEntityConfig)
-				{
-					RegisterEntityTemplates();
-				}
-				else
-				{
-					MassSpawnerSubsystem->RegisterCollection(GetSpawnSets());
-				}
+				RegisterEntityTemplates();
 			}
 			else
 			{
@@ -171,14 +166,7 @@ void AMassSpawner::OnPostWorldInit(UWorld* World, const UWorld::InitializationVa
 		UMassSpawnerSubsystem* MassSpawnerSubsystem = UWorld::GetSubsystem<UMassSpawnerSubsystem>(World);
 		check(MassSpawnerSubsystem);
 
-		if (bUseEntityConfig)
-		{
-			RegisterEntityTemplates();
-		}
-		else
-		{
-			MassSpawnerSubsystem->RegisterCollection(GetSpawnSets());
-		}
+		RegisterEntityTemplates();
 
 		FWorldDelegates::OnPostWorldInitialization.Remove(OnPostWorldInitDelegateHandle);
 	}
@@ -235,7 +223,6 @@ void AMassSpawner::EndPlay(const EEndPlayReason::Type EndPlayReason)
 #if WITH_EDITOR
 void AMassSpawner::PostEditChangeProperty(FPropertyChangedEvent& PropertyChangedEvent)
 {
-	static const FName SpawnSetsName = GET_MEMBER_NAME_CHECKED(AMassSpawner, SpawnSets);
 	static const FName EntityTypesName = GET_MEMBER_NAME_CHECKED(AMassSpawner, EntityTypes);
 
 	Super::PostEditChangeProperty(PropertyChangedEvent);
@@ -243,11 +230,7 @@ void AMassSpawner::PostEditChangeProperty(FPropertyChangedEvent& PropertyChanged
 	if (PropertyChangedEvent.Property)
 	{
 		const FName PropName = PropertyChangedEvent.Property->GetFName();
-		if (PropName == SpawnSetsName)
-		{
-			ValidateSpawnSets();
-		}
-		else if (PropName == EntityTypesName)
+		if (PropName == EntityTypesName)
 		{
 			// TODO: Should optimize this, i.e. set a dirty flag and update only when needed.
 			RegisterEntityTemplates();
@@ -255,11 +238,6 @@ void AMassSpawner::PostEditChangeProperty(FPropertyChangedEvent& PropertyChanged
 	}
 }
 
-void AMassSpawner::ValidateSpawnSets()
-{
-	// @todo shortly
-}
-
 void AMassSpawner::DEBUG_Spawn()
 {
 	DoSpawning();
@@ -285,53 +263,53 @@ void AMassSpawner::RegisterEntityTemplates()
 
 void AMassSpawner::DoSpawning()
 {
-		const int32 SpawnCount = GetSpawnCount();
-		if (SpawnCount > 0)
+	const int32 SpawnCount = GetSpawnCount();
+	if (SpawnCount > 0)
+	{
+		float TotalProportion = 0.0f;
+		for (FMassSpawnPointGenerator& SpawnPointsGenerator : SpawnPointsGenerators)
 		{
-			float TotalProportion = 0.0f;
-			for (FMassSpawnPointGenerator& SpawnPointsGenerator : SpawnPointsGenerators)
+		if (SpawnPointsGenerator.GeneratorInstance)
 			{
-			if (SpawnPointsGenerator.GeneratorInstance)
-				{
-					SpawnPointsGenerator.bPointsGenerated = false;
-					TotalProportion += SpawnPointsGenerator.Proportion;
-				}
+				SpawnPointsGenerator.bPointsGenerated = false;
+				TotalProportion += SpawnPointsGenerator.Proportion;
 			}
+		}
 
-			if (TotalProportion <= 0.0f)
-			{
-				UE_VLOG_UELOG(this, LogMassSpawner, Error, TEXT("The total combined porportion of all the generator needs to be greater than 0.0f."));
-				return;
-			}
+		if (TotalProportion <= 0.0f)
+		{
+			UE_VLOG_UELOG(this, LogMassSpawner, Error, TEXT("The total combined porportion of all the generator needs to be greater than 0.0f."));
+			return;
+		}
 
-			int32 SpawnPointCountRemaining = SpawnCount;
-			float ProportionRemaining = TotalProportion;
-			for(FMassSpawnPointGenerator& SpawnPointsGenerator : SpawnPointsGenerators)
+		int32 SpawnPointCountRemaining = SpawnCount;
+		float ProportionRemaining = TotalProportion;
+		for(FMassSpawnPointGenerator& SpawnPointsGenerator : SpawnPointsGenerators)
+		{
+			// Still needs to call the OnSpawnPointGenerationFinished for Generator we don't want to output Spawn Points for as the system will wait
+			// for the bPointsGenerated == true on every Generators before proceeding to the next steps.
+			if (SpawnPointsGenerator.Proportion == 0.0f || ProportionRemaining <= 0.0f || SpawnPointCountRemaining <= 0)
 			{
-				// Still needs to call the OnSpawnPointGenerationFinished for Generator we don't want to output Spawn Points for as the system will wait
-				// for the bPointsGenerated == true on every Generators before proceeding to the next steps.
-				if (SpawnPointsGenerator.Proportion == 0.0f || ProportionRemaining <= 0.0f || SpawnPointCountRemaining <= 0)
-				{
-					const TArray<FVector> EmptyLoc;
-					OnSpawnPointGenerationFinished(EmptyLoc, &SpawnPointsGenerator);
-					continue;
-				}
+				const TArray<FVector> EmptyLoc;
+				OnSpawnPointGenerationFinished(EmptyLoc, &SpawnPointsGenerator);
+				continue;
+			}
 
 			if (SpawnPointsGenerator.GeneratorInstance)
+			{
+				const float ProportionRatio = FMath::Min(SpawnPointsGenerator.Proportion / ProportionRemaining, 1.0f);
+				const int32 SpawnPointCount = FMath::CeilToInt(SpawnPointCountRemaining * ProportionRatio);
+				if(SpawnPointCount > 0)
 				{
-					const float ProportionRatio = FMath::Min(SpawnPointsGenerator.Proportion / ProportionRemaining, 1.0f);
-					const int32 SpawnPointCount = FMath::CeilToInt(SpawnPointCountRemaining * ProportionRatio);
-					if(SpawnPointCount > 0)
-					{
-						FFinishedGeneratingSpawnPointsSignature Delegate = FFinishedGeneratingSpawnPointsSignature::CreateUObject(this, &AMassSpawner::OnSpawnPointGenerationFinished, &SpawnPointsGenerator);
+					FFinishedGeneratingSpawnPointsSignature Delegate = FFinishedGeneratingSpawnPointsSignature::CreateUObject(this, &AMassSpawner::OnSpawnPointGenerationFinished, &SpawnPointsGenerator);
 					SpawnPointsGenerator.GeneratorInstance->GenerateSpawnPoints(*this, SpawnPointCount, Delegate);
-						SpawnPointCountRemaining -= SpawnPointCount;
-						ProportionRemaining -= SpawnPointsGenerator.Proportion;
-					}
+					SpawnPointCountRemaining -= SpawnPointCount;
+					ProportionRemaining -= SpawnPointsGenerator.Proportion;
 				}
 			}
 		}
 	}
+}
 
 void AMassSpawner::OnSpawnPointGenerationFinished(const TArray<FVector>& Locations, FMassSpawnPointGenerator* FinishedGenerator)
 {
@@ -361,21 +339,7 @@ void AMassSpawner::OnSpawnPointGenerationFinished(const TArray<FVector>& Locatio
 int32 AMassSpawner::GetSpawnCount() const
 {
 	const float FinalSpawningCountScale = SpawningCountScale * UE::MassSpawner::ScalabilitySpawnDensityMultiplier;
-	
-	if (bUseEntityConfig)
-	{
-		return int32(FinalSpawningCountScale * Count);
-	}
-	else
-	{
-		int32 OldCount = 0;
-		for (const FInstancedStruct& Entry : SpawnSets)
-		{
-			const FMassSpawnConfigBase& Config = Entry.Get<FMassSpawnConfigBase>();
-			OldCount += Config.MaxNumber;
-		}
-		return int32(FinalSpawningCountScale * OldCount);
-	}
+	return int32(FinalSpawningCountScale * Count);
 }
 
 void AMassSpawner::SpawnAtLocations(const TArray<FVector>& Locations)
@@ -396,51 +360,44 @@ void AMassSpawner::SpawnAtLocations(const TArray<FVector>& Locations)
 		Transform.SetLocation(Location);
 	}
 
-	if (bUseEntityConfig)
+	const int32 SpawnCount = GetSpawnCount();
+	if(SpawnCount > 0)
 	{
-		const int32 SpawnCount = GetSpawnCount();
-		if(SpawnCount > 0)
+		float TotalProportion = 0.0f;
+		for (const FMassSpawnedEntityType& EntityType : EntityTypes)
 		{
-			float TotalProportion = 0.0f;
-			for (const FMassSpawnedEntityType& EntityType : EntityTypes)
-			{
-				TotalProportion += EntityType.Proportion;
-			}
+			TotalProportion += EntityType.Proportion;
+		}
 
-			if (TotalProportion <= 0)
-			{
-				UE_VLOG_UELOG(this, LogMassSpawner, Error, TEXT("The total combined porportion of all the entity types needs to be greater than 0.0f."));
-				return;
-			}
+		if (TotalProportion <= 0)
+		{
+			UE_VLOG_UELOG(this, LogMassSpawner, Error, TEXT("The total combined porportion of all the entity types needs to be greater than 0.0f."));
+			return;
+		}
 
-			AllSpawnedEntities.Reserve(AllSpawnedEntities.Num() + EntityTypes.Num());
+		AllSpawnedEntities.Reserve(AllSpawnedEntities.Num() + EntityTypes.Num());
 
-			for (const FMassSpawnedEntityType& EntityType : EntityTypes)
+		for (const FMassSpawnedEntityType& EntityType : EntityTypes)
+		{
+			if (const UMassEntityConfigAsset* EntityConfig = EntityType.GetEntityConfig())
 			{
-				if (const UMassEntityConfigAsset* EntityConfig = EntityType.GetEntityConfig())
+				const int32 EntityCount = int32(SpawnCount * EntityType.Proportion / TotalProportion);
+				if(EntityCount > 0)
 				{
-					const int32 EntityCount = int32(SpawnCount * EntityType.Proportion / TotalProportion);
-					if(EntityCount > 0)
+					FMassSpawnConfigBase SpawnConfig;
+					SpawnConfig.MinNumber = EntityCount;
+					SpawnConfig.MaxNumber = EntityCount;
+					const FMassEntityTemplate* EntityTemplate = EntityConfig->GetConfig().GetOrCreateEntityTemplate(*this, *EntityConfig);
+					if (EntityTemplate && EntityTemplate->IsValid())
 					{
-						FMassSpawnConfigBase SpawnConfig;
-						SpawnConfig.MinNumber = EntityCount;
-						SpawnConfig.MaxNumber = EntityCount;
-						const FMassEntityTemplate* EntityTemplate = EntityConfig->GetConfig().GetOrCreateEntityTemplate(*this, *EntityConfig);
-						if (EntityTemplate && EntityTemplate->IsValid())
-						{
-							FSpawnedEntities& SpawnedEntities = AllSpawnedEntities.AddDefaulted_GetRef();
-							SpawnedEntities.TemplateID = EntityTemplate->GetTemplateID();
-							SpawnerSystem->SpawnEntities(EntityTemplate->GetTemplateID(), SpawnConfig, FStructView::Make(MassSpawnAuxData), SpawnedEntities.Entities);
-						}
+						FSpawnedEntities& SpawnedEntities = AllSpawnedEntities.AddDefaulted_GetRef();
+						SpawnedEntities.TemplateID = EntityTemplate->GetTemplateID();
+						SpawnerSystem->SpawnEntities(EntityTemplate->GetTemplateID(), SpawnConfig, FStructView::Make(MassSpawnAuxData), SpawnedEntities.Entities);
 					}
 				}
 			}
 		}
 	}
-	else
-	{
-		SpawnerSystem->SpawnCollection(MakeArrayView(SpawnSets), /*Count=*/-1, FStructView::Make(MassSpawnAuxData));
-	}
 }
 
 void AMassSpawner::DoDespawning()
diff --git a/Engine/Plugins/Runtime/MassGameplay/Source/MassSpawner/Public/MassSpawner.h b/Engine/Plugins/Runtime/MassGameplay/Source/MassSpawner/Public/MassSpawner.h
--- a/Engine/Plugins/Runtime/MassGameplay/Source/MassSpawner/Public/MassSpawner.h
+++ b/Engine/Plugins/Runtime/MassGameplay/Source/MassSpawner/Public/MassSpawner.h
@@ -29,8 +29,6 @@ protected:
 	virtual void BeginDestroy() override;
 
 public:
-	TArray<FInstancedStruct>& GetSpawnSets() { return SpawnSets; }
-
 #if WITH_EDITOR
 	UFUNCTION(BlueprintCallable, CallInEditor, Category = "Debug")
 	void DEBUG_Spawn();
@@ -43,7 +41,6 @@ public:
 protected:
 #if WITH_EDITOR
 	virtual void PostEditChangeProperty(FPropertyChangedEvent& PropertyChangedEvent) override;
-	void ValidateSpawnSets();
 #endif // WITH_EDITOR
 
 	void RegisterEntityTemplates();
@@ -75,7 +72,7 @@ public:
 	bool DespawnEntity(const FMassEntityHandle Entity);
 
 	/**
-	 * Scales the spawning counts (this only works when bUseEntityConfig is set)
+	 * Scales the spawning counts
 	 * @param Scale is the number to multiply the all counts of each agent types 
 	 */
 	UFUNCTION(BlueprintCallable, Category = "Spawning")
@@ -97,17 +94,10 @@ protected:
 		TArray<FMassEntityHandle> Entities;
 	};
 
-	// @TODO CONSIDER HAVING ONE PER ENTRY IN SpawnSets
-	UPROPERTY(Category = "Mass|Spawn", EditAnywhere, meta = (BaseStruct = "MassSpawnConfigBase", EditCondition = "!bUseEntityConfig"))
-	TArray<FInstancedStruct> SpawnSets;
-
 	UPROPERTY(EditAnywhere, Category = "Mass|Spawn")
-	bool bUseEntityConfig = false;
-
-	UPROPERTY(EditAnywhere, Category = "Mass|Spawn", meta = (EditCondition = "bUseEntityConfig"))
 	int32 Count;
 
-	UPROPERTY(BlueprintReadWrite, EditAnywhere, Category = "Mass|Spawn", meta = (EditCondition = "bUseEntityConfig"))
+	UPROPERTY(BlueprintReadWrite, EditAnywhere, Category = "Mass|Spawn")
 	TArray<FMassSpawnedEntityType> EntityTypes;
 
 	/** Asset that describes the way we want to generate SpawnPoints of the entity */
