[MassDebug] fixed ensure in UDebugVisLocationProcessor in cooked builds since UMassDebugVisualizationComponent is available only in Editor

diff --git a/Engine/Plugins/Runtime/MassGameplay/Source/MassGameplayDebug/Private/DebugVisLocationProcessor.cpp b/Engine/Plugins/Runtime/MassGameplay/Source/MassGameplayDebug/Private/DebugVisLocationProcessor.cpp
--- a/Engine/Plugins/Runtime/MassGameplay/Source/MassGameplayDebug/Private/DebugVisLocationProcessor.cpp
+++ b/Engine/Plugins/Runtime/MassGameplay/Source/MassGameplayDebug/Private/DebugVisLocationProcessor.cpp
@@ -68,12 +68,15 @@ void UDebugVisLocationProcessor::Execute(UMassEntitySubsystem& EntitySubsystem,
 void UDebugVisLocationProcessor::Initialize(UObject& InOwner)
 {
 	Super::Initialize(InOwner);
+
+#if WITH_EDITORONLY_DATA
 	UMassDebuggerSubsystem* Debugger = UWorld::GetSubsystem<UMassDebuggerSubsystem>(InOwner.GetWorld());
 	if (ensure(Debugger))
 	{
 		WeakVisualizer = Debugger->GetVisualizationComponent();
 		ensure(WeakVisualizer.Get());
 	}
+#endif // WITH_EDITORONLY_DATA
 }
 
 //----------------------------------------------------------------------//
