[MassDebug] modules MassGameplayDebug and MassAIDebug in all runtime targets (including cooked)

diff --git a/Engine/Plugins/AI/MassAI/MassAI.uplugin b/Engine/Plugins/AI/MassAI/MassAI.uplugin
--- a/Engine/Plugins/AI/MassAI/MassAI.uplugin
+++ b/Engine/Plugins/AI/MassAI/MassAI.uplugin
@@ -35,7 +35,7 @@
 		},
 		{
 			"Name": "MassAIDebug",
-			"Type": "UncookedOnly"
+			"Type": "Runtime"
 		},
 		{
 			"Name": "MassAITestSuite",
diff --git a/Engine/Plugins/Runtime/MassGameplay/MassGameplay.uplugin b/Engine/Plugins/Runtime/MassGameplay/MassGameplay.uplugin
--- a/Engine/Plugins/Runtime/MassGameplay/MassGameplay.uplugin
+++ b/Engine/Plugins/Runtime/MassGameplay/MassGameplay.uplugin
@@ -58,7 +58,7 @@
 		},
 		{
 			"Name": "MassGameplayDebug",
-			"Type": "UncookedOnly"
+			"Type": "Runtime"
 		},
 		{
 			"Name": "MassGameplayEditor",
