Integrating Mass AI stuff To //UE5/Release-5.0

#rnx
#prefilght 61803e3b1ff8ae0001a3cf4f

diff --git a/Engine/Plugins/AI/MassAI/Source/MassAIBehavior/Private/MassStateTreeProcessors.cpp b/Engine/Plugins/AI/MassAI/Source/MassAIBehavior/Private/MassStateTreeProcessors.cpp
--- a/Engine/Plugins/AI/MassAI/Source/MassAIBehavior/Private/MassStateTreeProcessors.cpp
+++ b/Engine/Plugins/AI/MassAI/Source/MassAIBehavior/Private/MassStateTreeProcessors.cpp
@@ -5,6 +5,7 @@
 #include "MassStateTreeExecutionContext.h"
 #include "MassEntityView.h"
 #include "MassAIMovementTypes.h"
+#include "MassSimulationLOD.h"
 #include "MassComponentHitTypes.h"
 #include "MassSmartObjectTypes.h"
 #include "MassZoneGraphAnnotationTypes.h"
@@ -154,17 +155,13 @@ void UMassStateTreeFragmentInitializer::ConfigureQueries()
 
 void UMassStateTreeFragmentInitializer::Execute(UMassEntitySubsystem& EntitySubsystem, FMassExecutionContext& Context)
 {
-	// StateTree processor relies on signals to be ticked but we need an 'initial tick' to set the tree in the proper state.
-	// The initializer provides that by sending a signal to all new entities that use StateTree.
-	TArray<FMassEntityHandle> EntitiesToSignal;
-
 	FMassStateTreeExecutionContext StateTreeContext(EntitySubsystem, Context);
 	UMassStateTreeSubsystem* MassStateTreeSubsystem = UWorld::GetSubsystem<UMassStateTreeSubsystem>(EntitySubsystem.GetWorld());
 
 	EntityQuery.ForEachEntityChunk(
 		EntitySubsystem,
 		Context,
-		[this, &StateTreeContext, &MassStateTreeSubsystem, &EntitiesToSignal](const FMassExecutionContext& Context)
+		[this, &StateTreeContext, &MassStateTreeSubsystem](const FMassExecutionContext& Context)
 		{
 			UE::MassBehavior::ForEachEntityInChunk(
 				StateTreeContext,
@@ -174,18 +171,7 @@ void UMassStateTreeFragmentInitializer::Execute(UMassEntitySubsystem& EntitySubs
 					// Start the tree instance
 					StateTreeExecutionContext.Start(Storage);
 				});
-
-			// Append all entities of the current chunk to the consolidated list to send signal once
-			EntitiesToSignal.Append(Context.GetEntities().GetData(), Context.GetEntities().Num());
 		});
-
-	// Signal all entities inside the consolidated list
-	if (EntitiesToSignal.Num())
-	{
-		UMassSignalSubsystem* SignalSubsystem = UWorld::GetSubsystem<UMassSignalSubsystem>(EntitySubsystem.GetWorld());
-		checkf(SignalSubsystem != nullptr, TEXT("MassSignalSubsystem should exist when executing fragment initializers."));
-		SignalSubsystem->SignalEntities(UE::Mass::Signals::StateTreeInitializationRequested, EntitiesToSignal);
-	}
 }
 
 //----------------------------------------------------------------------//
@@ -223,6 +209,60 @@ void UMassStateTreeFragmentDestructor::Execute(UMassEntitySubsystem& EntitySubsy
 		});
 }
 
+//----------------------------------------------------------------------//
+// UMassStateTreeSignalRequestInitialization
+//----------------------------------------------------------------------//
+UMassStateTreeSignalRequestInitialization::UMassStateTreeSignalRequestInitialization()
+{
+	ExecutionOrder.ExecuteAfter.Add(UE::Mass::ProcessorGroupNames::LOD);
+	ExecutionOrder.ExecuteBefore.Add(UE::Mass::ProcessorGroupNames::Behavior);
+	MaxInitializationRequestsPerLOD[EMassLOD::High] = 250;
+	MaxInitializationRequestsPerLOD[EMassLOD::Medium] = 250;
+	MaxInitializationRequestsPerLOD[EMassLOD::Low] = 250;
+	MaxInitializationRequestsPerLOD[EMassLOD::Off] = 250;
+}
+void UMassStateTreeSignalRequestInitialization::Initialize(UObject& Owner)
+{
+	SignalSubsystem = UWorld::GetSubsystem<UMassSignalSubsystem>(Owner.GetWorld());
+}
+void UMassStateTreeSignalRequestInitialization::ConfigureQueries()
+{
+	EntityQuery.AddRequirement<FMassStateTreeFragment>(EMassFragmentAccess::ReadOnly);
+	EntityQuery.AddTagRequirement<FMassStateTreeInitializationRequestDone>(EMassFragmentPresence::None);
+	EntityQuery.AddChunkRequirement<FMassSimulationVariableTickChunkFragment>(EMassFragmentAccess::ReadOnly);
+}
+void UMassStateTreeSignalRequestInitialization::Execute(UMassEntitySubsystem& EntitySubsystem, FMassExecutionContext& Context)
+{
+	// StateTree processor relies on signals to be ticked but we need an 'initial tick' to set the tree in the proper state.
+	// The initializer provides that by sending a signal to all new entities that use StateTree.
+	TArray<FMassEntityHandle> EntitiesToSignal;
+	int32 InitializationRequestCounts[EMassLOD::Max] {0,0,0,0};
+	EntityQuery.ForEachEntityChunk(EntitySubsystem, Context, [&EntitiesToSignal, &InitializationRequestCounts, MaxInitializationRequestsPerLOD = MaxInitializationRequestsPerLOD](FMassExecutionContext& Context)
+	{
+		const int32 NumEntities = Context.GetNumEntities();
+		// Check if we already reached the maximum for this frame
+		const EMassLOD::Type ChunkLOD = FMassSimulationVariableTickChunkFragment::GetChunkLOD(Context);
+		if (InitializationRequestCounts[ChunkLOD] > MaxInitializationRequestsPerLOD[ChunkLOD])
+		{
+			return;
+		}
+		InitializationRequestCounts[ChunkLOD] += NumEntities;
+		// Append all entities of the current chunk to the consolidated list to send signal once
+		EntitiesToSignal.Append(Context.GetEntities().GetData(), Context.GetEntities().Num());
+		// Adding a tag on each entities to remember we have sent the state tree initialization signal
+		for (int32 i = 0; i < NumEntities; ++i)
+		{
+			Context.Defer().AddTag<FMassStateTreeInitializationRequestDone>(Context.GetEntity(i));
+		}
+	});
+	// Signal all entities inside the consolidated list
+	if (EntitiesToSignal.Num())
+	{
+		checkf(SignalSubsystem != nullptr, TEXT("MassSignalSubsystem should exist when executing sending state tree initialization requests."));
+		SignalSubsystem->SignalEntities(UE::Mass::Signals::StateTreeInitializationRequested, EntitiesToSignal);
+	}
+}
+
 //----------------------------------------------------------------------//
 // UMassStateTreeProcessor
 //----------------------------------------------------------------------//
diff --git a/Engine/Plugins/AI/MassAI/Source/MassAIBehavior/Private/MassZoneGraphAnnotationFragments.cpp b/Engine/Plugins/AI/MassAI/Source/MassAIBehavior/Private/MassZoneGraphAnnotationFragments.cpp
--- a/Engine/Plugins/AI/MassAI/Source/MassAIBehavior/Private/MassZoneGraphAnnotationFragments.cpp
+++ b/Engine/Plugins/AI/MassAI/Source/MassAIBehavior/Private/MassZoneGraphAnnotationFragments.cpp
@@ -24,17 +24,24 @@ bool FMassZoneGraphAnnotationVariableTickChunkFragment::UpdateChunk(FMassExecuti
 {
 	FMassZoneGraphAnnotationVariableTickChunkFragment& ChunkFrag = Context.GetMutableChunkFragment<FMassZoneGraphAnnotationVariableTickChunkFragment>();
 	ChunkFrag.TimeUntilNextTick -= Context.GetDeltaTimeSeconds();
+	if (!ChunkFrag.bInitialized)
+	{
+		const bool bOffLOD = FMassSimulationVariableTickChunkFragment::GetChunkLOD(Context) == EMassLOD::Off;
+		ChunkFrag.TimeUntilNextTick = FMath::RandRange(0.0f, bOffLOD ? UE::Mass::ZoneGraphAnnotations::OffLODMaxUpdateInterval : UE::Mass::ZoneGraphAnnotations::MaxUpdateInterval);
+		ChunkFrag.bInitialized = true;
+	}
+	else
+	{
+		ChunkFrag.TimeUntilNextTick -= Context.GetDeltaTimeSeconds();
+	}
+	
 	if (ChunkFrag.TimeUntilNextTick <= 0.0f)
 	{
-		// @todo Possible future optimization, right now it is faster but not by a whole lot.
-		if (FMassSimulationVariableTickChunkFragment::GetChunkLOD(Context) == EMassLOD::Off)
-		{
-			ChunkFrag.TimeUntilNextTick = FMath::RandRange(UE::Mass::ZoneGraphAnnotations::OffLODMinUpdateInterval, UE::Mass::ZoneGraphAnnotations::OffLODMaxUpdateInterval);
-		}
-		else
-		{
-			ChunkFrag.TimeUntilNextTick = FMath::RandRange(UE::Mass::ZoneGraphAnnotations::MinUpdateInterval, UE::Mass::ZoneGraphAnnotations::MaxUpdateInterval);
-		}
+		const bool bOffLOD = FMassSimulationVariableTickChunkFragment::GetChunkLOD(Context) == EMassLOD::Off;
+		ChunkFrag.TimeUntilNextTick = FMath::RandRange(
+			bOffLOD ? UE::Mass::ZoneGraphAnnotations::OffLODMinUpdateInterval : UE::Mass::ZoneGraphAnnotations::MinUpdateInterval, 
+			bOffLOD ? UE::Mass::ZoneGraphAnnotations::OffLODMaxUpdateInterval : UE::Mass::ZoneGraphAnnotations::MaxUpdateInterval);
+
 		return true;
 	}
 
diff --git a/Engine/Plugins/AI/MassAI/Source/MassAIBehavior/Public/MassStateTreeProcessors.h b/Engine/Plugins/AI/MassAI/Source/MassAIBehavior/Public/MassStateTreeProcessors.h
--- a/Engine/Plugins/AI/MassAI/Source/MassAIBehavior/Public/MassStateTreeProcessors.h
+++ b/Engine/Plugins/AI/MassAI/Source/MassAIBehavior/Public/MassStateTreeProcessors.h
@@ -5,6 +5,7 @@
 #include "MassSignals/Public/MassSignalProcessorBase.h"
 #include "MassStateTreeFragments.h"
 #include "MassTranslator.h"
+#include "MassLODTypes.h"
 #include "MassStateTreeProcessors.generated.h"
 
 struct FMassStateTreeExecutionContext;
@@ -42,6 +43,38 @@ protected:
 	FMassEntityQuery EntityQuery;
 };
 
+/**
+ * Special tag to know if the state tree initialization request was signaled
+ */
+USTRUCT()
+struct MASSAIBEHAVIOR_API FMassStateTreeInitializationRequestDone : public FMassTag
+{
+	GENERATED_BODY()
+};
+/**
+ * Processor to signal initialization requests
+ */
+UCLASS()
+class MASSAIBEHAVIOR_API UMassStateTreeSignalRequestInitialization : public UMassProcessor
+{
+	GENERATED_BODY()
+public:
+	UMassStateTreeSignalRequestInitialization();
+protected:
+	virtual void Initialize(UObject& Owner) override;
+	virtual void ConfigureQueries() override;
+	virtual void Execute(UMassEntitySubsystem& EntitySubsystem, FMassExecutionContext& Context) override;
+
+	UPROPERTY(Transient)
+	UMassSignalSubsystem* SignalSubsystem = nullptr;
+	
+	FMassEntityQuery EntityQuery;
+
+	/** The maximum number of entities to send State Tree initialization request per LOD for each update, -1 means no limit */
+	UPROPERTY(EditAnywhere, Category = "Mass|LOD", config)
+	int32 MaxInitializationRequestsPerLOD[EMassLOD::Max];
+};
+
 /** 
  * Processor for executing a StateTree
  */
diff --git a/Engine/Plugins/AI/MassAI/Source/MassAIBehavior/Public/MassZoneGraphAnnotationFragments.h b/Engine/Plugins/AI/MassAI/Source/MassAIBehavior/Public/MassZoneGraphAnnotationFragments.h
--- a/Engine/Plugins/AI/MassAI/Source/MassAIBehavior/Public/MassZoneGraphAnnotationFragments.h
+++ b/Engine/Plugins/AI/MassAI/Source/MassAIBehavior/Public/MassZoneGraphAnnotationFragments.h
@@ -27,4 +27,5 @@ struct FMassZoneGraphAnnotationVariableTickChunkFragment : public FMassChunkFrag
 	static bool UpdateChunk(FMassExecutionContext& Context);
 
 	float TimeUntilNextTick = 0.0f;
+	bool bInitialized = false;
 };
\ No newline at end of file
diff --git a/Engine/Plugins/AI/MassAI/Source/MassAIMovement/Private/MassAvoidanceTrait.cpp b/Engine/Plugins/AI/MassAI/Source/MassAIMovement/Private/MassAvoidanceTrait.cpp
--- a/Engine/Plugins/AI/MassAI/Source/MassAIMovement/Private/MassAvoidanceTrait.cpp
+++ b/Engine/Plugins/AI/MassAI/Source/MassAIMovement/Private/MassAvoidanceTrait.cpp
@@ -23,6 +23,7 @@ void UMassAvoidanceTrait::BuildTemplate(FMassEntityTemplateBuildContext& BuildCo
 	BuildContext.AddFragmentWithDefaultInitializer<FDataFragment_Transform>();
 	BuildContext.AddFragmentWithDefaultInitializer<FMassLODInfoFragment>();
 	BuildContext.AddFragmentWithDefaultInitializer<FMassSimulationLODFragment>();
+	BuildContext.AddTag<FMassOffLODTag>();
 	BuildContext.AddChunkFragment<FMassSimulationVariableTickChunkFragment>();
 
 	if (bUseZoneGraphMovement)
diff --git a/Engine/Plugins/Runtime/MassGameplay/Source/MassLOD/MassLOD.Build.cs b/Engine/Plugins/Runtime/MassGameplay/Source/MassLOD/MassLOD.Build.cs
--- a/Engine/Plugins/Runtime/MassGameplay/Source/MassLOD/MassLOD.Build.cs
+++ b/Engine/Plugins/Runtime/MassGameplay/Source/MassLOD/MassLOD.Build.cs
@@ -30,6 +30,7 @@ namespace UnrealBuildTool.Rules
 					"MassEntity",
 					"StructUtils",
 					"MassCommon",
+					"MassSimulation",
 				}
 			);
 
diff --git a/Engine/Plugins/Runtime/MassGameplay/Source/MassLOD/Private/MassLODManager.cpp b/Engine/Plugins/Runtime/MassGameplay/Source/MassLOD/Private/MassLODManager.cpp
--- a/Engine/Plugins/Runtime/MassGameplay/Source/MassLOD/Private/MassLODManager.cpp
+++ b/Engine/Plugins/Runtime/MassGameplay/Source/MassLOD/Private/MassLODManager.cpp
@@ -5,6 +5,7 @@
 #include "WorldPartition/WorldPartition.h"
 #include "WorldPartition/WorldPartitionStreamingSource.h"
 #include "Engine/World.h"
+#include "MassSimulationSubsystem.h"
 
 namespace UE
 {
@@ -30,11 +31,20 @@ void UMassProcessor_LODBase::Initialize(UObject& Owner)
 
 void UMassLODManager::Initialize(FSubsystemCollectionBase& Collection)
 {
+	Collection.InitializeDependency(UMassSimulationSubsystem::StaticClass());
+
 	Super::Initialize(Collection);
 	SynchronizeViewers();
+	
+	if (UWorld* World = GetWorld())
+	{
+		UMassSimulationSubsystem* SimSystem = World->GetSubsystem<UMassSimulationSubsystem>();
+		check(SimSystem);
+		SimSystem->GetOnProcessingPhaseStarted(EMassProcessingPhase::PrePhysics).AddUObject(this, &UMassLODManager::OnPrePhysicsPhaseStarted);
+	}
 }
 
-void UMassLODManager::Tick(float DeltaTime)
+void UMassLODManager::OnPrePhysicsPhaseStarted(float DeltaTime)
 {
 	SynchronizeViewers();
 }
@@ -55,6 +65,14 @@ void UMassLODManager::Deinitialize()
 			RemoveViewer(ViewerInfo.Handle);
 		}
 	}
+	
+	if (UWorld* World = GetWorld())
+	{
+		if (UMassSimulationSubsystem* SimSystem = World->GetSubsystem<UMassSimulationSubsystem>())
+		{
+			SimSystem->GetOnProcessingPhaseStarted(EMassProcessingPhase::PrePhysics).RemoveAll(this);
+		}
+	}
 }
 
 int32 UMassLODManager::GetValidViewerIdx(const FMassViewerHandle& ViewerHandle) const
diff --git a/Engine/Plugins/Runtime/MassGameplay/Source/MassLOD/Public/MassLODManager.h b/Engine/Plugins/Runtime/MassGameplay/Source/MassLOD/Public/MassLODManager.h
--- a/Engine/Plugins/Runtime/MassGameplay/Source/MassLOD/Public/MassLODManager.h
+++ b/Engine/Plugins/Runtime/MassGameplay/Source/MassLOD/Public/MassLODManager.h
@@ -106,10 +106,12 @@ public:
 
 protected:
 	virtual void Initialize(FSubsystemCollectionBase& Collection) override;
-	virtual void Tick(float DeltaTime) override;
 	virtual TStatId GetStatId() const override;
 	virtual void Deinitialize() override;
 
+	/** Called at the start of the PrePhysics mass processing phase and calls SynchronizeViewers */ 
+	void OnPrePhysicsPhaseStarted(float DeltaTime);
+
 	/** Synchronizes the viewers from the engine PlayerController list */
 	void SynchronizeViewers();
 
diff --git a/Engine/Plugins/Runtime/MassGameplay/Source/MassLOD/Public/MassLODTickRateController.h b/Engine/Plugins/Runtime/MassGameplay/Source/MassLOD/Public/MassLODTickRateController.h
--- a/Engine/Plugins/Runtime/MassGameplay/Source/MassLOD/Public/MassLODTickRateController.h
+++ b/Engine/Plugins/Runtime/MassGameplay/Source/MassLOD/Public/MassLODTickRateController.h
@@ -88,6 +88,7 @@ bool TMassLODTickRateController<FVariableTickChunkFragment, FLODLogic>::UpdateTi
 	bool bShouldTickThisFrame = true;
 	bool bWasChunkTicked = true;
 	const float DeltaTime = Context.GetDeltaTimeSeconds();
+	bool bFirstUpdate = false;
 
 	FMassVariableTickChunkFragment& ChunkData = Context.GetMutableChunkFragment<FVariableTickChunkFragment>();
 	EMassLOD::Type ChunkLOD = ChunkData.GetLOD();
@@ -96,6 +97,7 @@ bool TMassLODTickRateController<FVariableTickChunkFragment, FLODLogic>::UpdateTi
 		// The LOD on the chunk fragment data isn't set yet, let see if the Archetype has an LOD tag and set it on the ChunkData
 		ChunkLOD = UE::MassLOD::GetLODFromArchetype(Context);
 		ChunkData.SetLOD(ChunkLOD);
+		bFirstUpdate = true;
 	}
 	else
 	{
@@ -107,11 +109,22 @@ bool TMassLODTickRateController<FVariableTickChunkFragment, FLODLogic>::UpdateTi
 		float TimeUntilNextTick = ChunkData.GetTimeUntilNextTick();
 		bWasChunkTicked = ChunkData.ShouldTickThisFrame();
 
-		// Reset DeltaTime if we ticked last frame
-		if (bWasChunkTicked)
+		const int32 LastChunkSerialModificationNumber = ChunkData.GetLastChunkSerialModificationNumber();
+		const int32 ChunkSerialModificationNumber = Context.GetChunkSerialModificationNumber();
+
+		// Prevent the chunk modification tracking logic to trigger a tick until we actually tick from the first update tick calculation
+		int32 NewChunkSerialModificationNumber = (LastChunkSerialModificationNumber == INDEX_NONE) ? INDEX_NONE : ChunkSerialModificationNumber;
+
+		const float TickRate = TickRates[ChunkLOD];
+		if (bFirstUpdate)
 		{
-			const float TickRate = TickRates[ChunkLOD];
+			TimeUntilNextTick = FMath::RandRange(0.0f, TickRate);
+		}
+		else if(bWasChunkTicked)
+		{
+			// Reset DeltaTime if we ticked last frame and start tracking chunk modifications
 			TimeUntilNextTick = TickRate * (1.0f + FMath::RandRange(-0.1f, 0.1f));
+			NewChunkSerialModificationNumber = ChunkSerialModificationNumber;
 		}
 		else
 		{
@@ -120,9 +133,8 @@ bool TMassLODTickRateController<FVariableTickChunkFragment, FLODLogic>::UpdateTi
 		}
 
 		// Should we tick this frame?
-		const int32 ChunkSerialModificationNumber = Context.GetChunkSerialModificationNumber();
-		bShouldTickThisFrame = ChunkData.GetLastChunkSerialModificationNumber() != ChunkSerialModificationNumber || TimeUntilNextTick <= 0.0f;
-		ChunkData.Update(bShouldTickThisFrame, TimeUntilNextTick, ChunkSerialModificationNumber);
+		bShouldTickThisFrame = TimeUntilNextTick <= 0.0f || LastChunkSerialModificationNumber != NewChunkSerialModificationNumber;
+		ChunkData.Update(bShouldTickThisFrame, TimeUntilNextTick, NewChunkSerialModificationNumber);
 	}
 
 	if (bWasChunkTicked)
diff --git a/Engine/Plugins/Runtime/MassGameplay/Source/MassLOD/Public/MassLODTypes.h b/Engine/Plugins/Runtime/MassGameplay/Source/MassLOD/Public/MassLODTypes.h
--- a/Engine/Plugins/Runtime/MassGameplay/Source/MassLOD/Public/MassLODTypes.h
+++ b/Engine/Plugins/Runtime/MassGameplay/Source/MassLOD/Public/MassLODTypes.h
@@ -129,7 +129,7 @@ private:
 	bool bShouldTickThisFrame = true;
 	EMassLOD::Type LOD = EMassLOD::Max;
 	float TimeUntilNextTick = 0.0f;
-	int32 LastChunkSerialModificationNumber = -1;
+	int32 LastChunkSerialModificationNumber = INDEX_NONE;
 };
 
 enum class EMassVisibility : uint8
diff --git a/Engine/Plugins/Runtime/MassGameplay/Source/MassRepresentation/Private/MassRepresentationProcessor.cpp b/Engine/Plugins/Runtime/MassGameplay/Source/MassRepresentation/Private/MassRepresentationProcessor.cpp
--- a/Engine/Plugins/Runtime/MassGameplay/Source/MassRepresentation/Private/MassRepresentationProcessor.cpp
+++ b/Engine/Plugins/Runtime/MassGameplay/Source/MassRepresentation/Private/MassRepresentationProcessor.cpp
@@ -344,6 +344,8 @@ void UMassRepresentationProcessor::UpdateVisualization(FMassExecutionContext& Co
 
 FMassVisualizationChunkFragment& UMassRepresentationProcessor::UpdateChunkVisibility(FMassExecutionContext& Context) const
 {
+	bool bFirstUpdate = false;
+
 	// Setup chunk fragment data about visibility
 	FMassVisualizationChunkFragment& ChunkData = Context.GetMutableChunkFragment<FMassVisualizationChunkFragment>();
 	EMassVisibility ChunkVisibility = ChunkData.GetVisibility();
@@ -352,6 +354,7 @@ FMassVisualizationChunkFragment& UMassRepresentationProcessor::UpdateChunkVisibi
 		// The visibility on the chunk fragment data isn't set yet, let see if the Archetype has an visibility tag and set it on the ChunkData
 		ChunkVisibility = UE::MassRepresentation::GetVisibilityFromArchetype(Context);
 		ChunkData.SetVisibility(ChunkVisibility);
+		bFirstUpdate = true;
 	}
 	else
 	{
@@ -361,11 +364,19 @@ FMassVisualizationChunkFragment& UMassRepresentationProcessor::UpdateChunkVisibi
 	if (ChunkVisibility == EMassVisibility::CulledByDistance)
 	{
 		float DeltaTime = ChunkData.GetDeltaTime();
-		if (DeltaTime < 0.0f)
+		if (bFirstUpdate)
+		{
+			// A DeltaTime of 0.0f means it will tick this frame.
+			DeltaTime = FMath::RandRange(0.0f, NotVisibleUpdateRate);
+		}
+		else 
 		{
-			DeltaTime += NotVisibleUpdateRate * (1.0f + FMath::RandRange(-0.1f, 0.1f));
+			if (DeltaTime < 0.0f)
+			{
+				DeltaTime += NotVisibleUpdateRate * (1.0f + FMath::RandRange(-0.1f, 0.1f));
+			}
+			DeltaTime -= Context.GetDeltaTimeSeconds();
 		}
-		DeltaTime -= Context.GetDeltaTimeSeconds();
 
 		ChunkData.Update(DeltaTime);
 	}
diff --git a/Engine/Plugins/Runtime/MassGameplay/Source/MassRepresentation/Private/MassVisualizationComponent.cpp b/Engine/Plugins/Runtime/MassGameplay/Source/MassRepresentation/Private/MassVisualizationComponent.cpp
--- a/Engine/Plugins/Runtime/MassGameplay/Source/MassRepresentation/Private/MassVisualizationComponent.cpp
+++ b/Engine/Plugins/Runtime/MassGameplay/Source/MassRepresentation/Private/MassVisualizationComponent.cpp
@@ -64,7 +64,7 @@ void UMassVisualizationComponent::ConstructStaticMeshComponents()
 		}
 		for (const FStaticMeshInstanceVisualizationMeshDesc& MeshDesc : Info.Desc.Meshes)
 		{			
-			FISMCSharedData* SharedData = ISMCSharedData.Find(MeshDesc.Mesh);
+			FISMCSharedData* SharedData = ISMCSharedData.Find(MeshDesc);
 			UInstancedStaticMeshComponent* ISMC = SharedData ? SharedData->ISMC : nullptr;
 			if (SharedData)
 			{
@@ -86,10 +86,11 @@ void UMassVisualizationComponent::ConstructStaticMeshComponents()
 				ISMC->SetCanEverAffectNavigation(false);
 				ISMC->SetCollisionProfileName(UCollisionProfile::NoCollision_ProfileName);
 				ISMC->SetCastShadow(MeshDesc.bCastShadows);
+				ISMC->Mobility = MeshDesc.Mobility;
 				ISMC->SetReceivesDecals(false);
 				ISMC->RegisterComponent();
 
-				ISMCSharedData.Emplace(MeshDesc.Mesh, FISMCSharedData(ISMC));
+				ISMCSharedData.Emplace(MeshDesc, FISMCSharedData(ISMC));
 			}
 
 			Info.InstancedStaticMeshComponents.Add(ISMC);
@@ -133,7 +134,7 @@ void UMassVisualizationComponent::ConstructStaticMeshComponents()
 				const bool bAddMeshInRange = (Range.MinSignificance >= MeshDesc.MinLODSignificance && Range.MinSignificance < MeshDesc.MaxLODSignificance);
 				if (bAddMeshInRange)
 				{
-					Range.StaticMeshRefs.Add(MeshDesc.Mesh);
+					Range.StaticMeshRefs.Add(MeshDesc);
 				}
 			}
 		}
@@ -270,11 +271,11 @@ void UMassVisualizationComponent::EndVisualChanges()
 // FMassInstancedStaticMeshInfo
 //---------------------------------------------------------------
 
-void FMassInstancedStaticMeshInfo::ClearVisualInstance(TMap<UStaticMesh*, FISMCSharedData>& ISMCSharedData)
+void FMassInstancedStaticMeshInfo::ClearVisualInstance(FISMCSharedDataMap& ISMCSharedData)
 {
 	for (int i = 0; i < Desc.Meshes.Num(); i++)
 	{
-		FISMCSharedData* SharedData = ISMCSharedData.Find(Desc.Meshes[i].Mesh);
+		FISMCSharedData* SharedData = ISMCSharedData.Find(Desc.Meshes[i]);
 		if (SharedData)
 		{
 			SharedData->RefCount -= 1;
@@ -282,7 +283,7 @@ void FMassInstancedStaticMeshInfo::ClearVisualInstance(TMap<UStaticMesh*, FISMCS
 			{
 				SharedData->ISMC->ClearInstances();
 				SharedData->ISMC->DestroyComponent();
-				ISMCSharedData.Remove(Desc.Meshes[i].Mesh);
+				ISMCSharedData.Remove(Desc.Meshes[i]);
 			}
 		}
 	}
diff --git a/Engine/Plugins/Runtime/MassGameplay/Source/MassRepresentation/Private/MassVisualizationTrait.cpp b/Engine/Plugins/Runtime/MassGameplay/Source/MassRepresentation/Private/MassVisualizationTrait.cpp
--- a/Engine/Plugins/Runtime/MassGameplay/Source/MassRepresentation/Private/MassVisualizationTrait.cpp
+++ b/Engine/Plugins/Runtime/MassGameplay/Source/MassRepresentation/Private/MassVisualizationTrait.cpp
@@ -37,6 +37,7 @@ void UMassVisualizationTrait::BuildTemplate(FMassEntityTemplateBuildContext& Bui
 		BuildContext.AddDeinitializer(*(RepresentationFragmentDeinitializerClass->GetDefaultObject<UMassProcessor>()));
 
 		BuildContext.AddFragment<FMassRepresentationLODFragment>();
+		BuildContext.AddTag<FMassVisibilityCulledByDistanceTag>();
 
 		BuildContext.AddChunkFragment<FMassVisualizationChunkFragment>();
 	}
diff --git a/Engine/Plugins/Runtime/MassGameplay/Source/MassRepresentation/Public/MassRepresentationTypes.h b/Engine/Plugins/Runtime/MassGameplay/Source/MassRepresentation/Public/MassRepresentationTypes.h
--- a/Engine/Plugins/Runtime/MassGameplay/Source/MassRepresentation/Public/MassRepresentationTypes.h
+++ b/Engine/Plugins/Runtime/MassGameplay/Source/MassRepresentation/Public/MassRepresentationTypes.h
@@ -55,6 +55,10 @@ struct FStaticMeshInstanceVisualizationMeshDesc
 	/** Controls whether the ISM can cast shadow or not */
 	UPROPERTY(EditAnywhere, Category = "Mass|Visual")
 	bool bCastShadows = false;
+	
+	/** Controls the mobility of the ISM */
+	UPROPERTY(EditAnywhere, Category = "Mass|Visual")
+	TEnumAsByte<EComponentMobility::Type> Mobility = EComponentMobility::Movable;
 
 	bool operator==(const FStaticMeshInstanceVisualizationMeshDesc& Other) const
 	{
@@ -62,7 +66,29 @@ struct FStaticMeshInstanceVisualizationMeshDesc
 			MaterialOverrides == Other.MaterialOverrides && 
 			FMath::IsNearlyEqual(MinLODSignificance, Other.MinLODSignificance, KINDA_SMALL_NUMBER) &&
 			FMath::IsNearlyEqual(MaxLODSignificance, Other.MaxLODSignificance, KINDA_SMALL_NUMBER) &&
-			bCastShadows == Other.bCastShadows;
+			bCastShadows == Other.bCastShadows && 
+			Mobility == Other.Mobility;
+	}
+	
+	friend FORCEINLINE uint32 GetTypeHash(const FStaticMeshInstanceVisualizationMeshDesc& MeshDesc)
+	{
+		uint32 Hash = 0x0;
+		Hash = PointerHash(MeshDesc.Mesh, Hash);
+		Hash = HashCombine(GetTypeHash(MeshDesc.bCastShadows), Hash);
+		Hash = HashCombine(GetTypeHash(MeshDesc.Mobility), Hash);
+		for (UMaterialInterface* MaterialOverride : MeshDesc.MaterialOverrides)
+		{
+			if (MaterialOverride)
+			{
+				Hash = PointerHash(MaterialOverride, Hash);
+			}
+		}
+		return Hash;
+	}
+
+	operator uint32() const
+	{
+		return GetTypeHash(*this);
 	}
 };
 
@@ -124,6 +150,8 @@ struct MASSREPRESENTATION_API FISMCSharedData
 	int32 WriteIterator = 0;
 };
 
+typedef TMap<uint32, FISMCSharedData> FISMCSharedDataMap;
+
 USTRUCT()
 struct MASSREPRESENTATION_API FMassLODSignificanceRange
 {
@@ -160,9 +188,9 @@ public:
 
 	/** The component handling these instances */
 	UPROPERTY()
-	TArray<UStaticMesh*> StaticMeshRefs;
+	TArray<uint32> StaticMeshRefs;
 
-	TMap<UStaticMesh*, FISMCSharedData>* ISMCSharedDataPtr = nullptr;
+	FISMCSharedDataMap* ISMCSharedDataPtr = nullptr;
 };
 
 USTRUCT()
@@ -182,6 +210,11 @@ public:
 	{
 		return Desc == OtherDesc;
 	}
+	
+	const FStaticMeshInstanceVisualizationDesc& GetDesc() const
+	{
+		return Desc;
+	}
 
 	/** Whether or not to transform the static meshes if not align the mass agent transform */
 	bool ShouldUseTransformOffset() const { return Desc.bUseTransformOffset; }
@@ -237,7 +270,7 @@ public:
 protected:
 
 	/** Destroy the visual instance */
-	void ClearVisualInstance(TMap<UStaticMesh*, FISMCSharedData>& ISMCSharedData);
+	void ClearVisualInstance(FISMCSharedDataMap& ISMCSharedData);
 
 	/** Information about this static mesh which will represent all instances */
 	UPROPERTY()
diff --git a/Engine/Plugins/Runtime/MassGameplay/Source/MassRepresentation/Public/MassVisualizationComponent.h b/Engine/Plugins/Runtime/MassGameplay/Source/MassRepresentation/Public/MassVisualizationComponent.h
--- a/Engine/Plugins/Runtime/MassGameplay/Source/MassRepresentation/Public/MassVisualizationComponent.h
+++ b/Engine/Plugins/Runtime/MassGameplay/Source/MassRepresentation/Public/MassVisualizationComponent.h
@@ -59,7 +59,7 @@ protected:
 	TArray<FMassInstancedStaticMeshInfo> InstancedStaticMeshInfos;
 	UE_MT_DECLARE_RW_RECURSIVE_ACCESS_DETECTOR(InstancedStaticMeshInfosDetector);
 
-	TMap<UStaticMesh*, FISMCSharedData> ISMCSharedData;
+	FISMCSharedDataMap ISMCSharedData;
 
 
 	/** Whether there is a need to create a StaticMeshComponent */
diff --git a/Engine/Plugins/Runtime/MassGameplay/Source/MassSpawner/Private/MassSpawner.cpp b/Engine/Plugins/Runtime/MassGameplay/Source/MassSpawner/Private/MassSpawner.cpp
--- a/Engine/Plugins/Runtime/MassGameplay/Source/MassSpawner/Private/MassSpawner.cpp
+++ b/Engine/Plugins/Runtime/MassGameplay/Source/MassSpawner/Private/MassSpawner.cpp
@@ -12,6 +12,8 @@
 #include "VisualLogger/VisualLogger.h"
 #include "MassEntityConfigAsset.h"
 #include "EngineUtils.h"
+#include "Engine/StreamableManager.h"
+#include "Engine/AssetManager.h"
 
 namespace UE::MassSpawner
 {
@@ -178,6 +180,11 @@ void AMassSpawner::BeginDestroy()
 	UMassSimulationSubsystem::GetOnAdjustTickSchematics().Remove(TickSchematicsAdjustHandle);
 
 	DoDespawning();
+	
+	if (StreamingHandle.IsValid() && StreamingHandle->IsActive())
+	{
+		StreamingHandle->CancelHandle();
+	}
 
 	Super::BeginDestroy();
 }
@@ -269,7 +276,7 @@ void AMassSpawner::DoSpawning()
 		float TotalProportion = 0.0f;
 		for (FMassSpawnPointGenerator& SpawnPointsGenerator : SpawnPointsGenerators)
 		{
-		if (SpawnPointsGenerator.GeneratorInstance)
+		    if (SpawnPointsGenerator.GeneratorInstance)
 			{
 				SpawnPointsGenerator.bPointsGenerated = false;
 				TotalProportion += SpawnPointsGenerator.Proportion;
@@ -281,34 +288,66 @@ void AMassSpawner::DoSpawning()
 			UE_VLOG_UELOG(this, LogMassSpawner, Error, TEXT("The total combined porportion of all the generator needs to be greater than 0.0f."));
 			return;
 		}
-
-		int32 SpawnPointCountRemaining = SpawnCount;
-		float ProportionRemaining = TotalProportion;
-		for(FMassSpawnPointGenerator& SpawnPointsGenerator : SpawnPointsGenerators)
+		// Check if it needs loading
+		if (StreamingHandle.IsValid() && StreamingHandle->IsActive())
 		{
-			// Still needs to call the OnSpawnPointGenerationFinished for Generator we don't want to output Spawn Points for as the system will wait
-			// for the bPointsGenerated == true on every Generators before proceeding to the next steps.
-			if (SpawnPointsGenerator.Proportion == 0.0f || ProportionRemaining <= 0.0f || SpawnPointCountRemaining <= 0)
+			// @todo, instead of blindly canceling, we should remember what was asked to load with that handle and compare if more is needed?
+			StreamingHandle->CancelHandle();
+		}
+		TArray<FSoftObjectPath> AssetsToLoad;
+		for (const FMassSpawnedEntityType& EntityType : EntityTypes)
+		{
+			const int32 EntityCount = int32(SpawnCount * EntityType.Proportion / TotalProportion);
+			if (EntityCount > 0)
+
 			{
-				const TArray<FVector> EmptyLoc;
-				OnSpawnPointGenerationFinished(EmptyLoc, &SpawnPointsGenerator);
-				continue;
+				if (!EntityType.IsLoaded())
+				{
+					AssetsToLoad.Add(EntityType.EntityConfig.ToSoftObjectPath());
+				}
 			}
-
-			if (SpawnPointsGenerator.GeneratorInstance)
+		}
+		auto GenerateSpawningPoints = [this, SpawnCount, TotalProportion]()
+		{
+			int32 SpawnPointCountRemaining = SpawnCount;
+			float ProportionRemaining = TotalProportion;
+			for (FMassSpawnPointGenerator& SpawnPointsGenerator : SpawnPointsGenerators)
 			{
-				const float ProportionRatio = FMath::Min(SpawnPointsGenerator.Proportion / ProportionRemaining, 1.0f);
-				const int32 SpawnPointCount = FMath::CeilToInt(SpawnPointCountRemaining * ProportionRatio);
-				if(SpawnPointCount > 0)
+				// Still needs to call the OnSpawnPointGenerationFinished for Generator we don't want to output Spawn Points for as the system will wait
+				// for the bPointsGenerated == true on every Generators before proceeding to the next steps.
+				if (SpawnPointsGenerator.Proportion == 0.0f || ProportionRemaining <= 0.0f || SpawnPointCountRemaining <= 0)
 				{
-					FFinishedGeneratingSpawnPointsSignature Delegate = FFinishedGeneratingSpawnPointsSignature::CreateUObject(this, &AMassSpawner::OnSpawnPointGenerationFinished, &SpawnPointsGenerator);
-					SpawnPointsGenerator.GeneratorInstance->GenerateSpawnPoints(*this, SpawnPointCount, Delegate);
-					SpawnPointCountRemaining -= SpawnPointCount;
-					ProportionRemaining -= SpawnPointsGenerator.Proportion;
+					const TArray<FVector> EmptyLoc;
+					OnSpawnPointGenerationFinished(EmptyLoc, &SpawnPointsGenerator);
+					continue;
+				}
+
+				if (SpawnPointsGenerator.GeneratorInstance)
+				{
+					const float ProportionRatio = FMath::Min(SpawnPointsGenerator.Proportion / ProportionRemaining, 1.0f);
+					const int32 SpawnPointCount = FMath::CeilToInt(SpawnPointCountRemaining * ProportionRatio);
+					if (SpawnPointCount > 0)
+					{
+						FFinishedGeneratingSpawnPointsSignature Delegate = FFinishedGeneratingSpawnPointsSignature::CreateUObject(this, &AMassSpawner::OnSpawnPointGenerationFinished, &SpawnPointsGenerator);
+						SpawnPointsGenerator.GeneratorInstance->GenerateSpawnPoints(*this, SpawnPointCount, Delegate);
+						SpawnPointCountRemaining -= SpawnPointCount;
+						ProportionRemaining -= SpawnPointsGenerator.Proportion;
+					}
 				}
 			}
+		};
+		
+		if (AssetsToLoad.Num())
+		{
+			FStreamableManager& StreamableManager = UAssetManager::GetStreamableManager();
+			StreamingHandle = StreamableManager.RequestAsyncLoad(AssetsToLoad, GenerateSpawningPoints);
+		}
+		else
+		{
+			GenerateSpawningPoints();
 		}
 	}
+
 }
 
 void AMassSpawner::OnSpawnPointGenerationFinished(const TArray<FVector>& Locations, FMassSpawnPointGenerator* FinishedGenerator)
@@ -328,9 +367,9 @@ void AMassSpawner::OnSpawnPointGenerationFinished(const TArray<FVector>& Locatio
 		bAllSpawnPointsGenerated &= SpawnPointsGenerator.bPointsGenerated;
 	}
 
-	checkf(bFoundFinishedGenerator, TEXT("Something went wrong, we are receiving a callback on an unknow spawn point generator"))
+	checkf(bFoundFinishedGenerator, TEXT("Something went wrong, we are receiving a callback on an unknow spawn point generator"));
 	
-	if(bAllSpawnPointsGenerated)
+	if (bAllSpawnPointsGenerated)
 	{
 		SpawnAtLocations(AllGeneratedLocations);
 	}
@@ -361,7 +400,7 @@ void AMassSpawner::SpawnAtLocations(const TArray<FVector>& Locations)
 	}
 
 	const int32 SpawnCount = GetSpawnCount();
-	if(SpawnCount > 0)
+	if (SpawnCount > 0)
 	{
 		float TotalProportion = 0.0f;
 		for (const FMassSpawnedEntityType& EntityType : EntityTypes)
@@ -379,10 +418,10 @@ void AMassSpawner::SpawnAtLocations(const TArray<FVector>& Locations)
 
 		for (const FMassSpawnedEntityType& EntityType : EntityTypes)
 		{
-			if (const UMassEntityConfigAsset* EntityConfig = EntityType.GetEntityConfig())
+			const int32 EntityCount = int32(SpawnCount * EntityType.Proportion / TotalProportion);
+			if (EntityCount > 0)
 			{
-				const int32 EntityCount = int32(SpawnCount * EntityType.Proportion / TotalProportion);
-				if(EntityCount > 0)
+				if (const UMassEntityConfigAsset* EntityConfig = EntityType.GetEntityConfig())
 				{
 					FMassSpawnConfigBase SpawnConfig;
 					SpawnConfig.MinNumber = EntityCount;
@@ -463,6 +502,11 @@ void AMassSpawner::UnloadConfig()
 	{
 		EntityType.UnloadEntityConfig();
 	}
+
+	if (StreamingHandle.IsValid() && StreamingHandle->IsActive())
+	{
+		StreamingHandle->CancelHandle();
+	}
 }
 
 void AMassSpawner::OnAdjustTickSchematics(UWorld* World,TArray<TSoftObjectPtr<UMassSchematic>>& InOutTickSchematics)
diff --git a/Engine/Plugins/Runtime/MassGameplay/Source/MassSpawner/Public/MassSpawner.h b/Engine/Plugins/Runtime/MassGameplay/Source/MassSpawner/Public/MassSpawner.h
--- a/Engine/Plugins/Runtime/MassGameplay/Source/MassSpawner/Public/MassSpawner.h
+++ b/Engine/Plugins/Runtime/MassGameplay/Source/MassSpawner/Public/MassSpawner.h
@@ -11,6 +11,7 @@
 
 class UMassEntitySpawnPointsGeneratorConfigAsset;
 class UMassSchematic;
+struct FStreamableHandle;
 
 /** A spawner you can put on a map and configure it to spawn different things */
 UCLASS(hidecategories = (Object, Actor, Input, Rendering, LOD, Cooking, Collision, HLOD, Partition))
@@ -128,6 +129,8 @@ protected:
 
 	TArray<FVector> AllGeneratedLocations;
 
+	TSharedPtr<FStreamableHandle> StreamingHandle;
+
 #if WITH_EDITORONLY_DATA
 private:
 	UPROPERTY()
diff --git a/Engine/Plugins/Runtime/MassGameplay/Source/MassSpawner/Public/MassSpawnerTypes.h b/Engine/Plugins/Runtime/MassGameplay/Source/MassSpawner/Public/MassSpawnerTypes.h
--- a/Engine/Plugins/Runtime/MassGameplay/Source/MassSpawner/Public/MassSpawnerTypes.h
+++ b/Engine/Plugins/Runtime/MassGameplay/Source/MassSpawner/Public/MassSpawnerTypes.h
@@ -64,6 +64,8 @@ struct FMassSpawnedEntityType
 	/** Proportion of the count that should be this agent type, (the proportions will be normalized with other sibling agent types) */
 	UPROPERTY(BlueprintReadWrite, EditAnywhere, Category = "Mass|Spawn", meta = (ClampMin = 0.0f, UIMin = 0.0f))
 	float Proportion = 1.0f;
+	
+	bool IsLoaded() const { return EntityConfigPtr != nullptr; }
 
 	void UnloadEntityConfig()
 	{
