Added vis-logging of spawn locations, controlled by newly added UMassGameplaySettings.bLogSpawnLocations.

diff --git a/Engine/Plugins/Runtime/MassGameplay/Source/MassCommon/Public/MassGameplaySettings.h b/Engine/Plugins/Runtime/MassGameplay/Source/MassCommon/Public/MassGameplaySettings.h
new file mode 100644
--- /dev/null
+++ b/Engine/Plugins/Runtime/MassGameplay/Source/MassCommon/Public/MassGameplaySettings.h
@@ -0,0 +1,18 @@
+// Copyright Epic Games, Inc. All Rights Reserved.
+
+#pragma once
+
+#include "MassSettings.h"
+#include "MassGameplaySettings.generated.h"
+
+
+UCLASS(config = Mass, defaultconfig, DisplayName = "Mass Gameplay")
+class MASSCOMMON_API UMassGameplaySettings : public UMassModuleSettings
+{
+	GENERATED_BODY()
+
+public:
+
+	UPROPERTY(EditAnywhere, config, Category = Debug, meta = (ConsoleVariable = "mass.debug.VLogSpawnLocations"))
+	bool bLogSpawnLocations = true;
+};
diff --git a/Engine/Plugins/Runtime/MassGameplay/Source/MassSpawner/Private/MassSpawner.cpp b/Engine/Plugins/Runtime/MassGameplay/Source/MassSpawner/Private/MassSpawner.cpp
--- a/Engine/Plugins/Runtime/MassGameplay/Source/MassSpawner/Private/MassSpawner.cpp
+++ b/Engine/Plugins/Runtime/MassGameplay/Source/MassSpawner/Private/MassSpawner.cpp
@@ -14,6 +14,7 @@
 #include "EngineUtils.h"
 #include "Engine/StreamableManager.h"
 #include "Engine/AssetManager.h"
+#include "MassGameplaySettings.h"
 
 namespace UE::MassSpawner
 {
@@ -276,7 +277,7 @@ void AMassSpawner::DoSpawning()
 		float TotalProportion = 0.0f;
 		for (FMassSpawnPointGenerator& SpawnPointsGenerator : SpawnPointsGenerators)
 		{
-		    if (SpawnPointsGenerator.GeneratorInstance)
+			if (SpawnPointsGenerator.GeneratorInstance)
 			{
 				SpawnPointsGenerator.bPointsGenerated = false;
 				TotalProportion += SpawnPointsGenerator.Proportion;
@@ -399,6 +400,21 @@ void AMassSpawner::SpawnAtLocations(const TArray<FVector>& Locations)
 		Transform.SetLocation(Location);
 	}
 
+#if ENABLE_VISUAL_LOG
+	UE_VLOG(this, LogMassSpawner, Log, TEXT("Spawning at %d locations"), Locations.Num());
+	if (GetDefault<UMassGameplaySettings>()->bLogSpawnLocations)
+	{
+		if (FVisualLogEntry* LogEntry = FVisualLogger::Get().GetLastEntryForObject(this))
+		{
+			FVisualLogShapeElement Element(TEXT(""), FColor::Orange, 20.f, LogMassSpawner.GetCategoryName());
+			Element.Points += Locations;
+			Element.Type = EVisualLoggerShapeElement::SinglePoint;
+			Element.Verbosity = ELogVerbosity::Display;
+			LogEntry->AddElement(Element);
+		}
+	}
+#endif // ENABLE_VISUAL_LOG
+
 	const int32 SpawnCount = GetSpawnCount();
 	if (SpawnCount > 0)
 	{
