Integrating Mass AI To //UE5/Release-5.0

#prefilght 618155fbfc786a000146f869

diff --git a/Engine/Plugins/AI/MassCrowd/Source/MassCrowd/Private/MassCrowdRepresentationProcessor.cpp b/Engine/Plugins/AI/MassCrowd/Source/MassCrowd/Private/MassCrowdRepresentationProcessor.cpp
--- a/Engine/Plugins/AI/MassCrowd/Source/MassCrowd/Private/MassCrowdRepresentationProcessor.cpp
+++ b/Engine/Plugins/AI/MassCrowd/Source/MassCrowd/Private/MassCrowdRepresentationProcessor.cpp
@@ -146,7 +146,19 @@ void UMassCrowdRepresentationProcessor::TeleportActor(const FTransform& Transfor
 
 	if (const UCapsuleComponent* CapsuleComp = Actor.FindComponentByClass<UCapsuleComponent>())
 	{
-		RootTransform.AddToTranslation(FVector(0.0f, 0.0f, CapsuleComp->GetScaledCapsuleHalfHeight()));
+		const FVector HalfHeight(0.0f, 0.0f, CapsuleComp->GetScaledCapsuleHalfHeight());
+		RootTransform.AddToTranslation(HalfHeight);
+		const FVector RootLocation = RootTransform.GetLocation();
+		const FVector SweepOffset(0.0f, 0.0f, 20.0f);
+		const FVector Start = RootLocation + SweepOffset;
+		const FVector End = RootLocation - SweepOffset;
+		FCollisionQueryParams Params;
+		Params.AddIgnoredActor(&Actor);
+		FHitResult OutHit;
+		if (Actor.GetWorld()->SweepSingleByChannel(OutHit, Start, End, Transform.GetRotation(), CapsuleComp->GetCollisionObjectType(), CapsuleComp->GetCollisionShape(), Params))
+		{
+			RootTransform.SetLocation(OutHit.Location);
+		}
 	}
 	Super::TeleportActor(RootTransform, Actor, Context);
 }
diff --git a/Engine/Plugins/AI/MassCrowd/Source/MassCrowd/Private/MassCrowdVisualizationProcessor.cpp b/Engine/Plugins/AI/MassCrowd/Source/MassCrowd/Private/MassCrowdVisualizationProcessor.cpp
--- a/Engine/Plugins/AI/MassCrowd/Source/MassCrowd/Private/MassCrowdVisualizationProcessor.cpp
+++ b/Engine/Plugins/AI/MassCrowd/Source/MassCrowd/Private/MassCrowdVisualizationProcessor.cpp
@@ -28,6 +28,7 @@ UMassCrowdVisualizationProcessor::UMassCrowdVisualizationProcessor()
 	ExecutionOrder.ExecuteAfter.Add(UE::Mass::ProcessorGroupNames::LOD);
 
 	bRequiresGameThreadExecution = true;
+	bKeepActorExtraFrame = true;
 }
 
 void UMassCrowdVisualizationProcessor::ConfigureQueries()
diff --git a/Engine/Plugins/Runtime/MassEntity/Source/MassEntity/Private/MassArchetypeData.cpp b/Engine/Plugins/Runtime/MassEntity/Source/MassEntity/Private/MassArchetypeData.cpp
--- a/Engine/Plugins/Runtime/MassEntity/Source/MassEntity/Private/MassArchetypeData.cpp
+++ b/Engine/Plugins/Runtime/MassEntity/Source/MassEntity/Private/MassArchetypeData.cpp
@@ -215,8 +215,7 @@ void FMassArchetypeData::RemoveEntityInternal(const int32 AbsoluteIndex, const b
 
 	FMassArchetypeChunk& Chunk = Chunks[ChunkIndex];
 
-	Chunk.RemoveInstance();
-	const int32 IndexToSwapFrom = Chunk.GetNumInstances();
+	const int32 IndexToSwapFrom = Chunk.GetNumInstances() - 1;
 
 	checkf(bDestroyFragments || UE::Mass::Core::bBitwiseRelocateFragments, TEXT("We allow not to destroy fragments only in bit wise relocation mode."));
 
@@ -266,6 +265,8 @@ void FMassArchetypeData::RemoveEntityInternal(const int32 AbsoluteIndex, const b
 			FragmentConfig.FragmentType->DestroyStruct(DyingFragmentPtr);
 		}
 	}
+	
+	Chunk.RemoveInstance();
 
 	// If the chunk itself is empty now, see if we can remove it entirely
 	// Note: This is only possible for trailing chunks, to avoid messing up the absolute indices in the entities map
@@ -759,10 +760,20 @@ void FMassArchetypeData::BindChunkFragmentRequirements(FMassExecutionContext& Ru
 
 SIZE_T FMassArchetypeData::GetAllocatedSize() const
 {
+	int32 NumAllocatedChunkBuffers = 0;
+	for (const FMassArchetypeChunk& Chunk : Chunks)
+	{
+		if (Chunk.GetRawMemory() != nullptr)
+		{
+			++NumAllocatedChunkBuffers;
+		}
+	}
+
 	return sizeof(FMassArchetypeData) +
+		ChunkFragmentTemplates.GetAllocatedSize() +
 		FragmentConfigs.GetAllocatedSize() +
 		Chunks.GetAllocatedSize() +
-		(Chunks.Num() * GetChunkAllocSize()) +
+		(NumAllocatedChunkBuffers * GetChunkAllocSize()) +
 		EntityMap.GetAllocatedSize() +
 		FragmentIndexMap.GetAllocatedSize();
 }
diff --git a/Engine/Plugins/Runtime/MassEntity/Source/MassEntity/Private/MassArchetypeData.h b/Engine/Plugins/Runtime/MassEntity/Source/MassEntity/Private/MassArchetypeData.h
--- a/Engine/Plugins/Runtime/MassEntity/Source/MassEntity/Private/MassArchetypeData.h
+++ b/Engine/Plugins/Runtime/MassEntity/Source/MassEntity/Private/MassArchetypeData.h
@@ -15,20 +15,27 @@ struct FMassArchetypeChunk
 {
 private:
 	uint8* RawMemory = nullptr;
+	int32 AllocSize = 0;
 	int32 NumInstances = 0;
 	int32 SerialModificationNumber = 0;
 	TArray<FInstancedStruct> ChunkFragmentData;
 
 public:
-	explicit FMassArchetypeChunk(int32 AllocSize, TConstArrayView<FInstancedStruct> InChunkFragmentTemplates)
-		: ChunkFragmentData(InChunkFragmentTemplates)
+	explicit FMassArchetypeChunk(int32 InAllocSize, TConstArrayView<FInstancedStruct> InChunkFragmentTemplates)
+		: AllocSize(InAllocSize)
+		, ChunkFragmentData(InChunkFragmentTemplates)
 	{
 		RawMemory = (uint8*)FMemory::Malloc(AllocSize);
 	}
 
 	~FMassArchetypeChunk()
 	{
-		FMemory::Free(RawMemory);
+		// Only release memory if it was not done already.
+		if (RawMemory != nullptr)
+		{
+			FMemory::Free(RawMemory);
+			RawMemory = nullptr;
+		}
 	}
 
 	// Returns the Entity array element at the specified index
@@ -58,6 +65,14 @@ public:
 		NumInstances -= Count;
 		check(NumInstances >= 0);
 		SerialModificationNumber++;
+
+		// Because we only remove trailing chunks to avoid messing up the absolute indices in the entities map,
+		// We are freeing the memory here to save memory
+		if (NumInstances == 0)
+		{
+			FMemory::Free(RawMemory);
+			RawMemory = nullptr;
+		}
 	}
 
 	void AddInstance()
@@ -90,6 +105,12 @@ public:
 		checkf(NumInstances == 0, TEXT("Recycling a chunk that is not empty."));
 		SerialModificationNumber++;
 		ChunkFragmentData = InChunkFragmentsTemplate;
+		
+		// If this chunk previously had entity and it does not anymore, we might have to reallocate the memory as it was freed to save memory
+		if (RawMemory == nullptr)
+		{
+			RawMemory = (uint8*)FMemory::Malloc(AllocSize);
+		}
 	}
 
 #if WITH_MASSENTITY_DEBUG
diff --git a/Engine/Plugins/Runtime/MassGameplay/Source/MassActors/Private/MassAgentTraits.cpp b/Engine/Plugins/Runtime/MassGameplay/Source/MassActors/Private/MassAgentTraits.cpp
--- a/Engine/Plugins/Runtime/MassGameplay/Source/MassActors/Private/MassAgentTraits.cpp
+++ b/Engine/Plugins/Runtime/MassGameplay/Source/MassActors/Private/MassAgentTraits.cpp
@@ -169,7 +169,8 @@ void UMassAgentFeetLocationSyncTrait::BuildTemplate(FMassEntityTemplateBuildCont
 				// the entity is the authority
 				if (CurrentDirection == EMassTranslationDirection::MassToActor)
 				{
-					Component->SetWorldLocation(TransformFragment.GetTransform().GetLocation() + FVector(0.f, 0.f, Component->Bounds.BoxExtent.Z));
+					// Temporary disabling this as it is already done earlier in the MassRepresentation and we needed to do a sweep to find he floor
+					//Component->SetWorldLocation(FeetLocation, /*bSweep*/true, nullptr, ETeleportType::TeleportPhysics);
 				}
 				// actor is the authority
 				else
diff --git a/Engine/Plugins/Runtime/MassGameplay/Source/MassLOD/Private/MassLODLogic.cpp b/Engine/Plugins/Runtime/MassGameplay/Source/MassLOD/Private/MassLODLogic.cpp
--- a/Engine/Plugins/Runtime/MassGameplay/Source/MassLOD/Private/MassLODLogic.cpp
+++ b/Engine/Plugins/Runtime/MassGameplay/Source/MassLOD/Private/MassLODLogic.cpp
@@ -6,8 +6,8 @@
 
 void FMassLODBaseLogic::CacheViewerInformation(TConstArrayView<FViewerInfo> Viewers, const bool bLocalViewersOnly)
 {
-	NumOfViewers = Viewers.Num();
-	check(NumOfViewers <= UE::MassLOD::MaxNumOfViewers);
+	check(Viewers.Num() <= UE::MassLOD::MaxNumOfViewers);
+	NumOfViewers = FMath::Min(Viewers.Num(), UE::MassLOD::MaxNumOfViewers);
 
 	// Cache viewer info
 	for (int ViewerIdx = 0; ViewerIdx < NumOfViewers; ++ViewerIdx)
diff --git a/Engine/Plugins/Runtime/MassGameplay/Source/MassRepresentation/Private/MassRepresentationProcessor.cpp b/Engine/Plugins/Runtime/MassGameplay/Source/MassRepresentation/Private/MassRepresentationProcessor.cpp
--- a/Engine/Plugins/Runtime/MassGameplay/Source/MassRepresentation/Private/MassRepresentationProcessor.cpp
+++ b/Engine/Plugins/Runtime/MassGameplay/Source/MassRepresentation/Private/MassRepresentationProcessor.cpp
@@ -10,6 +10,12 @@
 #include "MassEntityView.h"
 #include "Engine/World.h"
 
+namespace UE::MassRepresentation
+{
+	int32 bAllowKeepActorExtraFrame = 1;
+	FAutoConsoleVariableRef CVarAllowKeepActorExtraFrame(TEXT("ai.massrepresentation.AllowKeepActorExtraFrame"), bAllowKeepActorExtraFrame, TEXT("Allow the mass representation to keep actor an extra frame when switching to ISM"), ECVF_Default);
+}
+
 //----------------------------------------------------------------------//
 // UMassRepresentationProcessor(); 
 //----------------------------------------------------------------------//
@@ -66,6 +72,8 @@ void UMassRepresentationProcessor::UpdateRepresentation(FMassExecutionContext& C
 	const TConstArrayView<FMassRepresentationLODFragment> RepresentationLODList = Context.GetFragmentView<FMassRepresentationLODFragment>();
 	const TArrayView<FDataFragment_Actor> ActorList = Context.GetMutableFragmentView<FDataFragment_Actor>();
 
+	const bool bDoKeepActorExtraFrame = UE::MassRepresentation::bAllowKeepActorExtraFrame ? bKeepLowResActors : false;
+
 	const int32 NumEntities = Context.GetNumEntities();
 	for (int32 EntityIdx = 0; EntityIdx < NumEntities; EntityIdx++)
 	{
@@ -75,6 +83,8 @@ void UMassRepresentationProcessor::UpdateRepresentation(FMassExecutionContext& C
 		FMassRepresentationFragment& Representation = RepresentationList[EntityIdx];
 		FDataFragment_Actor& ActorInfo = ActorList[EntityIdx];
 
+		// Keeping a copy of the that last calculated previous representation
+		const ERepresentationType PrevRepresentationCopy = Representation.PrevRepresentation;
 		Representation.PrevRepresentation = Representation.CurrentRepresentation;
 		
 		ERepresentationType WantedRepresentationType = LODRepresentation[FMath::Min((int32)RepresentationLOD.LOD, (int32)EMassLOD::Off)];
@@ -85,6 +95,37 @@ void UMassRepresentationProcessor::UpdateRepresentation(FMassExecutionContext& C
 		{
 			WantedRepresentationType = DefaultRepresentationType;
 		}
+		
+		auto DisableActorForISM = [&](AActor*& Actor)
+		{
+			if (!Actor || ActorInfo.IsOwnedByMass())
+			{
+				// Execute only if the high res is different than the low res Actor 
+				// Or if we do not wish to keep the low res actor while in ISM
+				if (Representation.HighResTemplateActorIndex != Representation.LowResTemplateActorIndex || !bKeepLowResActors)
+				{
+					// Try releasing the high actor or any high res spawning request
+					if (ReleaseActorOrCancelSpawning(MassAgent, ActorInfo, Representation.HighResTemplateActorIndex, Representation.ActorSpawnRequestHandle, Context))
+					{
+						Actor = ActorInfo.GetOwnedByMassMutable();
+					}
+					// Do not do the same with low res if indicated so
+					if (!bKeepLowResActors && ReleaseActorOrCancelSpawning(MassAgent, ActorInfo, Representation.LowResTemplateActorIndex, Representation.ActorSpawnRequestHandle, Context))
+					{
+						Actor = ActorInfo.GetOwnedByMassMutable();
+					}
+				}
+				// If we already queued spawn request but have changed our mind, continue with it but once we get the actor back, disable it immediately
+				if (Representation.ActorSpawnRequestHandle.IsValid())
+				{
+					Actor = GetOrSpawnActor(MassAgent, ActorInfo, TransformFragment.GetTransform(), Representation.LowResTemplateActorIndex, Representation.ActorSpawnRequestHandle, GetSpawnPriority(RepresentationLOD));
+				}
+			}
+			if (Actor != nullptr)
+			{
+				SetActorEnabled(EActorEnabledType::Disabled, *Actor, EntityIdx, Context);
+			}
+		};
 
 		// Process switch between representation if there is a change in the representation or there is a pending spawning request
 		if (WantedRepresentationType != Representation.CurrentRepresentation || Representation.ActorSpawnRequestHandle.IsValid())
@@ -92,6 +133,7 @@ void UMassRepresentationProcessor::UpdateRepresentation(FMassExecutionContext& C
 			if (Representation.CurrentRepresentation == ERepresentationType::None)
 			{
 				Representation.PrevTransform = TransformFragment.GetTransform();
+				Representation.PrevLODSignificance = RepresentationLOD.LODSignificance;
 			}
 
 			AActor* Actor = ActorInfo.GetMutable();
@@ -154,35 +196,12 @@ void UMassRepresentationProcessor::UpdateRepresentation(FMassExecutionContext& C
 					break;
 				}
 				case ERepresentationType::StaticMeshInstance:
-					if (!Actor || ActorInfo.IsOwnedByMass())
+					if (!bDoKeepActorExtraFrame || 
+					   (Representation.PrevRepresentation != ERepresentationType::HighResSpawnedActor && Representation.PrevRepresentation != ERepresentationType::LowResSpawnedActor))
 					{
-						// If high res actor if it is different than the low res, 
-						if (Representation.HighResTemplateActorIndex != Representation.LowResTemplateActorIndex|| !bKeepLowResActors)
-						{
-							// Try releasing the high actor or any hig res spawning request
-							if (ReleaseActorOrCancelSpawning(MassAgent, ActorInfo, Representation.HighResTemplateActorIndex, Representation.ActorSpawnRequestHandle, Context))
-							{
-								Actor = ActorInfo.GetOwnedByMassMutable();
-							}
-
-							// Do not do the same with low res if indicated so
-							if (!bKeepLowResActors && ReleaseActorOrCancelSpawning(MassAgent, ActorInfo, Representation.LowResTemplateActorIndex, Representation.ActorSpawnRequestHandle, Context))
-							{
-								Actor = ActorInfo.GetOwnedByMassMutable();
-							}
-						}
-
-						// If we were in the process of spawning a low res actor, continue and keep it around disabled
-						if (Representation.ActorSpawnRequestHandle.IsValid())
-						{
-							Actor = GetOrSpawnActor(MassAgent, ActorInfo, TransformFragment.GetTransform(), Representation.LowResTemplateActorIndex, Representation.ActorSpawnRequestHandle, GetSpawnPriority(RepresentationLOD));
-						}
-					}
-
-					if (Actor != nullptr)
-					{
-						SetActorEnabled(EActorEnabledType::Disabled, *Actor, EntityIdx, Context);
+						DisableActorForISM(Actor);
 					}
+ 
 					Representation.CurrentRepresentation = ERepresentationType::StaticMeshInstance;
 					break;
 				case ERepresentationType::None:
@@ -203,6 +222,13 @@ void UMassRepresentationProcessor::UpdateRepresentation(FMassExecutionContext& C
 					break;
 			}
 		}
+		else if (bDoKeepActorExtraFrame && 
+				 Representation.PrevRepresentation == ERepresentationType::StaticMeshInstance &&
+			    (PrevRepresentationCopy == ERepresentationType::HighResSpawnedActor || PrevRepresentationCopy == ERepresentationType::LowResSpawnedActor))
+		{
+			AActor* Actor = ActorInfo.GetMutable();
+			DisableActorForISM(Actor);
+		}
 	}
 }
 
@@ -236,10 +262,10 @@ bool UMassRepresentationProcessor::ReleaseActorOrCancelSpawning(const FMassEntit
 	// the spawning of whatever SpawnRequestHandle reference to
 	const bool bSuccess = bCancelSpawningOnly ? RepresentationSubsystem->CancelSpawning(MassAgent, TemplateActorIndex, SpawnRequestHandle) :
 			RepresentationSubsystem->ReleaseTemplateActorOrCancelSpawning(MassAgent, TemplateActorIndex, Actor, SpawnRequestHandle);
-	if(bSuccess)
+	if (bSuccess)
 	{
 		Actor = ActorInfo.GetOwnedByMassMutable();
-		if(Actor && RepresentationSubsystem->DoesActorMatchTemplate(*Actor, TemplateActorIndex))
+		if (Actor && RepresentationSubsystem->DoesActorMatchTemplate(*Actor, TemplateActorIndex))
 		{
 			ActorInfo.ResetNoHandleMapUpdate();
 			
diff --git a/Engine/Plugins/Runtime/MassGameplay/Source/MassRepresentation/Private/MassUpdateISMProcessor.cpp b/Engine/Plugins/Runtime/MassGameplay/Source/MassRepresentation/Private/MassUpdateISMProcessor.cpp
--- a/Engine/Plugins/Runtime/MassGameplay/Source/MassRepresentation/Private/MassUpdateISMProcessor.cpp
+++ b/Engine/Plugins/Runtime/MassGameplay/Source/MassRepresentation/Private/MassUpdateISMProcessor.cpp
@@ -62,15 +62,16 @@ void UMassUpdateISMProcessor::Execute(UMassEntitySubsystem& EntitySubsystem, FMa
 
 				if (Representation.CurrentRepresentation == ERepresentationType::StaticMeshInstance)
 				{
-					UpdateISMTransform(GetTypeHash(Context.GetEntity(EntityIdx)), ISMInfo[Representation.StaticMeshDescIndex], TransformFragment.GetTransform(), Representation.PrevTransform, RepresentationLOD.LODSignificance);
+					UpdateISMTransform(GetTypeHash(Context.GetEntity(EntityIdx)), ISMInfo[Representation.StaticMeshDescIndex], TransformFragment.GetTransform(), Representation.PrevTransform, RepresentationLOD.LODSignificance, Representation.PrevLODSignificance);
 				}
 				Representation.PrevTransform = TransformFragment.GetTransform();
+				Representation.PrevLODSignificance = RepresentationLOD.LODSignificance;
 			}
 		});
 	}
 }
 
-void UMassUpdateISMProcessor::UpdateISMTransform(int32 EntityId, FMassInstancedStaticMeshInfo& ISMInfo, const FTransform& Transform, const FTransform& PrevTransform, const float LODSignificance)
+void UMassUpdateISMProcessor::UpdateISMTransform(int32 EntityId, FMassInstancedStaticMeshInfo& ISMInfo, const FTransform& Transform, const FTransform& PrevTransform, const float LODSignificance, const float PrevLODSignificance/* = -1.0f*/)
 {
 	if (ISMInfo.ShouldUseTransformOffset())
 	{
@@ -78,10 +79,10 @@ void UMassUpdateISMProcessor::UpdateISMTransform(int32 EntityId, FMassInstancedS
 		const FTransform SMTransform = TransformOffset * Transform;
 		const FTransform SMPrevTransform = TransformOffset * PrevTransform;
 
-		ISMInfo.AddBatchedTransform(EntityId, SMTransform, SMPrevTransform, LODSignificance);
+		ISMInfo.AddBatchedTransform(EntityId, SMTransform, SMPrevTransform, LODSignificance, PrevLODSignificance);
 	}
 	else
 	{
-		ISMInfo.AddBatchedTransform(EntityId, Transform, PrevTransform, LODSignificance);
+		ISMInfo.AddBatchedTransform(EntityId, Transform, PrevTransform, LODSignificance, PrevLODSignificance);
 	}
 }
diff --git a/Engine/Plugins/Runtime/MassGameplay/Source/MassRepresentation/Private/MassVisualizationComponent.cpp b/Engine/Plugins/Runtime/MassGameplay/Source/MassRepresentation/Private/MassVisualizationComponent.cpp
--- a/Engine/Plugins/Runtime/MassGameplay/Source/MassRepresentation/Private/MassVisualizationComponent.cpp
+++ b/Engine/Plugins/Runtime/MassGameplay/Source/MassRepresentation/Private/MassVisualizationComponent.cpp
@@ -296,11 +296,16 @@ void FMassInstancedStaticMeshInfo::ClearVisualInstance(FISMCSharedDataMap& ISMCS
 // FMassLODSignificanceRange
 //---------------------------------------------------------------
 
-void FMassLODSignificanceRange::AddBatchedTransform(const int32 InstanceId, const FTransform& Transform, const FTransform& PrevTransform)
+void FMassLODSignificanceRange::AddBatchedTransform(const int32 InstanceId, const FTransform& Transform, const FTransform& PrevTransform, const TArray<uint32>& ExcludeStaticMeshRefs)
 {
 	check(ISMCSharedDataPtr);
 	for (int i = 0; i < StaticMeshRefs.Num(); i++)
 	{
+		if (ExcludeStaticMeshRefs.Contains(StaticMeshRefs[i]))
+		{
+			continue;
+		}
+
 		FISMCSharedData& SharedData = (*ISMCSharedDataPtr)[StaticMeshRefs[i]];
 
 		SharedData.UpdateInstanceIds.Add(InstanceId);
@@ -309,21 +314,31 @@ void FMassLODSignificanceRange::AddBatchedTransform(const int32 InstanceId, cons
 	}
 }
 
-void FMassLODSignificanceRange::AddBatchedCustomDataFloats(const TArray<float>& CustomFloats)
+void FMassLODSignificanceRange::AddBatchedCustomDataFloats(const TArray<float>& CustomFloats, const TArray<uint32>& ExcludeStaticMeshRefs)
 {
 	check(ISMCSharedDataPtr);
 	for (int i = 0; i < StaticMeshRefs.Num(); i++)
 	{
+		if (ExcludeStaticMeshRefs.Contains(StaticMeshRefs[i]))
+		{
+			continue;
+		}
+
 		FISMCSharedData& SharedData = (*ISMCSharedDataPtr)[StaticMeshRefs[i]];
 		SharedData.StaticMeshInstanceCustomFloats.Append(CustomFloats);
 	}
 }
 
-void FMassLODSignificanceRange::WriteCustomDataFloatsAtStartIndex(int32 StaticMeshIndex, const TArrayView<float> CustomFloats, int32 FloatsPerInstance, int32 StartFloatIndex)
+void FMassLODSignificanceRange::WriteCustomDataFloatsAtStartIndex(int32 StaticMeshIndex, const TArrayView<float>& CustomFloats, const int32 FloatsPerInstance, const int32 StartFloatIndex, const TArray<uint32>& ExcludeStaticMeshRefs)
 {
 	check(ISMCSharedDataPtr);
 	if (StaticMeshRefs.IsValidIndex(StaticMeshIndex))
 	{
+		if (ExcludeStaticMeshRefs.Contains(StaticMeshRefs[StaticMeshIndex]))
+		{
+			return;
+		}
+
 		FISMCSharedData& SharedData = (*ISMCSharedDataPtr)[StaticMeshRefs[StaticMeshIndex]];
 
 		int32 StartIndex = FloatsPerInstance * SharedData.WriteIterator + StartFloatIndex;
diff --git a/Engine/Plugins/Runtime/MassGameplay/Source/MassRepresentation/Public/MassRepresentationFragments.h b/Engine/Plugins/Runtime/MassGameplay/Source/MassRepresentation/Public/MassRepresentationFragments.h
--- a/Engine/Plugins/Runtime/MassGameplay/Source/MassRepresentation/Public/MassRepresentationFragments.h
+++ b/Engine/Plugins/Runtime/MassGameplay/Source/MassRepresentation/Public/MassRepresentationFragments.h
@@ -51,6 +51,6 @@ struct MASSREPRESENTATION_API FMassRepresentationFragment : public FMassFragment
 
 	FTransform PrevTransform;
 
-	// @todo: Remove this when we address character movement
-	bool bControlCharacterMovement = true;
+	/** Value scaling from 0 to 3, 0 highest LOD we support and 3 being completely off LOD */
+	float PrevLODSignificance = -1.0f;
 };
\ No newline at end of file
diff --git a/Engine/Plugins/Runtime/MassGameplay/Source/MassRepresentation/Public/MassRepresentationProcessor.h b/Engine/Plugins/Runtime/MassGameplay/Source/MassRepresentation/Public/MassRepresentationProcessor.h
--- a/Engine/Plugins/Runtime/MassGameplay/Source/MassRepresentation/Public/MassRepresentationProcessor.h
+++ b/Engine/Plugins/Runtime/MassGameplay/Source/MassRepresentation/Public/MassRepresentationProcessor.h
@@ -174,6 +174,10 @@ protected:
 	UPROPERTY(EditAnywhere, Category = "Mass|Representation", config)
 	bool bKeepLowResActors = true;
 
+	/** When switching to ISM keep the actor an extra frame, helps cover rendering glitches (i.e. occlusion query being one frame late) */
+	UPROPERTY(EditAnywhere, Category = "Mass|Representation", config)
+	bool bKeepActorExtraFrame = false;
+
 	/** World Partition grid name to test collision against, default None will be the main grid */
 	UPROPERTY(EditAnywhere, Category = "Mass|Representation", config)
 	FName WorldPartitionGridNameContainingCollision;
diff --git a/Engine/Plugins/Runtime/MassGameplay/Source/MassRepresentation/Public/MassRepresentationTypes.h b/Engine/Plugins/Runtime/MassGameplay/Source/MassRepresentation/Public/MassRepresentationTypes.h
--- a/Engine/Plugins/Runtime/MassGameplay/Source/MassRepresentation/Public/MassRepresentationTypes.h
+++ b/Engine/Plugins/Runtime/MassGameplay/Source/MassRepresentation/Public/MassRepresentationTypes.h
@@ -158,12 +158,12 @@ struct MASSREPRESENTATION_API FMassLODSignificanceRange
 	GENERATED_BODY()
 public:
 
-	void AddBatchedTransform(const int32 InstanceId, const FTransform& Transform, const FTransform& PrevTransform);
+	void AddBatchedTransform(const int32 InstanceId, const FTransform& Transform, const FTransform& PrevTransform, const TArray<uint32>& ExcludeStaticMeshRefs);
 
 	// Adds the specified struct reinterpreted as custom floats to our custom data. Individual members of the specified struct should always fit into a float.
 	// When adding any custom data, the custom data must be added for every instance.
 	template<typename InCustomDataType>
-	void AddBatchedCustomData(InCustomDataType InCustomData, int32 NumFloatsToPad = 0)
+	void AddBatchedCustomData(InCustomDataType InCustomData, const TArray<uint32>& ExcludeStaticMeshRefs, int32 NumFloatsToPad = 0)
 	{
 		check(ISMCSharedDataPtr);
 		static_assert((sizeof(InCustomDataType) % sizeof(float)) == 0, "AddBatchedCustomData: InCustomDataType should have a total size multiple of sizeof(float), and have members that fit in a float's boundaries");
@@ -171,6 +171,11 @@ public:
 		const size_t StructSizeInFloats = StructSize / sizeof(float);
 		for (int i = 0; i < StaticMeshRefs.Num(); i++)
 		{
+			if (ExcludeStaticMeshRefs.Contains(StaticMeshRefs[i]))
+			{
+				continue;
+			}
+
 			FISMCSharedData& SharedData = (*ISMCSharedDataPtr)[StaticMeshRefs[i]];
 			const int32 StartIndex = SharedData.StaticMeshInstanceCustomFloats.AddDefaulted(StructSizeInFloats + NumFloatsToPad);
 			InCustomDataType* CustomData = reinterpret_cast<InCustomDataType*>(&SharedData.StaticMeshInstanceCustomFloats[StartIndex]);
@@ -178,9 +183,9 @@ public:
 		}
 	}
 
-	void AddBatchedCustomDataFloats(const TArray<float>& CustomFloats);
+	void AddBatchedCustomDataFloats(const TArray<float>& CustomFloats, const TArray<uint32>& ExcludeStaticMeshRefs);
 
-	void WriteCustomDataFloatsAtStartIndex(int32 StaticMeshIndex, const TArrayView<float> CustomFloats, int32 FloatsPerInstance, int32 StartIndex);
+	void WriteCustomDataFloatsAtStartIndex(int32 StaticMeshIndex, const TArrayView<float>& CustomFloats, const int32 FloatsPerInstance, const int32 StartIndex, const TArray<uint32>& ExcludeStaticMeshRefs);
 
 	/** LOD Significance range */
 	float MinSignificance;
@@ -232,38 +237,70 @@ public:
 		return nullptr;
 	}
 
-	FORCEINLINE void AddBatchedTransform(const int32 InstanceId, const FTransform& Transform, const FTransform& PrevTransform, float LODSignificance)
+	FORCEINLINE void AddBatchedTransform(const int32 InstanceId, const FTransform& Transform, const FTransform& PrevTransform, const float LODSignificance, const float PrevLODSignificance = -1.0f)
 	{
 		if (FMassLODSignificanceRange* Range = GetLODSignificanceRange(LODSignificance))
 		{
-			Range->AddBatchedTransform(InstanceId, Transform, PrevTransform);
+			Range->AddBatchedTransform(InstanceId, Transform, PrevTransform, TArray<uint32>());
+			if(PrevLODSignificance >= 0.0f)
+			{
+				FMassLODSignificanceRange* PrevRange = GetLODSignificanceRange(PrevLODSignificance);
+				if (PrevRange != Range)
+				{
+					PrevRange->AddBatchedTransform(InstanceId, Transform, PrevTransform, Range->StaticMeshRefs);
+				}
+			}
 		}
 	}
 
 	// Adds the specified struct reinterpreted as custom floats to our custom data. Individual members of the specified struct should always fit into a float.
 	// When adding any custom data, the custom data must be added for every instance.
 	template<typename InCustomDataType>
-	void AddBatchedCustomData(InCustomDataType InCustomData, float LODSignificance, int32 NumFloatsToPad = 0)
+	void AddBatchedCustomData(InCustomDataType InCustomData, const float LODSignificance, const float PrevLODSignificance = -1.0f, int32 NumFloatsToPad = 0)
 	{
 		if (FMassLODSignificanceRange* Range = GetLODSignificanceRange(LODSignificance))
 		{
-			Range->AddBatchedCustomData(InCustomData, NumFloatsToPad);
+			Range->AddBatchedCustomData(InCustomData, TArray<uint32>(), NumFloatsToPad);
+			if(PrevLODSignificance >= 0.0f)
+			{
+				FMassLODSignificanceRange* PrevRange = GetLODSignificanceRange(PrevLODSignificance);
+				if (PrevRange != Range)
+				{
+					PrevRange->AddBatchedCustomData(InCustomData, Range->StaticMeshRefs, NumFloatsToPad);
+				}
+			}
 		}
 	}
 
-	FORCEINLINE void AddBatchedCustomDataFloats(const TArray<float>& CustomFloats, float LODSignificance)
+	FORCEINLINE void AddBatchedCustomDataFloats(const TArray<float>& CustomFloats, const float LODSignificance, const float PrevLODSignificance = -1.0f)
 	{
 		if (FMassLODSignificanceRange* Range = GetLODSignificanceRange(LODSignificance))
 		{
-			Range->AddBatchedCustomDataFloats(CustomFloats);
+			Range->AddBatchedCustomDataFloats(CustomFloats, TArray<uint32>());
+			if(PrevLODSignificance >= 0.0f)
+			{
+				FMassLODSignificanceRange* PrevRange = GetLODSignificanceRange(PrevLODSignificance);
+				if (PrevRange != Range)
+				{
+					PrevRange->AddBatchedCustomDataFloats(CustomFloats, Range->StaticMeshRefs);
+				}
+			}
 		}
 	}
 
-	FORCEINLINE void WriteCustomDataFloatsAtStartIndex(int32 StaticMeshIndex, const TArrayView<float> CustomFloats, float LODSignificance, int32 FloatsPerInstance, int32 FloatStartIndex)
+	FORCEINLINE void WriteCustomDataFloatsAtStartIndex(int32 StaticMeshIndex, const TArrayView<float>& CustomFloats, const float LODSignificance, const int32 FloatsPerInstance, const int32 FloatStartIndex, const float PrevLODSignificance = -1.0f)
 	{
 		if (FMassLODSignificanceRange* Range = GetLODSignificanceRange(LODSignificance))
 		{
-			Range->WriteCustomDataFloatsAtStartIndex(StaticMeshIndex, CustomFloats, FloatsPerInstance, FloatStartIndex);
+			Range->WriteCustomDataFloatsAtStartIndex(StaticMeshIndex, CustomFloats, FloatsPerInstance, FloatStartIndex, TArray<uint32>());
+			if(PrevLODSignificance >= 0.0f)
+			{
+				FMassLODSignificanceRange* PrevRange = GetLODSignificanceRange(PrevLODSignificance);
+				if (PrevRange != Range)
+				{
+					PrevRange->WriteCustomDataFloatsAtStartIndex(StaticMeshIndex, CustomFloats, FloatsPerInstance, FloatStartIndex, Range->StaticMeshRefs);
+				}
+			}
 		}
 	}
 
diff --git a/Engine/Plugins/Runtime/MassGameplay/Source/MassRepresentation/Public/MassUpdateISMProcessor.h b/Engine/Plugins/Runtime/MassGameplay/Source/MassRepresentation/Public/MassUpdateISMProcessor.h
--- a/Engine/Plugins/Runtime/MassGameplay/Source/MassRepresentation/Public/MassUpdateISMProcessor.h
+++ b/Engine/Plugins/Runtime/MassGameplay/Source/MassRepresentation/Public/MassUpdateISMProcessor.h
@@ -18,7 +18,7 @@ class MASSREPRESENTATION_API UMassUpdateISMProcessor : public UMassProcessor
 public:
 	UMassUpdateISMProcessor();
 
-	static void UpdateISMTransform(int32 EntityId, FMassInstancedStaticMeshInfo& ISMInfo, const FTransform& Transform, const FTransform& PrevTransform, float LODSignificance);
+	static void UpdateISMTransform(int32 EntityId, FMassInstancedStaticMeshInfo& ISMInfo, const FTransform& Transform, const FTransform& PrevTransform, const float LODSignificance, const float PrevLODSignificance = -1.0f);
 
 protected:
 	/**
