[MassSmartObject] Renamed FDataFragment_SmartObjectUser to FMassSmartObjectUser

diff --git a/Engine/Plugins/AI/MassAI/Source/MassAIBehavior/Private/Evaluators/MassStateTreeSmartObjectEvaluator.cpp b/Engine/Plugins/AI/MassAI/Source/MassAIBehavior/Private/Evaluators/MassStateTreeSmartObjectEvaluator.cpp
--- a/Engine/Plugins/AI/MassAI/Source/MassAIBehavior/Private/Evaluators/MassStateTreeSmartObjectEvaluator.cpp
+++ b/Engine/Plugins/AI/MassAI/Source/MassAIBehavior/Private/Evaluators/MassStateTreeSmartObjectEvaluator.cpp
@@ -66,7 +66,7 @@ void FMassStateTreeSmartObjectEvaluator::Reset()
 
 void FMassStateTreeSmartObjectEvaluator::Evaluate(FStateTreeExecutionContext& Context, const EStateTreeEvaluationType EvalType,  const float DeltaTime)
 {
-	const FDataFragment_SmartObjectUser& SOUser = Context.GetExternalItem(SmartObjectUserHandle);
+	const FMassSmartObjectUserFragment& SOUser = Context.GetExternalItem(SmartObjectUserHandle);
 
 	bCandidatesFound = false;
 	bClaimed = SOUser.GetClaimHandle().IsValid();
diff --git a/Engine/Plugins/AI/MassAI/Source/MassAIBehavior/Private/Tasks/MassStateTreeSmartObjectTask.cpp b/Engine/Plugins/AI/MassAI/Source/MassAIBehavior/Private/Tasks/MassStateTreeSmartObjectTask.cpp
--- a/Engine/Plugins/AI/MassAI/Source/MassAIBehavior/Private/Tasks/MassStateTreeSmartObjectTask.cpp
+++ b/Engine/Plugins/AI/MassAI/Source/MassAIBehavior/Private/Tasks/MassStateTreeSmartObjectTask.cpp
@@ -46,7 +46,7 @@ EStateTreeRunStatus FMassClaimSmartObjectTask::EnterState(FStateTreeExecutionCon
 
 	// Retrieve fragments and subsystems
 	USmartObjectSubsystem& SmartObjectSubsystem = Context.GetExternalItem(SmartObjectSubsystemHandle);
-	FDataFragment_SmartObjectUser& SOUser = Context.GetExternalItem(SmartObjectUserHandle);
+	FMassSmartObjectUserFragment& SOUser = Context.GetExternalItem(SmartObjectUserHandle);
 
 	// Setup MassSmartObject handler and claim
 	const FMassStateTreeExecutionContext& MassContext = static_cast<FMassStateTreeExecutionContext&>(Context);
@@ -83,7 +83,7 @@ void FMassClaimSmartObjectTask::ExitState(FStateTreeExecutionContext& Context, c
 
 	ClaimResult = EMassSmartObjectClaimResult::Unset;
 
-	FDataFragment_SmartObjectUser& SOUser = Context.GetExternalItem(SmartObjectUserHandle);
+	FMassSmartObjectUserFragment& SOUser = Context.GetExternalItem(SmartObjectUserHandle);
 
 	if (SOUser.ClaimHandle.IsValid())
 	{
@@ -134,7 +134,7 @@ EStateTreeRunStatus FMassUseSmartObjectTask::EnterState(FStateTreeExecutionConte
 
 	// Retrieve fragments and subsystems
 	USmartObjectSubsystem& SmartObjectSubsystem = Context.GetExternalItem(SmartObjectSubsystemHandle);
-	FDataFragment_SmartObjectUser& SOUser = Context.GetExternalItem(SmartObjectUserHandle);
+	FMassSmartObjectUserFragment& SOUser = Context.GetExternalItem(SmartObjectUserHandle);
 	const FDataFragment_Transform& TransformFragment = Context.GetExternalItem(EntityTransformHandle);
 	FMassMoveTargetFragment& MoveTarget = Context.GetExternalItem(MoveTargetHandle);
 
@@ -161,7 +161,7 @@ void FMassUseSmartObjectTask::ExitState(FStateTreeExecutionContext& Context, con
 {
 	if (ChangeType == EStateTreeStateChangeType::Changed)
 	{
-		FDataFragment_SmartObjectUser& SOUser = Context.GetExternalItem(SmartObjectUserHandle);
+		FMassSmartObjectUserFragment& SOUser = Context.GetExternalItem(SmartObjectUserHandle);
 
 		if (SOUser.ClaimHandle.IsValid())
 		{
@@ -181,7 +181,7 @@ void FMassUseSmartObjectTask::ExitState(FStateTreeExecutionContext& Context, con
 
 void FMassUseSmartObjectTask::StateCompleted(FStateTreeExecutionContext& Context, const EStateTreeRunStatus CompletionStatus, const FStateTreeHandle CompletedState)
 {
-	FDataFragment_SmartObjectUser& SOUser = Context.GetExternalItem(SmartObjectUserHandle);
+	FMassSmartObjectUserFragment& SOUser = Context.GetExternalItem(SmartObjectUserHandle);
 
 	if (SOUser.ClaimHandle.IsValid())
 	{
@@ -199,7 +199,7 @@ EStateTreeRunStatus FMassUseSmartObjectTask::Tick(FStateTreeExecutionContext& Co
 {
 	EStateTreeRunStatus Status = EStateTreeRunStatus::Failed;
 
-	FDataFragment_SmartObjectUser& SOUser = Context.GetExternalItem(SmartObjectUserHandle);
+	FMassSmartObjectUserFragment& SOUser = Context.GetExternalItem(SmartObjectUserHandle);
 	switch (SOUser.InteractionStatus)
 	{
 	case EMassSmartObjectInteractionStatus::InProgress:
diff --git a/Engine/Plugins/AI/MassAI/Source/MassAIBehavior/Private/Tasks/MassStateTreeTestTask.cpp b/Engine/Plugins/AI/MassAI/Source/MassAIBehavior/Private/Tasks/MassStateTreeTestTask.cpp
--- a/Engine/Plugins/AI/MassAI/Source/MassAIBehavior/Private/Tasks/MassStateTreeTestTask.cpp
+++ b/Engine/Plugins/AI/MassAI/Source/MassAIBehavior/Private/Tasks/MassStateTreeTestTask.cpp
@@ -35,7 +35,7 @@ EStateTreeRunStatus FMassStateTreeTestTask::EnterState(FStateTreeExecutionContex
 EStateTreeRunStatus FMassStateTreeTestTask::Tick(FStateTreeExecutionContext& Context, const float DeltaTime)
 {
 	UMassStateTreeSubsystem& MassStateTreeSubSystem = Context.GetExternalItem(MassStateTreeSubSystemHandle);
-	const FDataFragment_SmartObjectUser* SmartObjectUser = Context.GetExternalItemPtr(SmartObjectUserHandle);
+	const FMassSmartObjectUserFragment* SmartObjectUser = Context.GetExternalItemPtr(SmartObjectUserHandle);
 	FDataFragment_Transform& Transform = Context.GetExternalItem(TransformHandle);
 
 	GEngine->AddOnScreenDebugMessage(INDEX_NONE, 1.0f, FColor::Orange, *FString::Printf(TEXT("[%s] Time=%f X=%f\n"), *Name.ToString(), Time, Transform.GetTransform().GetLocation().X));
diff --git a/Engine/Plugins/AI/MassAI/Source/MassAIBehavior/Private/Tasks/MassZoneGraphFindSmartObjectTarget.cpp b/Engine/Plugins/AI/MassAI/Source/MassAIBehavior/Private/Tasks/MassZoneGraphFindSmartObjectTarget.cpp
--- a/Engine/Plugins/AI/MassAI/Source/MassAIBehavior/Private/Tasks/MassZoneGraphFindSmartObjectTarget.cpp
+++ b/Engine/Plugins/AI/MassAI/Source/MassAIBehavior/Private/Tasks/MassZoneGraphFindSmartObjectTarget.cpp
@@ -21,7 +21,7 @@ bool FMassZoneGraphFindSmartObjectTarget::Link(FStateTreeLinker& Linker)
 EStateTreeRunStatus FMassZoneGraphFindSmartObjectTarget::EnterState(FStateTreeExecutionContext& Context, const EStateTreeStateChangeType ChangeType, const FStateTreeTransitionResult& Transition)
 {
 	const UZoneGraphAnnotationSubsystem& AnnotationSubsystem = Context.GetExternalItem(AnnotationSubsystemHandle);
-	const FDataFragment_SmartObjectUser& SOUser = Context.GetExternalItem(SmartObjectUserHandle);
+	const FMassSmartObjectUserFragment& SOUser = Context.GetExternalItem(SmartObjectUserHandle);
 	const FMassZoneGraphLaneLocationFragment& LaneLocation = Context.GetExternalItem(LocationHandle);
 	const FZoneGraphLaneHandle LaneHandle(LaneLocation.LaneHandle);
 
diff --git a/Engine/Plugins/AI/MassAI/Source/MassAIBehavior/Public/Evaluators/MassStateTreeSmartObjectEvaluator.h b/Engine/Plugins/AI/MassAI/Source/MassAIBehavior/Public/Evaluators/MassStateTreeSmartObjectEvaluator.h
--- a/Engine/Plugins/AI/MassAI/Source/MassAIBehavior/Public/Evaluators/MassStateTreeSmartObjectEvaluator.h
+++ b/Engine/Plugins/AI/MassAI/Source/MassAIBehavior/Public/Evaluators/MassStateTreeSmartObjectEvaluator.h
@@ -29,7 +29,7 @@ protected:
 	TStateTreeItemHandle<USmartObjectSubsystem> SmartObjectSubsystemHandle;
 	TStateTreeItemHandle<UMassSignalSubsystem> MassSignalSubsystemHandle;
 	TStateTreeItemHandle<FDataFragment_Transform> EntityTransformHandle;
-	TStateTreeItemHandle<FDataFragment_SmartObjectUser> SmartObjectUserHandle;
+	TStateTreeItemHandle<FMassSmartObjectUserFragment> SmartObjectUserHandle;
 	TStateTreeItemHandle<FMassZoneGraphLaneLocationFragment, EStateTreeItemRequirement::Optional> LocationHandle;
 
 	/** The identifier of the search request send by the evaluator to find candidates */
diff --git a/Engine/Plugins/AI/MassAI/Source/MassAIBehavior/Public/Tasks/MassStateTreeSmartObjectTask.h b/Engine/Plugins/AI/MassAI/Source/MassAIBehavior/Public/Tasks/MassStateTreeSmartObjectTask.h
--- a/Engine/Plugins/AI/MassAI/Source/MassAIBehavior/Public/Tasks/MassStateTreeSmartObjectTask.h
+++ b/Engine/Plugins/AI/MassAI/Source/MassAIBehavior/Public/Tasks/MassStateTreeSmartObjectTask.h
@@ -9,7 +9,7 @@
 #include "MassStateTreeSmartObjectTask.generated.h"
 
 struct FStateTreeExecutionContext;
-struct FDataFragment_SmartObjectUser;
+struct FMassSmartObjectUserFragment;
 class USmartObjectSubsystem;
 struct FDataFragment_Transform;
 struct FMassMoveTargetFragment;
@@ -28,7 +28,7 @@ protected:
 	virtual void ExitState(FStateTreeExecutionContext& Context, const EStateTreeStateChangeType ChangeType, const FStateTreeTransitionResult& Transition) override;
 	virtual EStateTreeRunStatus Tick(FStateTreeExecutionContext& Context, const float DeltaTime) override;
 
-	TStateTreeItemHandle<FDataFragment_SmartObjectUser> SmartObjectUserHandle;
+	TStateTreeItemHandle<FMassSmartObjectUserFragment> SmartObjectUserHandle;
 	TStateTreeItemHandle<USmartObjectSubsystem> SmartObjectSubsystemHandle;
 
 	/** Result of the candidates search request (Input) */
@@ -57,7 +57,7 @@ protected:
 
 	TStateTreeItemHandle<USmartObjectSubsystem> SmartObjectSubsystemHandle;
 	TStateTreeItemHandle<FDataFragment_Transform> EntityTransformHandle;
-	TStateTreeItemHandle<FDataFragment_SmartObjectUser> SmartObjectUserHandle;
+	TStateTreeItemHandle<FMassSmartObjectUserFragment> SmartObjectUserHandle;
 	TStateTreeItemHandle<FMassMoveTargetFragment> MoveTargetHandle;
 
 	/** Delay in seconds before trying to find & use another smart object */
diff --git a/Engine/Plugins/AI/MassAI/Source/MassAIBehavior/Public/Tasks/MassStateTreeTestTask.h b/Engine/Plugins/AI/MassAI/Source/MassAIBehavior/Public/Tasks/MassStateTreeTestTask.h
--- a/Engine/Plugins/AI/MassAI/Source/MassAIBehavior/Public/Tasks/MassStateTreeTestTask.h
+++ b/Engine/Plugins/AI/MassAI/Source/MassAIBehavior/Public/Tasks/MassStateTreeTestTask.h
@@ -7,7 +7,7 @@
 #include "MassStateTreeTestTask.generated.h"
 
 class UMassStateTreeSubsystem;
-struct FDataFragment_SmartObjectUser;
+struct FMassSmartObjectUserFragment;
 struct FDataFragment_Transform;
 
 /**
@@ -27,7 +27,7 @@ protected:
 	virtual EStateTreeRunStatus Tick(FStateTreeExecutionContext& Context, const float DeltaTime) override;
 
 	TStateTreeItemHandle<UMassStateTreeSubsystem> MassStateTreeSubSystemHandle;
-	TStateTreeItemHandle<FDataFragment_SmartObjectUser> SmartObjectUserHandle;
+	TStateTreeItemHandle<FMassSmartObjectUserFragment> SmartObjectUserHandle;
 	TStateTreeItemHandle<FDataFragment_Transform> TransformHandle;
 
 	UPROPERTY()
diff --git a/Engine/Plugins/AI/MassAI/Source/MassAIBehavior/Public/Tasks/MassZoneGraphFindSmartObjectTarget.h b/Engine/Plugins/AI/MassAI/Source/MassAIBehavior/Public/Tasks/MassZoneGraphFindSmartObjectTarget.h
--- a/Engine/Plugins/AI/MassAI/Source/MassAIBehavior/Public/Tasks/MassZoneGraphFindSmartObjectTarget.h
+++ b/Engine/Plugins/AI/MassAI/Source/MassAIBehavior/Public/Tasks/MassZoneGraphFindSmartObjectTarget.h
@@ -5,7 +5,7 @@
 #include "MassZoneGraphPathFollowTask.h"
 #include "MassZoneGraphFindSmartObjectTarget.generated.h"
 
-struct FDataFragment_SmartObjectUser;
+struct FMassSmartObjectUserFragment;
 struct FMassZoneGraphLaneLocationFragment;
 class UZoneGraphAnnotationSubsystem;
 
@@ -23,7 +23,7 @@ protected:
 	virtual void ExitState(FStateTreeExecutionContext& Context, const EStateTreeStateChangeType ChangeType, const FStateTreeTransitionResult& Transition) override;
 	virtual EStateTreeRunStatus Tick(FStateTreeExecutionContext& Context, const float DeltaTime) override;
 
-	TStateTreeItemHandle<FDataFragment_SmartObjectUser> SmartObjectUserHandle;
+	TStateTreeItemHandle<FMassSmartObjectUserFragment> SmartObjectUserHandle;
 	TStateTreeItemHandle<FMassZoneGraphLaneLocationFragment> LocationHandle;
 	TStateTreeItemHandle<UZoneGraphAnnotationSubsystem> AnnotationSubsystemHandle;
 
diff --git a/Engine/Plugins/Runtime/MassGameplay/Source/MassSmartObjects/Private/MassSmartObjectHandler.cpp b/Engine/Plugins/Runtime/MassGameplay/Source/MassSmartObjects/Private/MassSmartObjectHandler.cpp
--- a/Engine/Plugins/Runtime/MassGameplay/Source/MassSmartObjects/Private/MassSmartObjectHandler.cpp
+++ b/Engine/Plugins/Runtime/MassGameplay/Source/MassSmartObjects/Private/MassSmartObjectHandler.cpp
@@ -18,7 +18,7 @@ struct FPayload
 
 void OnSlotInvalidated(const FSmartObjectClaimHandle& ClaimHandle, const ESmartObjectSlotState State, FPayload Payload)
 {
-	FDataFragment_SmartObjectUser& User = Payload.EntitySubsystem->GetFragmentDataChecked<FDataFragment_SmartObjectUser>(Payload.Entity);
+	FMassSmartObjectUserFragment& User = Payload.EntitySubsystem->GetFragmentDataChecked<FMassSmartObjectUserFragment>(Payload.Entity);
 	if (USmartObjectSubsystem* SmartObjectSubsystem = Payload.SmartObjectSubsystem.Get())
 	{
 		SmartObjectSubsystem->Release(User.ClaimHandle);
@@ -96,7 +96,7 @@ void FMassSmartObjectHandler::RemoveRequest(const FMassSmartObjectRequestID& Req
 	ExecutionContext.Defer().DestroyEntity(RequestEntity);
 }
 
-EMassSmartObjectClaimResult FMassSmartObjectHandler::ClaimCandidate(const FMassEntityHandle Entity, FDataFragment_SmartObjectUser& User, const FMassSmartObjectRequestID& RequestID) const
+EMassSmartObjectClaimResult FMassSmartObjectHandler::ClaimCandidate(const FMassEntityHandle Entity, FMassSmartObjectUserFragment& User, const FMassSmartObjectRequestID& RequestID) const
 {
 	if (!ensureMsgf(RequestID.IsSet(), TEXT("Trying to claim using an invalid request.")))
 	{
@@ -113,7 +113,7 @@ EMassSmartObjectClaimResult FMassSmartObjectHandler::ClaimCandidate(const FMassE
 	return ClaimCandidate(Entity, User, RequestFragment.Result);
 }
 
-EMassSmartObjectClaimResult FMassSmartObjectHandler::ClaimCandidate(const FMassEntityHandle Entity, FDataFragment_SmartObjectUser& User, const FMassSmartObjectRequestResult& SearchRequestResult) const
+EMassSmartObjectClaimResult FMassSmartObjectHandler::ClaimCandidate(const FMassEntityHandle Entity, FMassSmartObjectUserFragment& User, const FMassSmartObjectRequestResult& SearchRequestResult) const
 {
 	if (!ensureMsgf(SearchRequestResult.bProcessed, TEXT("Trying to claim using a search request that is not processed.")))
 	{
@@ -143,7 +143,7 @@ EMassSmartObjectClaimResult FMassSmartObjectHandler::ClaimCandidate(const FMassE
 	return User.ClaimHandle.IsValid() ? EMassSmartObjectClaimResult::Succeeded : EMassSmartObjectClaimResult::Failed_NoAvailableCandidate;
 }
 
-bool FMassSmartObjectHandler::ClaimSmartObject(const FMassEntityHandle Entity, FDataFragment_SmartObjectUser& User, const FSmartObjectID& ObjectID) const
+bool FMassSmartObjectHandler::ClaimSmartObject(const FMassEntityHandle Entity, FMassSmartObjectUserFragment& User, const FSmartObjectID& ObjectID) const
 {
 	FSmartObjectRequestFilter Filter;
 	Filter.BehaviorConfigurationClass = USmartObjectMassBehaviorConfig::StaticClass();
@@ -180,7 +180,7 @@ bool FMassSmartObjectHandler::ClaimSmartObject(const FMassEntityHandle Entity, F
 
 bool FMassSmartObjectHandler::UseSmartObject(
 	const FMassEntityHandle Entity,
-	FDataFragment_SmartObjectUser& User,
+	FMassSmartObjectUserFragment& User,
 	const FDataFragment_Transform& Transform) const
 {
 #if WITH_MASSGAMEPLAY_DEBUG
@@ -207,7 +207,7 @@ bool FMassSmartObjectHandler::UseSmartObject(
 	return true;
 }
 
-void FMassSmartObjectHandler::ReleaseSmartObject(const FMassEntityHandle Entity, FDataFragment_SmartObjectUser& User, const EMassSmartObjectInteractionStatus NewStatus) const
+void FMassSmartObjectHandler::ReleaseSmartObject(const FMassEntityHandle Entity, FMassSmartObjectUserFragment& User, const EMassSmartObjectInteractionStatus NewStatus) const
 {
 #if WITH_MASSGAMEPLAY_DEBUG
 	UE_CVLOG(UE::Mass::Debug::IsDebuggingEntity(Entity),
diff --git a/Engine/Plugins/Runtime/MassGameplay/Source/MassSmartObjects/Private/MassSmartObjectProcessor.cpp b/Engine/Plugins/Runtime/MassGameplay/Source/MassSmartObjects/Private/MassSmartObjectProcessor.cpp
--- a/Engine/Plugins/Runtime/MassGameplay/Source/MassSmartObjects/Private/MassSmartObjectProcessor.cpp
+++ b/Engine/Plugins/Runtime/MassGameplay/Source/MassSmartObjects/Private/MassSmartObjectProcessor.cpp
@@ -268,7 +268,7 @@ void UMassProcessor_SmartObjectCandidatesFinder::Execute(UMassEntitySubsystem& E
 //----------------------------------------------------------------------//
 void UMassProcessor_SmartObjectTimedBehavior::ConfigureQueries()
 {
-	EntityQuery.AddRequirement<FDataFragment_SmartObjectUser>(EMassFragmentAccess::ReadWrite);
+	EntityQuery.AddRequirement<FMassSmartObjectUserFragment>(EMassFragmentAccess::ReadWrite);
 	EntityQuery.AddTagRequirement<FMassSmartObjectTimedBehaviorTag>(EMassFragmentPresence::All);
 }
 
@@ -283,18 +283,18 @@ void UMassProcessor_SmartObjectTimedBehavior::Execute(UMassEntitySubsystem& Enti
 	checkf(SmartObjectSubsystem != nullptr, TEXT("SmartObjectSubsystem should exist when executing processors."));
 	checkf(SignalSubsystem != nullptr, TEXT("MassSignalSubsystem should exist when executing processors."));
 
-	TArray<TTuple<FMassEntityHandle, FDataFragment_SmartObjectUser*>> ToRelease;
+	TArray<TTuple<FMassEntityHandle, FMassSmartObjectUserFragment*>> ToRelease;
 
 	QUICK_SCOPE_CYCLE_COUNTER(UMassProcessor_SmartObjectTestBehavior_Run);
 
 	EntityQuery.ForEachEntityChunk(EntitySubsystem, Context, [this, &ToRelease](FMassExecutionContext& Context)
 	{
 		const int32 NumEntities = Context.GetNumEntities();
-		const TArrayView<FDataFragment_SmartObjectUser> UserList = Context.GetMutableFragmentView<FDataFragment_SmartObjectUser>();
+		const TArrayView<FMassSmartObjectUserFragment> UserList = Context.GetMutableFragmentView<FMassSmartObjectUserFragment>();
 
 		for (int32 i = 0; i < NumEntities; ++i)
 		{
-			FDataFragment_SmartObjectUser& SOUser = UserList[i];
+			FMassSmartObjectUserFragment& SOUser = UserList[i];
 			if (SOUser.InteractionStatus != EMassSmartObjectInteractionStatus::InProgress)
 			{
 				// Only expecting in progress interaction status
@@ -334,12 +334,12 @@ void UMassProcessor_SmartObjectTimedBehavior::Execute(UMassEntitySubsystem& Enti
 		}
 	});
 
-	for (const TTuple<FMassEntityHandle, FDataFragment_SmartObjectUser*> EntryToRelease : ToRelease)
+	for (const TTuple<FMassEntityHandle, FMassSmartObjectUserFragment*> EntryToRelease : ToRelease)
 	{
 		const FMassEntityHandle Entity = EntryToRelease.Get<0>();
 		SignalSubsystem->SignalEntity(UE::Mass::Signals::SmartObjectInteractionDone, Entity);
 
-		FDataFragment_SmartObjectUser* User = EntryToRelease.Get<1>();
+		FMassSmartObjectUserFragment* User = EntryToRelease.Get<1>();
 		SmartObjectSubsystem->Release(User->ClaimHandle);
 		User->InteractionStatus = EMassSmartObjectInteractionStatus::Completed;
 		Context.Defer().RemoveTag<FMassSmartObjectTimedBehaviorTag>(Entity);
diff --git a/Engine/Plugins/Runtime/MassGameplay/Source/MassSmartObjects/Private/MassSmartObjectsModule.cpp b/Engine/Plugins/Runtime/MassGameplay/Source/MassSmartObjects/Private/MassSmartObjectsModule.cpp
--- a/Engine/Plugins/Runtime/MassGameplay/Source/MassSmartObjects/Private/MassSmartObjectsModule.cpp
+++ b/Engine/Plugins/Runtime/MassGameplay/Source/MassSmartObjects/Private/MassSmartObjectsModule.cpp
@@ -2,6 +2,7 @@
 
 #include "CoreMinimal.h"
 #include "IMassSmartObjectsModule.h"
+#include "UObject/CoreRedirects.h"
 
 
 class FMassSmartObjectsModule : public IMassSmartObjectsModule
@@ -18,6 +19,11 @@ IMPLEMENT_MODULE(FMassSmartObjectsModule, MassSmartObjects)
 void FMassSmartObjectsModule::StartupModule()
 {
 	// This code will execute after your module is loaded into memory (but after global variables are initialized, of course.)
+
+	TArray<FCoreRedirect> Redirects;
+	Redirects.Emplace(ECoreRedirectFlags::Type_Struct, TEXT("DataFragment_SmartObjectUser"), TEXT("MassSmartObjectUserFragment"));
+
+	FCoreRedirects::AddRedirectList(Redirects, TEXT("MassSmartObject"));
 }
 
 
diff --git a/Engine/Plugins/Runtime/MassGameplay/Source/MassSmartObjects/Public/MassSmartObjectBehaviorConfig.h b/Engine/Plugins/Runtime/MassGameplay/Source/MassSmartObjects/Public/MassSmartObjectBehaviorConfig.h
--- a/Engine/Plugins/Runtime/MassGameplay/Source/MassSmartObjects/Public/MassSmartObjectBehaviorConfig.h
+++ b/Engine/Plugins/Runtime/MassGameplay/Source/MassSmartObjects/Public/MassSmartObjectBehaviorConfig.h
@@ -10,7 +10,7 @@ class UMassEntitySubsystem;
 class USmartObjectSubsystem;
 struct FMassExecutionContext;
 struct FDataFragment_Transform;
-struct FDataFragment_SmartObjectUser;
+struct FMassSmartObjectUserFragment;
 
 /**
  * Struct to pass around the required set of information to activate a mass behavior configuration on a given entity.
@@ -19,13 +19,13 @@ struct MASSSMARTOBJECTS_API FMassBehaviorEntityContext
 {
 	FMassBehaviorEntityContext() = delete;
 
-	FMassBehaviorEntityContext(const FMassEntityHandle InEntity, const FDataFragment_Transform& InTransformFragment, FDataFragment_SmartObjectUser& InSOUser, USmartObjectSubsystem& InSubsystem)
+	FMassBehaviorEntityContext(const FMassEntityHandle InEntity, const FDataFragment_Transform& InTransformFragment, FMassSmartObjectUserFragment& InSOUser, USmartObjectSubsystem& InSubsystem)
 		: Entity(InEntity), TransformFragment(InTransformFragment), SOUser(InSOUser), Subsystem(InSubsystem)
 	{}
 
 	const FMassEntityHandle Entity;
 	const FDataFragment_Transform& TransformFragment;
-	FDataFragment_SmartObjectUser& SOUser;
+	FMassSmartObjectUserFragment& SOUser;
 	USmartObjectSubsystem& Subsystem;
 };
 
diff --git a/Engine/Plugins/Runtime/MassGameplay/Source/MassSmartObjects/Public/MassSmartObjectHandler.h b/Engine/Plugins/Runtime/MassGameplay/Source/MassSmartObjects/Public/MassSmartObjectHandler.h
--- a/Engine/Plugins/Runtime/MassGameplay/Source/MassSmartObjects/Public/MassSmartObjectHandler.h
+++ b/Engine/Plugins/Runtime/MassGameplay/Source/MassSmartObjects/Public/MassSmartObjectHandler.h
@@ -9,7 +9,7 @@ class UMassEntitySubsystem;
 class USmartObjectSubsystem;
 struct FMassExecutionContext;
 struct FMassEntityHandle;
-struct FDataFragment_SmartObjectUser;
+struct FMassSmartObjectUserFragment;
 struct FDataFragment_Transform;
 struct FSmartObjectClaimHandle;
 struct FSmartObjectID;
@@ -79,7 +79,7 @@ struct MASSSMARTOBJECTS_API FMassSmartObjectHandler
 	 * @param RequestID A valid request identifier (method will ensure otherwise)
 	 * @return Whether the slot has been successfully claimed or not
 	 */
-	UE_NODISCARD EMassSmartObjectClaimResult ClaimCandidate(const FMassEntityHandle Entity, FDataFragment_SmartObjectUser& User, const FMassSmartObjectRequestID& RequestID) const;
+	UE_NODISCARD EMassSmartObjectClaimResult ClaimCandidate(const FMassEntityHandle Entity, FMassSmartObjectUserFragment& User, const FMassSmartObjectRequestID& RequestID) const;
 
 	/**
 	 * Claims the first available smart object from the provided results.
@@ -88,7 +88,7 @@ struct MASSSMARTOBJECTS_API FMassSmartObjectHandler
 	 * @param SearchRequestResult Results of completed search request
 	 * @return Whether the slot has been successfully claimed or not
 	 */
-	UE_NODISCARD EMassSmartObjectClaimResult ClaimCandidate(const FMassEntityHandle Entity, FDataFragment_SmartObjectUser& User, const FMassSmartObjectRequestResult& SearchRequestResult) const;
+	UE_NODISCARD EMassSmartObjectClaimResult ClaimCandidate(const FMassEntityHandle Entity, FMassSmartObjectUserFragment& User, const FMassSmartObjectRequestResult& SearchRequestResult) const;
 
 	/**
 	 * Claims the first available slot holding any type of USmartObjectMassBehaviorConfig in the smart object
@@ -98,7 +98,7 @@ struct MASSSMARTOBJECTS_API FMassSmartObjectHandler
 	 * @param ObjectID A valid smart object identifier (method will ensure otherwise)
 	 * @return Whether the slot has been successfully claimed or not
 	 */
-	bool ClaimSmartObject(const FMassEntityHandle Entity, FDataFragment_SmartObjectUser& User, const FSmartObjectID& ObjectID) const;
+	bool ClaimSmartObject(const FMassEntityHandle Entity, FMassSmartObjectUserFragment& User, const FSmartObjectID& ObjectID) const;
 
 	/**
 	 * Activates the mass gameplay behavior associated to the previously claimed smart object.
@@ -107,7 +107,7 @@ struct MASSSMARTOBJECTS_API FMassSmartObjectHandler
 	 * @param Transform Fragment holding the transform of the user claiming
 	 * @return Whether the slot has been successfully claimed or not
 	 */
-	bool UseSmartObject(const FMassEntityHandle Entity, FDataFragment_SmartObjectUser& User, const FDataFragment_Transform& Transform) const;
+	bool UseSmartObject(const FMassEntityHandle Entity, FMassSmartObjectUserFragment& User, const FDataFragment_Transform& Transform) const;
 
 	/**
 	 * Releases a claimed/in-use smart object and update user fragment.
@@ -115,7 +115,7 @@ struct MASSSMARTOBJECTS_API FMassSmartObjectHandler
 	 * @param User Fragment of the user claiming
 	 * @param Status The new status for in-progress interaction
 	 */
-	void ReleaseSmartObject(const FMassEntityHandle Entity, FDataFragment_SmartObjectUser& User, const EMassSmartObjectInteractionStatus Status = EMassSmartObjectInteractionStatus::Unset) const;
+	void ReleaseSmartObject(const FMassEntityHandle Entity, FMassSmartObjectUserFragment& User, const EMassSmartObjectInteractionStatus Status = EMassSmartObjectInteractionStatus::Unset) const;
 
 private:
 	UMassEntitySubsystem& EntitySubsystem;
diff --git a/Engine/Plugins/Runtime/MassGameplay/Source/MassSmartObjects/Public/MassSmartObjectProcessor.h b/Engine/Plugins/Runtime/MassGameplay/Source/MassSmartObjects/Public/MassSmartObjectProcessor.h
--- a/Engine/Plugins/Runtime/MassGameplay/Source/MassSmartObjects/Public/MassSmartObjectProcessor.h
+++ b/Engine/Plugins/Runtime/MassGameplay/Source/MassSmartObjects/Public/MassSmartObjectProcessor.h
@@ -15,7 +15,7 @@ class UZoneGraphSubsystem;
 class UZoneGraphAnnotationSubsystem;
 
 USTRUCT()
-struct MASSSMARTOBJECTS_API FDataFragment_SmartObjectUser : public FMassFragment
+struct MASSSMARTOBJECTS_API FMassSmartObjectUserFragment : public FMassFragment
 {
 	GENERATED_BODY()
 
