Added option to take most of the initialization hit in the first frame instead of spreading throughout the first period.

diff --git a/Engine/Plugins/AI/MassAI/Source/MassAIBehavior/Private/MassStateTreeProcessors.cpp b/Engine/Plugins/AI/MassAI/Source/MassAIBehavior/Private/MassStateTreeProcessors.cpp
--- a/Engine/Plugins/AI/MassAI/Source/MassAIBehavior/Private/MassStateTreeProcessors.cpp
+++ b/Engine/Plugins/AI/MassAI/Source/MassAIBehavior/Private/MassStateTreeProcessors.cpp
@@ -216,10 +216,10 @@ UMassStateTreeSignalRequestInitialization::UMassStateTreeSignalRequestInitializa
 {
 	ExecutionOrder.ExecuteAfter.Add(UE::Mass::ProcessorGroupNames::LOD);
 	ExecutionOrder.ExecuteBefore.Add(UE::Mass::ProcessorGroupNames::Behavior);
-	MaxInitializationRequestsPerLOD[EMassLOD::High] = 250;
-	MaxInitializationRequestsPerLOD[EMassLOD::Medium] = 250;
-	MaxInitializationRequestsPerLOD[EMassLOD::Low] = 250;
-	MaxInitializationRequestsPerLOD[EMassLOD::Off] = 250;
+	MaxInitializationRequestsPerLOD[EMassLOD::High] = 100;
+	MaxInitializationRequestsPerLOD[EMassLOD::Medium] = 100;
+	MaxInitializationRequestsPerLOD[EMassLOD::Low] = 100;
+	MaxInitializationRequestsPerLOD[EMassLOD::Off] = 100;
 }
 void UMassStateTreeSignalRequestInitialization::Initialize(UObject& Owner)
 {
diff --git a/Engine/Plugins/Runtime/MassGameplay/Source/MassLOD/Private/MassSimulationLOD.cpp b/Engine/Plugins/Runtime/MassGameplay/Source/MassLOD/Private/MassSimulationLOD.cpp
--- a/Engine/Plugins/Runtime/MassGameplay/Source/MassLOD/Private/MassSimulationLOD.cpp
+++ b/Engine/Plugins/Runtime/MassGameplay/Source/MassLOD/Private/MassSimulationLOD.cpp
@@ -108,7 +108,7 @@ void UMassSimulationLODProcessor::Initialize(UObject& InOwner)
 	for (FMassSimulationLODConfig& LODConfig : LODConfigs)
 	{
 		LODConfig.LODCalculator.Initialize(LODConfig.LODDistance, LODConfig.BufferHysteresisOnDistancePercentage / 100.0f, LODConfig.LODMaxCount);
-		LODConfig.LODTickRateController.Initialize(LODConfig.TickRates);
+		LODConfig.LODTickRateController.Initialize(LODConfig.TickRates, LODConfig.bSpreadFirstSimulationUpdate);
 	}
 
 	Super::Initialize(InOwner);
diff --git a/Engine/Plugins/Runtime/MassGameplay/Source/MassLOD/Public/MassLODTickRateController.h b/Engine/Plugins/Runtime/MassGameplay/Source/MassLOD/Public/MassLODTickRateController.h
--- a/Engine/Plugins/Runtime/MassGameplay/Source/MassLOD/Public/MassLODTickRateController.h
+++ b/Engine/Plugins/Runtime/MassGameplay/Source/MassLOD/Public/MassLODTickRateController.h
@@ -17,9 +17,10 @@ public:
 
 	/**
 	 * Initializes the LOD trick rate controller, needed to be called once at initialization time (Only when FLODLogic::bDoVariableTickRate is enabled)
-	 * @Param InTickRate the rate at which entities should be ticked per LOD
+	 * @Param InTickRate the rate at which entities should be ticked per 
+	 * @Param bInShouldSpreadFirstUpdate over the period specified in InTickRate parameter
 	 */
-	void Initialize(const float InTickRate[EMassLOD::Max]);
+	void Initialize(const float InTickRate[EMassLOD::Max], const bool bInShouldSpreadFirstUpdate = false);
 
 	/**
 	 * Retrieve if it is needed to calculate the LOD for this chunk
@@ -49,10 +50,13 @@ protected:
 
 	/** Tick rate for each LOD */
 	TStaticArray<int32, EMassLOD::Max> TickRates;
+
+	/* Whether or not to spread the first update over the period specified in tick rate member for its LOD */
+	bool bShouldSpreadFirstUpdate = false;
 };
 
 template <typename FVariableTickChunkFragment, typename FLODLogic>
-void TMassLODTickRateController<FVariableTickChunkFragment,FLODLogic>::Initialize(const float InTickRates[EMassLOD::Max])
+void TMassLODTickRateController<FVariableTickChunkFragment, FLODLogic>::Initialize(const float InTickRates[EMassLOD::Max], const bool bInShouldSpreadFirstUpdate/* = false*/)
 {
 	checkf(InTickRates, TEXT("You need to provide tick rate values to use this class."));
 	checkf(FLODLogic::bDoVariableTickRate, TEXT("You need to enalbe bDoVariableTickRate to use this class."));
@@ -62,6 +66,8 @@ void TMassLODTickRateController<FVariableTickChunkFragment,FLODLogic>::Initializ
 	{
 		TickRates[x] = InTickRates[x];
 	}
+
+	bShouldSpreadFirstUpdate = bInShouldSpreadFirstUpdate;
 }
 
 template <typename FVariableTickChunkFragment, typename FLODLogic>
@@ -97,7 +103,7 @@ bool TMassLODTickRateController<FVariableTickChunkFragment, FLODLogic>::UpdateTi
 		// The LOD on the chunk fragment data isn't set yet, let see if the Archetype has an LOD tag and set it on the ChunkData
 		ChunkLOD = UE::MassLOD::GetLODFromArchetype(Context);
 		ChunkData.SetLOD(ChunkLOD);
-		bFirstUpdate = true;
+		bFirstUpdate = bShouldSpreadFirstUpdate;
 	}
 	else
 	{
diff --git a/Engine/Plugins/Runtime/MassGameplay/Source/MassLOD/Public/MassSimulationLOD.h b/Engine/Plugins/Runtime/MassGameplay/Source/MassLOD/Public/MassSimulationLOD.h
--- a/Engine/Plugins/Runtime/MassGameplay/Source/MassLOD/Public/MassSimulationLOD.h
+++ b/Engine/Plugins/Runtime/MassGameplay/Source/MassLOD/Public/MassSimulationLOD.h
@@ -118,6 +118,10 @@ struct FMassSimulationLODConfig
 	UPROPERTY(EditAnywhere, Category = "Mass|LOD", config)
 	float TickRates[EMassLOD::Max];
 
+	/** If true, will spread the first simulation update over TickRate period */
+	UPROPERTY(EditAnywhere, Category = "Mass|Representation", config)
+	bool bSpreadFirstSimulationUpdate = false;
+
 	/** Runtime data for matching the LOD config */
 	FMassEntityQuery EntityQuery;
 	TMassLODCalculator<FMassSimulationLODLogic> LODCalculator;
diff --git a/Engine/Plugins/Runtime/MassGameplay/Source/MassRepresentation/Private/MassRepresentationProcessor.cpp b/Engine/Plugins/Runtime/MassGameplay/Source/MassRepresentation/Private/MassRepresentationProcessor.cpp
--- a/Engine/Plugins/Runtime/MassGameplay/Source/MassRepresentation/Private/MassRepresentationProcessor.cpp
+++ b/Engine/Plugins/Runtime/MassGameplay/Source/MassRepresentation/Private/MassRepresentationProcessor.cpp
@@ -380,7 +380,7 @@ FMassVisualizationChunkFragment& UMassRepresentationProcessor::UpdateChunkVisibi
 		// The visibility on the chunk fragment data isn't set yet, let see if the Archetype has an visibility tag and set it on the ChunkData
 		ChunkVisibility = UE::MassRepresentation::GetVisibilityFromArchetype(Context);
 		ChunkData.SetVisibility(ChunkVisibility);
-		bFirstUpdate = true;
+		bFirstUpdate = bSpreadFirstVisualizationUpdate;
 	}
 	else
 	{
diff --git a/Engine/Plugins/Runtime/MassGameplay/Source/MassRepresentation/Public/MassRepresentationProcessor.h b/Engine/Plugins/Runtime/MassGameplay/Source/MassRepresentation/Public/MassRepresentationProcessor.h
--- a/Engine/Plugins/Runtime/MassGameplay/Source/MassRepresentation/Public/MassRepresentationProcessor.h
+++ b/Engine/Plugins/Runtime/MassGameplay/Source/MassRepresentation/Public/MassRepresentationProcessor.h
@@ -178,6 +178,10 @@ protected:
 	UPROPERTY(EditAnywhere, Category = "Mass|Representation", config)
 	bool bKeepActorExtraFrame = false;
 
+	/** If true, will spread the first visualization update over the period specified in NotVisibleUpdateRate member */
+	UPROPERTY(EditAnywhere, Category = "Mass|Representation", config)
+	bool bSpreadFirstVisualizationUpdate = false;
+
 	/** World Partition grid name to test collision against, default None will be the main grid */
 	UPROPERTY(EditAnywhere, Category = "Mass|Representation", config)
 	FName WorldPartitionGridNameContainingCollision;
