Saving the last calculated count per LOD, great for debugging.

diff --git a/Engine/Plugins/Runtime/MassGameplay/Source/MassLOD/Public/MassLODCalculator.h b/Engine/Plugins/Runtime/MassGameplay/Source/MassLOD/Public/MassLODCalculator.h
--- a/Engine/Plugins/Runtime/MassGameplay/Source/MassLOD/Public/MassLODCalculator.h
+++ b/Engine/Plugins/Runtime/MassGameplay/Source/MassLOD/Public/MassLODCalculator.h
@@ -112,6 +112,11 @@ protected:
 		/** Count of entities in each subdivision */
 		TStaticArray< TStaticArray<int32, UE::MassLOD::MaxBucketsPerLOD>, EMassLOD::Max > BaseBucketCounts;
 		TStaticArray< TStaticArray<int32, UE::MassLOD::MaxBucketsPerLOD>, EMassLOD::Max > VisibleBucketCounts;
+
+#if WITH_MASSGAMEPLAY_DEBUG
+		/* Last calculation count per LOD */
+		TStaticArray<int32, EMassLOD::Max> LastCalculatedLODCount;
+#endif // WITH_MASSGAMEPLAY_DEBUG
 	};
 
 
@@ -347,6 +352,11 @@ bool TMassLODCalculator<FLODLogic>::AdjustDistancesFromCountForRuntimeData(const
 		// Switch to next LOD if we have not reach the max
 		if (ProcessingLODIdx < BucketLODIdx)
 		{
+#if WITH_MASSGAMEPLAY_DEBUG
+			// Save the count of this LOD for this frame
+			Data.LastCalculatedLODCount[ProcessingLODIdx] = Count;
+#endif // WITH_MASSGAMEPLAY_DEBUG
+
 			// Switch to next LOD
 			ProcessingLODIdx = BucketLODIdx;
 
@@ -364,6 +374,11 @@ bool TMassLODCalculator<FLODLogic>::AdjustDistancesFromCountForRuntimeData(const
 
 				while (Count > MaxCount[ProcessingLODIdx])
 				{
+#if WITH_MASSGAMEPLAY_DEBUG
+					// Save the count of this LOD for this frame
+					Data.LastCalculatedLODCount[ProcessingLODIdx] = Count - Data.VisibleBucketCounts[BucketLODIdx][BucketIdx];
+#endif // WITH_MASSGAMEPLAY_DEBUG
+
 					// Switch to next LOD
 					ProcessingLODIdx++;
 
@@ -392,6 +407,11 @@ bool TMassLODCalculator<FLODLogic>::AdjustDistancesFromCountForRuntimeData(const
 			// Check if the count is going over max
 			while (Count > MaxCount[ProcessingLODIdx])
 			{
+#if WITH_MASSGAMEPLAY_DEBUG
+				// Save the count of this LOD for this frame
+				Data.LastCalculatedLODCount[ProcessingLODIdx] = Count - Data.BaseBucketCounts[BucketLODIdx][BucketIdx];
+#endif // WITH_MASSGAMEPLAY_DEBUG
+
 				// Switch to next LOD
 				ProcessingLODIdx++;
 
@@ -418,6 +438,14 @@ bool TMassLODCalculator<FLODLogic>::AdjustDistancesFromCountForRuntimeData(const
 		}
 	}
 
+#if WITH_MASSGAMEPLAY_DEBUG
+	if (ProcessingLODIdx < EMassLOD::Max - 1)
+	{
+		// Save the count of this LOD for this frame
+		Data.LastCalculatedLODCount[ProcessingLODIdx] = Count;
+	}
+#endif // WITH_MASSGAMEPLAY_DEBUG
+
 	return bNeedAdjustments;
 }
 
