[MassSmartObject] process requests even if annotations are missing
Prevent adding too many candidates to the results

diff --git a/Engine/Plugins/Runtime/MassGameplay/Source/MassSmartObjects/Private/MassSmartObjectProcessor.cpp b/Engine/Plugins/Runtime/MassGameplay/Source/MassSmartObjects/Private/MassSmartObjectProcessor.cpp
--- a/Engine/Plugins/Runtime/MassGameplay/Source/MassSmartObjects/Private/MassSmartObjectProcessor.cpp
+++ b/Engine/Plugins/Runtime/MassGameplay/Source/MassSmartObjects/Private/MassSmartObjectProcessor.cpp
@@ -149,10 +149,10 @@ void UMassProcessor_SmartObjectCandidatesFinder::Execute(UMassEntitySubsystem& E
 
 	// Process lane based requests
 	const FZoneGraphTag SmartObjectTag = GetDefault<UMassSmartObjectSettings>()->SmartObjectTag;
-	USmartObjectZoneAnnotations* SmartObjectBehavior = Cast<USmartObjectZoneAnnotations>(AnnotationSubsystem->GetFirstAnnotationForTag(SmartObjectTag));
-	if (SmartObjectBehavior != nullptr)
-	{
-		LaneRequestQuery.ForEachEntityChunk(EntitySubsystem, Context, [this, SmartObjectBehavior, &Filter, SmartObjectTag, &EntitiesToSignal, &BeginRequestProcessing, &EndRequestProcessing](FMassExecutionContext& Context)
+	USmartObjectZoneAnnotations* Annotations = Cast<USmartObjectZoneAnnotations>(AnnotationSubsystem->GetFirstAnnotationForTag(SmartObjectTag));
+
+	LaneRequestQuery.ForEachEntityChunk(EntitySubsystem, Context,
+		[this, Annotations, &Filter, SmartObjectTag, &EntitiesToSignal, &BeginRequestProcessing, &EndRequestProcessing](FMassExecutionContext& Context)
 		{
 			const int32 NumEntities = Context.GetNumEntities();
 			EntitiesToSignal.Reserve(EntitiesToSignal.Num() + NumEntities);
@@ -188,11 +188,17 @@ void UMassProcessor_SmartObjectCandidatesFinder::Execute(UMassEntitySubsystem& E
 					continue;
 				}
 
+				if (Annotations == nullptr)
+				{
+					UE_VLOG(SmartObjectSubsystem, LogSmartObject, Warning, TEXT("%d lane location based requests failed since SmartObject annotations are not available"), NumEntities);
+					return;
+				}
+
 				// Fetch smart object data associated to the current graph if different than last used one
 				if (LastUsedDataHandle != RequestLaneHandle.DataHandle)
 				{
 					LastUsedDataHandle = RequestLaneHandle.DataHandle;
-					GraphData = SmartObjectBehavior->GetAnnotationData(RequestLaneHandle.DataHandle);
+					GraphData = Annotations->GetAnnotationData(RequestLaneHandle.DataHandle);
 				}
 
 				if (GraphData == nullptr)
@@ -251,10 +257,14 @@ void UMassProcessor_SmartObjectCandidatesFinder::Execute(UMassEntitySubsystem& E
 						UE_VLOG_SEGMENT(SmartObjectSubsystem, LogSmartObject, Display, RequestLaneLocation.Position, EntryPointLaneLocation.Position, DebugColor, TEXT(""));
 					}
 #endif // WITH_MASSGAMEPLAY_DEBUG
+
+					if (Result.NumCandidates == FMassSmartObjectRequestResult::MaxNumCandidates)
+					{
+						break;
+					}
 				}
 			}
 		});
-	}
 
 	// Signal entities that their search results are ready
 	if (EntitiesToSignal.Num())
