[MassSmartObject] add SmartObjectUser trait

diff --git a/Engine/Plugins/AI/MassCrowd/Source/MassCrowd/Private/MassCrowdMemberTrait.cpp b/Engine/Plugins/AI/MassCrowd/Source/MassCrowd/Private/MassCrowdMemberTrait.cpp
--- a/Engine/Plugins/AI/MassCrowd/Source/MassCrowd/Private/MassCrowdMemberTrait.cpp
+++ b/Engine/Plugins/AI/MassCrowd/Source/MassCrowd/Private/MassCrowdMemberTrait.cpp
@@ -3,7 +3,6 @@
 #include "MassCrowdMemberTrait.h"
 #include "MassCrowdFragments.h"
 #include "MassEntityTemplateRegistry.h"
-#include "MassCrowdNavigationProcessor.h"
 
 void UMassCrowdMemberTrait::BuildTemplate(FMassEntityTemplateBuildContext& BuildContext, UWorld& World) const
 {
diff --git a/Engine/Plugins/AI/MassCrowd/Source/MassCrowd/Public/MassCrowdMemberTrait.h b/Engine/Plugins/AI/MassCrowd/Source/MassCrowd/Public/MassCrowdMemberTrait.h
--- a/Engine/Plugins/AI/MassCrowd/Source/MassCrowd/Public/MassCrowdMemberTrait.h
+++ b/Engine/Plugins/AI/MassCrowd/Source/MassCrowd/Public/MassCrowdMemberTrait.h
@@ -6,7 +6,7 @@
 #include "MassCrowdMemberTrait.generated.h"
 
 /**
- * Trait that makes the agent aware of traffic.
+ * Trait to mark an entity with the crowd tag and add required fragments to track current lane
  */
 UCLASS(BlueprintType, EditInlineNew, CollapseCategories, meta = (DisplayName = "CrowdMember"))
 class MASSCROWD_API UMassCrowdMemberTrait : public UMassEntityTraitBase
diff --git a/Engine/Plugins/Runtime/MassGameplay/Source/MassSmartObjects/MassSmartObjects.Build.cs b/Engine/Plugins/Runtime/MassGameplay/Source/MassSmartObjects/MassSmartObjects.Build.cs
--- a/Engine/Plugins/Runtime/MassGameplay/Source/MassSmartObjects/MassSmartObjects.Build.cs
+++ b/Engine/Plugins/Runtime/MassGameplay/Source/MassSmartObjects/MassSmartObjects.Build.cs
@@ -32,6 +32,7 @@ namespace UnrealBuildTool.Rules
 					"MassMovement",
 					"MassSignals",
 					"MassSimulation",
+					"MassSpawner",
 					"SmartObjectsModule",
 					"StructUtils",
 					"ZoneGraph",
diff --git a/Engine/Plugins/Runtime/MassGameplay/Source/MassSmartObjects/Private/MassSmartObjectUserTrait.cpp b/Engine/Plugins/Runtime/MassGameplay/Source/MassSmartObjects/Private/MassSmartObjectUserTrait.cpp
new file mode 100644
--- /dev/null
+++ b/Engine/Plugins/Runtime/MassGameplay/Source/MassSmartObjects/Private/MassSmartObjectUserTrait.cpp
@@ -0,0 +1,11 @@
+﻿// Copyright Epic Games, Inc. All Rights Reserved.
+
+#include "MassSmartObjectUserTrait.h"
+
+#include "MassSmartObjectProcessor.h"
+#include "MassEntityTemplateRegistry.h"
+
+void UMassSmartObjectUserTrait::BuildTemplate(FMassEntityTemplateBuildContext& BuildContext, UWorld& World) const
+{
+	BuildContext.AddFragmentWithDefaultInitializer<FMassSmartObjectUserFragment>();
+}
diff --git a/Engine/Plugins/Runtime/MassGameplay/Source/MassSmartObjects/Public/MassSmartObjectUserTrait.h b/Engine/Plugins/Runtime/MassGameplay/Source/MassSmartObjects/Public/MassSmartObjectUserTrait.h
new file mode 100644
--- /dev/null
+++ b/Engine/Plugins/Runtime/MassGameplay/Source/MassSmartObjects/Public/MassSmartObjectUserTrait.h
@@ -0,0 +1,18 @@
+﻿// Copyright Epic Games, Inc. All Rights Reserved.
+
+#pragma once
+
+#include "MassEntityTraitBase.h"
+#include "MassSmartObjectUserTrait.generated.h"
+
+/**
+ * Trait to allow an entity to interact with SmartObjects
+ */
+UCLASS(BlueprintType, EditInlineNew, CollapseCategories, meta = (DisplayName = "SmartObject User"))
+class MASSSMARTOBJECTS_API UMassSmartObjectUserTrait : public UMassEntityTraitBase
+{
+	GENERATED_BODY()
+
+protected:
+	virtual void BuildTemplate(FMassEntityTemplateBuildContext& BuildContext, UWorld& World) const override;
+};
