Adding ability to specify representation subsystem class in the UpdateISMProcessor

diff --git a/Engine/Plugins/Runtime/MassGameplay/Source/MassLOD/Public/MassLODCollectorProcessor.h b/Engine/Plugins/Runtime/MassGameplay/Source/MassLOD/Public/MassLODCollectorProcessor.h
--- a/Engine/Plugins/Runtime/MassGameplay/Source/MassLOD/Public/MassLODCollectorProcessor.h
+++ b/Engine/Plugins/Runtime/MassGameplay/Source/MassLOD/Public/MassLODCollectorProcessor.h
@@ -15,7 +15,7 @@ struct FMassCollectorLODConfig
 	GENERATED_BODY()
 
 	/** Component Tag that will be used to associate LOD config */
-	UPROPERTY(EditAnywhere, Category = "Mass|LOD", config, meta = (BaseStruct = "ComponentTag"))
+	UPROPERTY(EditAnywhere, Category = "Mass|LOD", config, meta = (BaseStruct = "MassTag"))
 	FInstancedStruct TagFilter;
 
 	/** Set FOV Angles to define the limits between using Visible or Base LOD Distances */
diff --git a/Engine/Plugins/Runtime/MassGameplay/Source/MassLOD/Public/MassRepresentationLODCollectorProcessor.h b/Engine/Plugins/Runtime/MassGameplay/Source/MassLOD/Public/MassRepresentationLODCollectorProcessor.h
--- a/Engine/Plugins/Runtime/MassGameplay/Source/MassLOD/Public/MassRepresentationLODCollectorProcessor.h
+++ b/Engine/Plugins/Runtime/MassGameplay/Source/MassLOD/Public/MassRepresentationLODCollectorProcessor.h
@@ -16,7 +16,7 @@ struct FMassRepresentationLODCollectorConfig
 	GENERATED_BODY()
 
 	/** Component Tag that will be used to associate LOD config */
-	UPROPERTY(EditAnywhere, Category = "Mass|LOD", config, meta = (BaseStruct = "ComponentTag"))
+	UPROPERTY(EditAnywhere, Category = "Mass|LOD", config, meta = (BaseStruct = "MassTag"))
 	FInstancedStruct TagFilter;
 
 	/** Set FOV Angles to define the limits between using Visible or Base LOD Distances */
diff --git a/Engine/Plugins/Runtime/MassGameplay/Source/MassRepresentation/Private/MassUpdateISMProcessor.cpp b/Engine/Plugins/Runtime/MassGameplay/Source/MassRepresentation/Private/MassUpdateISMProcessor.cpp
--- a/Engine/Plugins/Runtime/MassGameplay/Source/MassRepresentation/Private/MassUpdateISMProcessor.cpp
+++ b/Engine/Plugins/Runtime/MassGameplay/Source/MassRepresentation/Private/MassUpdateISMProcessor.cpp
@@ -20,34 +20,40 @@ void UMassUpdateISMProcessor::Initialize(UObject& Owner)
 {
 	Super::Initialize(Owner);
 
-	RepresentationSubsystem = UWorld::GetSubsystem<UMassRepresentationSubsystem>(Owner.GetWorld());
+	if (const UWorld* World = Owner.GetWorld())
+	{
+		for (FMassUpdateISMConfig& Config : ISMConfigs)
+		{
+			Config.RepresentationSubsystem = Cast<UMassRepresentationSubsystem>(World->GetSubsystemBase(Config.RepresentationSubsystemClass ? Config.RepresentationSubsystemClass : TSubclassOf<UMassRepresentationSubsystem>(UMassRepresentationSubsystem::StaticClass())));
+		}
+	}
 }
 
 void UMassUpdateISMProcessor::ConfigureQueries()
 {
-	for (const FInstancedStruct& TagFilter : TagFilters)
+	for (FMassUpdateISMConfig& Config : ISMConfigs)
 	{
-		FMassEntityQuery& EntityQuery = EntityQueries.AddDefaulted_GetRef();
-		if (TagFilter.GetScriptStruct())
+		if (Config.TagFilter.GetScriptStruct())
 		{
-			EntityQuery.AddTagRequirement(*TagFilter.GetScriptStruct(), EMassFragmentPresence::All);
+			Config.EntityQuery.AddTagRequirement(*Config.TagFilter.GetScriptStruct(), EMassFragmentPresence::All);
 		}
-		EntityQuery.AddRequirement<FDataFragment_Transform>(EMassFragmentAccess::ReadOnly);
-		EntityQuery.AddRequirement<FMassRepresentationFragment>(EMassFragmentAccess::ReadWrite);
-		EntityQuery.AddRequirement<FMassRepresentationLODFragment>(EMassFragmentAccess::ReadWrite);
-		EntityQuery.AddChunkRequirement<FMassVisualizationChunkFragment>(EMassFragmentAccess::ReadWrite);
-		EntityQuery.SetChunkFilter(&FMassVisualizationChunkFragment::AreAnyEntitiesVisibleInChunk);
+		Config.EntityQuery.AddRequirement<FDataFragment_Transform>(EMassFragmentAccess::ReadOnly);
+		Config.EntityQuery.AddRequirement<FMassRepresentationFragment>(EMassFragmentAccess::ReadWrite);
+		Config.EntityQuery.AddRequirement<FMassRepresentationLODFragment>(EMassFragmentAccess::ReadWrite);
+		Config.EntityQuery.AddChunkRequirement<FMassVisualizationChunkFragment>(EMassFragmentAccess::ReadWrite);
+		Config.EntityQuery.SetChunkFilter(&FMassVisualizationChunkFragment::AreAnyEntitiesVisibleInChunk);
 	}
 }
 
 void UMassUpdateISMProcessor::Execute(UMassEntitySubsystem& EntitySubsystem, FMassExecutionContext& Context)
 {
-	check(RepresentationSubsystem);
-	FMassInstancedStaticMeshInfoArrayView ISMInfo = RepresentationSubsystem->GetMutableInstancedStaticMeshInfos();
 
-	for (FMassEntityQuery& EntityQuery : EntityQueries)
+	for (FMassUpdateISMConfig& Config : ISMConfigs)
 	{
-		EntityQuery.ForEachEntityChunk(EntitySubsystem, Context, [&ISMInfo](FMassExecutionContext& Context)
+		check(Config.RepresentationSubsystem);
+		FMassInstancedStaticMeshInfoArrayView ISMInfo = Config.RepresentationSubsystem->GetMutableInstancedStaticMeshInfos();
+
+		Config.EntityQuery.ForEachEntityChunk(EntitySubsystem, Context, [&ISMInfo](FMassExecutionContext& Context)
 		{
 			const TConstArrayView<FDataFragment_Transform> TransformList = Context.GetFragmentView<FDataFragment_Transform>();
 			const TArrayView<FMassRepresentationFragment> RepresentationList = Context.GetMutableFragmentView<FMassRepresentationFragment>();
diff --git a/Engine/Plugins/Runtime/MassGameplay/Source/MassRepresentation/Public/MassUpdateISMProcessor.h b/Engine/Plugins/Runtime/MassGameplay/Source/MassRepresentation/Public/MassUpdateISMProcessor.h
--- a/Engine/Plugins/Runtime/MassGameplay/Source/MassRepresentation/Public/MassUpdateISMProcessor.h
+++ b/Engine/Plugins/Runtime/MassGameplay/Source/MassRepresentation/Public/MassUpdateISMProcessor.h
@@ -10,6 +10,26 @@ class UMassRepresentationSubsystem;
 
 struct FMassInstancedStaticMeshInfo;
 
+USTRUCT()
+struct FMassUpdateISMConfig
+{
+	GENERATED_BODY()
+
+	/** Component Tag that will be used to associate Update ISM config */
+	UPROPERTY(EditAnywhere, Category = "Mass|ISM", config, meta = (BaseStruct = "MassTag"))
+	FInstancedStruct TagFilter;
+
+	/** To specify a different MassRepresentationSubsystem, otherwise leave it null to use the default one  */
+	UPROPERTY(EditAnywhere, Category = "Mass|ISM", config)
+	TSubclassOf<UMassRepresentationSubsystem> RepresentationSubsystemClass;
+
+	/** A cached pointer to the representation subsystem */
+	UPROPERTY(Transient)
+	UMassRepresentationSubsystem* RepresentationSubsystem = nullptr;
+
+	FMassEntityQuery EntityQuery;
+};
+
 UCLASS()
 class MASSREPRESENTATION_API UMassUpdateISMProcessor : public UMassProcessor
 {
@@ -35,13 +55,8 @@ protected:
 	 * @param Context is the execution context to be passed when executing the lambdas */
 	virtual void Execute(UMassEntitySubsystem& EntitySubsystem, FMassExecutionContext& Context) override;
 
-	/** Component Tag that will be used to filter out entities to update ISM */
-	UPROPERTY(EditAnywhere, Category = "Mass|LOD", config, meta = (BaseStruct = "ComponentTag"))
-	TArray<FInstancedStruct> TagFilters;
-
-	/** A cache pointer to the representation subsystem */
-	UPROPERTY(Transient)
-	UMassRepresentationSubsystem* RepresentationSubsystem;
+	/** Config that will be used to update ISM */
+	UPROPERTY(EditAnywhere, Category = "Mass|LOD", config, meta = (BaseStruct = "MassTag"))
+	TArray<FMassUpdateISMConfig> ISMConfigs;
 
-	TArray<FMassEntityQuery> EntityQueries;
 };
\ No newline at end of file
