[MassSmartObject] rebuild annotations only for Editor world

diff --git a/Engine/Plugins/Runtime/MassGameplay/Source/MassSmartObjects/Private/SmartObjectZoneAnnotations.cpp b/Engine/Plugins/Runtime/MassGameplay/Source/MassSmartObjects/Private/SmartObjectZoneAnnotations.cpp
--- a/Engine/Plugins/Runtime/MassGameplay/Source/MassSmartObjects/Private/SmartObjectZoneAnnotations.cpp
+++ b/Engine/Plugins/Runtime/MassGameplay/Source/MassSmartObjects/Private/SmartObjectZoneAnnotations.cpp
@@ -23,7 +23,11 @@ void USmartObjectZoneAnnotations::PostSubsystemsInitialized()
 	{
 		OnMainCollectionChangedHandle = SmartObjectSubsystem->OnMainCollectionChanged.AddLambda([this]()
 		{
-			RebuildForAllGraphs();
+			const UWorld* World = GetWorld();
+			if (World != nullptr && !World->IsGameWorld())
+			{
+				RebuildForAllGraphs();
+			}
 		});
 	}
 
