Fixed UMassDebugVisualizationComponent overzealously clearing out HISMComponents during ConstructVisualComponent resulting in losing previously added data which manifested by only the last entity config spawned being visualized.

diff --git a/Engine/Plugins/Runtime/MassGameplay/Source/MassGameplayDebug/Private/MassDebugVisualizationComponent.cpp b/Engine/Plugins/Runtime/MassGameplay/Source/MassGameplayDebug/Private/MassDebugVisualizationComponent.cpp
--- a/Engine/Plugins/Runtime/MassGameplay/Source/MassGameplayDebug/Private/MassDebugVisualizationComponent.cpp
+++ b/Engine/Plugins/Runtime/MassGameplay/Source/MassGameplayDebug/Private/MassDebugVisualizationComponent.cpp
@@ -42,21 +42,13 @@ void UMassDebugVisualizationComponent::ConditionallyConstructVisualComponent()
 
 void UMassDebugVisualizationComponent::ConstructVisualComponent()
 {
-	for (UHierarchicalInstancedStaticMeshComponent* ISM : VisualDataISMCs)
-	{
-		if (ensure(ISM))
-		{
-			ISM->ClearInstances();
-			ISM->UnregisterComponent();
-		}
-	}
-	VisualDataISMCs.Reset();
-
 	AActor* ActorOwner = GetOwner();
 	check(ActorOwner);
 	
-	for (const FAgentDebugVisualization& VisualData : VisualDataTable)
+	// add HISMCs only for types not added yet
+	for (int NewTypeIndex = VisualDataISMCs.Num(); NewTypeIndex < VisualDataTable.Num(); ++NewTypeIndex)
 	{
+		const FAgentDebugVisualization& VisualData = VisualDataTable[NewTypeIndex];
 		UHierarchicalInstancedStaticMeshComponent* HISMC = NewObject<UHierarchicalInstancedStaticMeshComponent>(ActorOwner);
 		HISMC->SetStaticMesh(VisualData.Mesh);
 		if (VisualData.MaterialOverride != nullptr)
