Made FMassSimulationVariableTickChunkFragment::GetChunkLOD support fragment's optionality

diff --git a/Engine/Plugins/Runtime/MassGameplay/Source/MassLOD/Public/MassSimulationLOD.h b/Engine/Plugins/Runtime/MassGameplay/Source/MassLOD/Public/MassSimulationLOD.h
--- a/Engine/Plugins/Runtime/MassGameplay/Source/MassLOD/Public/MassSimulationLOD.h
+++ b/Engine/Plugins/Runtime/MassGameplay/Source/MassLOD/Public/MassSimulationLOD.h
@@ -85,8 +85,8 @@ struct FMassSimulationVariableTickChunkFragment : public FMassVariableTickChunkF
 
 	static EMassLOD::Type GetChunkLOD(const FMassExecutionContext& Context)
 	{
-		const FMassSimulationVariableTickChunkFragment& ChunkFragment = Context.GetChunkFragment<FMassSimulationVariableTickChunkFragment>();
-		return ChunkFragment.GetLOD();
+		const FMassSimulationVariableTickChunkFragment* ChunkFragment = Context.GetChunkFragmentPtr<FMassSimulationVariableTickChunkFragment>();
+		return ChunkFragment ? ChunkFragment->GetLOD() : EMassLOD::High;
 	}
 
 };
