Updated Mass Entity Trait UClasses' markup - all had the same so I moved it up to the parent class

diff --git a/Engine/Plugins/AI/MassAI/Source/MassAIBehavior/Public/MassLookAtTargetTrait.h b/Engine/Plugins/AI/MassAI/Source/MassAIBehavior/Public/MassLookAtTargetTrait.h
--- a/Engine/Plugins/AI/MassAI/Source/MassAIBehavior/Public/MassLookAtTargetTrait.h
+++ b/Engine/Plugins/AI/MassAI/Source/MassAIBehavior/Public/MassLookAtTargetTrait.h
@@ -6,7 +6,7 @@
 #include "MassEntityTraitBase.h"
 #include "MassLookAtTargetTrait.generated.h"
 
-UCLASS(BlueprintType, EditInlineNew, CollapseCategories, meta=(DisplayName="Look At Target"))
+UCLASS(meta=(DisplayName="Look At Target"))
 class UMassLookAtTargetTrait : public UMassEntityTraitBase
 {
 	GENERATED_BODY()
diff --git a/Engine/Plugins/AI/MassAI/Source/MassAIBehavior/Public/MassLookAtTrait.h b/Engine/Plugins/AI/MassAI/Source/MassAIBehavior/Public/MassLookAtTrait.h
--- a/Engine/Plugins/AI/MassAI/Source/MassAIBehavior/Public/MassLookAtTrait.h
+++ b/Engine/Plugins/AI/MassAI/Source/MassAIBehavior/Public/MassLookAtTrait.h
@@ -6,7 +6,7 @@
 #include "MassEntityTraitBase.h"
 #include "MassLookAtTrait.generated.h"
 
-UCLASS(BlueprintType, EditInlineNew, CollapseCategories, meta=(DisplayName="Look At"))
+UCLASS(meta=(DisplayName="Look At"))
 class UMassLookAtTrait : public UMassEntityTraitBase
 {
 	GENERATED_BODY()
diff --git a/Engine/Plugins/AI/MassAI/Source/MassAIBehavior/Public/MassStateTreeTrait.h b/Engine/Plugins/AI/MassAI/Source/MassAIBehavior/Public/MassStateTreeTrait.h
--- a/Engine/Plugins/AI/MassAI/Source/MassAIBehavior/Public/MassStateTreeTrait.h
+++ b/Engine/Plugins/AI/MassAI/Source/MassAIBehavior/Public/MassStateTreeTrait.h
@@ -10,7 +10,7 @@ class UStateTree;
 /**
  * Feature that adds StateTree execution functionality to a mass agent.
  */
-UCLASS(BlueprintType, EditInlineNew, CollapseCategories, meta=(DisplayName="StateTree"))
+UCLASS(meta=(DisplayName="StateTree"))
 class MASSAIBEHAVIOR_API UMassStateTreeTrait : public UMassEntityTraitBase
 {
 	GENERATED_BODY()
diff --git a/Engine/Plugins/AI/MassAI/Source/MassAIBehavior/Public/MassZoneGraphAnnotationTrait.h b/Engine/Plugins/AI/MassAI/Source/MassAIBehavior/Public/MassZoneGraphAnnotationTrait.h
--- a/Engine/Plugins/AI/MassAI/Source/MassAIBehavior/Public/MassZoneGraphAnnotationTrait.h
+++ b/Engine/Plugins/AI/MassAI/Source/MassAIBehavior/Public/MassZoneGraphAnnotationTrait.h
@@ -5,7 +5,7 @@
 #include "MassEntityTraitBase.h"
 #include "MassZoneGraphAnnotationTrait.generated.h"
 
-UCLASS(BlueprintType, EditInlineNew, CollapseCategories, meta = (DisplayName = "ZoneGraph Annotation"))
+UCLASS(meta = (DisplayName = "ZoneGraph Annotation"))
 class MASSAIBEHAVIOR_API UMassZoneGraphAnnotationTrait : public UMassEntityTraitBase
 {
 	GENERATED_BODY()
diff --git a/Engine/Plugins/AI/MassAI/Source/MassAIMovement/Public/MassAvoidanceObstacleTrait.h b/Engine/Plugins/AI/MassAI/Source/MassAIMovement/Public/MassAvoidanceObstacleTrait.h
--- a/Engine/Plugins/AI/MassAI/Source/MassAIMovement/Public/MassAvoidanceObstacleTrait.h
+++ b/Engine/Plugins/AI/MassAI/Source/MassAIMovement/Public/MassAvoidanceObstacleTrait.h
@@ -5,7 +5,7 @@
 #include "MassEntityTraitBase.h"
 #include "MassAvoidanceObstacleTrait.generated.h"
 
-UCLASS(BlueprintType, EditInlineNew, CollapseCategories, meta = (DisplayName = "Dynamic Obstacle"))
+UCLASS(meta = (DisplayName = "Dynamic Obstacle"))
 class MASSAIMOVEMENT_API UMassAvoidanceObstacleTrait : public UMassEntityTraitBase
 {
 	GENERATED_BODY()
diff --git a/Engine/Plugins/AI/MassAI/Source/MassAIMovement/Public/MassAvoidanceTrait.h b/Engine/Plugins/AI/MassAI/Source/MassAIMovement/Public/MassAvoidanceTrait.h
--- a/Engine/Plugins/AI/MassAI/Source/MassAIMovement/Public/MassAvoidanceTrait.h
+++ b/Engine/Plugins/AI/MassAI/Source/MassAIMovement/Public/MassAvoidanceTrait.h
@@ -5,7 +5,7 @@
 #include "MassEntityTraitBase.h"
 #include "MassAvoidanceTrait.generated.h"
 
-UCLASS(BlueprintType, EditInlineNew, CollapseCategories, meta = (DisplayName = "Avoidance"))
+UCLASS(meta = (DisplayName = "Avoidance"))
 class MASSAIMOVEMENT_API UMassAvoidanceTrait : public UMassEntityTraitBase
 {
 	GENERATED_BODY()
diff --git a/Engine/Plugins/AI/MassAI/Source/MassAIMovement/Public/MassZoneGraphMovementTrait.h b/Engine/Plugins/AI/MassAI/Source/MassAIMovement/Public/MassZoneGraphMovementTrait.h
--- a/Engine/Plugins/AI/MassAI/Source/MassAIMovement/Public/MassZoneGraphMovementTrait.h
+++ b/Engine/Plugins/AI/MassAI/Source/MassAIMovement/Public/MassZoneGraphMovementTrait.h
@@ -7,7 +7,7 @@
 #include "MassAIMovementTypes.h"
 #include "MassZoneGraphMovementTrait.generated.h"
 
-UCLASS(BlueprintType, EditInlineNew, CollapseCategories, meta = (DisplayName = "ZoneGraph Movement"))
+UCLASS(meta = (DisplayName = "ZoneGraph Movement"))
 class MASSAIMOVEMENT_API UMassZoneGraphMovementTrait : public UMassEntityTraitBase
 {
 	GENERATED_BODY()
diff --git a/Engine/Plugins/AI/MassCrowd/Source/MassCrowd/Public/MassCrowdMemberTrait.h b/Engine/Plugins/AI/MassCrowd/Source/MassCrowd/Public/MassCrowdMemberTrait.h
--- a/Engine/Plugins/AI/MassCrowd/Source/MassCrowd/Public/MassCrowdMemberTrait.h
+++ b/Engine/Plugins/AI/MassCrowd/Source/MassCrowd/Public/MassCrowdMemberTrait.h
@@ -8,7 +8,7 @@
 /**
  * Trait to mark an entity with the crowd tag and add required fragments to track current lane
  */
-UCLASS(BlueprintType, EditInlineNew, CollapseCategories, meta = (DisplayName = "CrowdMember"))
+UCLASS(meta = (DisplayName = "CrowdMember"))
 class MASSCROWD_API UMassCrowdMemberTrait : public UMassEntityTraitBase
 {
 	GENERATED_BODY()
diff --git a/Engine/Plugins/AI/MassCrowd/Source/MassCrowd/Public/MassCrowdServerRepresentationTrait.h b/Engine/Plugins/AI/MassCrowd/Source/MassCrowd/Public/MassCrowdServerRepresentationTrait.h
--- a/Engine/Plugins/AI/MassCrowd/Source/MassCrowd/Public/MassCrowdServerRepresentationTrait.h
+++ b/Engine/Plugins/AI/MassCrowd/Source/MassCrowd/Public/MassCrowdServerRepresentationTrait.h
@@ -8,7 +8,7 @@
 #include "MassCrowdServerRepresentationTrait.generated.h"
 
 
-UCLASS(BlueprintType, EditInlineNew, CollapseCategories, meta=(DisplayName="Crowd Server Representation"))
+UCLASS(meta=(DisplayName="Crowd Server Representation"))
 class MASSCROWD_API UMassCrowdServerRepresentationTrait : public UMassEntityTraitBase
 {
 	GENERATED_BODY()
diff --git a/Engine/Plugins/Runtime/MassGameplay/Source/MassGameplayDebug/Public/MassDebugVisualizationTrait.h b/Engine/Plugins/Runtime/MassGameplay/Source/MassGameplayDebug/Public/MassDebugVisualizationTrait.h
--- a/Engine/Plugins/Runtime/MassGameplay/Source/MassGameplayDebug/Public/MassDebugVisualizationTrait.h
+++ b/Engine/Plugins/Runtime/MassGameplay/Source/MassGameplayDebug/Public/MassDebugVisualizationTrait.h
@@ -7,7 +7,7 @@
 #include "MassDebugVisualizationTrait.generated.h"
 
 
-UCLASS(BlueprintType, EditInlineNew, CollapseCategories, meta = (DisplayName = "Debug Visualization"))
+UCLASS(meta = (DisplayName = "Debug Visualization"))
 class MASSGAMEPLAYDEBUG_API UMassDebugVisualizationTrait : public UMassEntityTraitBase
 {
 	GENERATED_BODY()
diff --git a/Engine/Plugins/Runtime/MassGameplay/Source/MassLOD/Public/MassLODTrait.h b/Engine/Plugins/Runtime/MassGameplay/Source/MassLOD/Public/MassLODTrait.h
--- a/Engine/Plugins/Runtime/MassGameplay/Source/MassLOD/Public/MassLODTrait.h
+++ b/Engine/Plugins/Runtime/MassGameplay/Source/MassLOD/Public/MassLODTrait.h
@@ -5,7 +5,7 @@
 #include "MassEntityTraitBase.h"
 #include "MassLODTrait.generated.h"
 
-UCLASS(BlueprintType, EditInlineNew, CollapseCategories, meta = (DisplayName = "LOD"))
+UCLASS(meta = (DisplayName = "LOD"))
 class MASSLOD_API UMassLODTrait : public UMassEntityTraitBase
 {
 	GENERATED_BODY()
diff --git a/Engine/Plugins/Runtime/MassGameplay/Source/MassReplication/Public/MassReplicationTrait.h b/Engine/Plugins/Runtime/MassGameplay/Source/MassReplication/Public/MassReplicationTrait.h
--- a/Engine/Plugins/Runtime/MassGameplay/Source/MassReplication/Public/MassReplicationTrait.h
+++ b/Engine/Plugins/Runtime/MassGameplay/Source/MassReplication/Public/MassReplicationTrait.h
@@ -6,7 +6,7 @@
 #include "MassReplicationTrait.generated.h"
 
 
-UCLASS(BlueprintType, EditInlineNew, CollapseCategories, meta=(DisplayName="Replication"))
+UCLASS(meta=(DisplayName="Replication"))
 class MASSREPLICATION_API UMassReplicationTrait : public UMassEntityTraitBase
 {
 	GENERATED_BODY()
diff --git a/Engine/Plugins/Runtime/MassGameplay/Source/MassRepresentation/Public/MassVisualizationTrait.h b/Engine/Plugins/Runtime/MassGameplay/Source/MassRepresentation/Public/MassVisualizationTrait.h
--- a/Engine/Plugins/Runtime/MassGameplay/Source/MassRepresentation/Public/MassVisualizationTrait.h
+++ b/Engine/Plugins/Runtime/MassGameplay/Source/MassRepresentation/Public/MassVisualizationTrait.h
@@ -10,7 +10,7 @@
 class UMassRepresentationSubsystem;
 class UMassProcessor;
 
-UCLASS(BlueprintType, EditInlineNew, CollapseCategories, meta=(DisplayName="Visualization"))
+UCLASS(meta=(DisplayName="Visualization"))
 class MASSREPRESENTATION_API UMassVisualizationTrait : public UMassEntityTraitBase
 {
 	GENERATED_BODY()
diff --git a/Engine/Plugins/Runtime/MassGameplay/Source/MassSmartObjects/Public/MassSmartObjectUserTrait.h b/Engine/Plugins/Runtime/MassGameplay/Source/MassSmartObjects/Public/MassSmartObjectUserTrait.h
--- a/Engine/Plugins/Runtime/MassGameplay/Source/MassSmartObjects/Public/MassSmartObjectUserTrait.h
+++ b/Engine/Plugins/Runtime/MassGameplay/Source/MassSmartObjects/Public/MassSmartObjectUserTrait.h
@@ -8,7 +8,7 @@
 /**
  * Trait to allow an entity to interact with SmartObjects
  */
-UCLASS(BlueprintType, EditInlineNew, CollapseCategories, meta = (DisplayName = "SmartObject User"))
+UCLASS(meta = (DisplayName = "SmartObject User"))
 class MASSSMARTOBJECTS_API UMassSmartObjectUserTrait : public UMassEntityTraitBase
 {
 	GENERATED_BODY()
diff --git a/Engine/Plugins/Runtime/MassGameplay/Source/MassSpawner/Public/MassAssortedFragmentsTrait.h b/Engine/Plugins/Runtime/MassGameplay/Source/MassSpawner/Public/MassAssortedFragmentsTrait.h
--- a/Engine/Plugins/Runtime/MassGameplay/Source/MassSpawner/Public/MassAssortedFragmentsTrait.h
+++ b/Engine/Plugins/Runtime/MassGameplay/Source/MassSpawner/Public/MassAssortedFragmentsTrait.h
@@ -10,7 +10,7 @@
 /**
 * Mass Agent Feature which appends a list of specified fragments.  
 */
-UCLASS(BlueprintType, EditInlineNew, CollapseCategories, meta=(DisplayName="Assorted Fragments"))
+UCLASS(meta=(DisplayName="Assorted Fragments"))
 class MASSSPAWNER_API UMassAssortedFragmentsTrait : public UMassEntityTraitBase
 {
 	GENERATED_BODY()
diff --git a/Engine/Plugins/Runtime/MassGameplay/Source/MassSpawner/Public/MassEntityTraitBase.h b/Engine/Plugins/Runtime/MassGameplay/Source/MassSpawner/Public/MassEntityTraitBase.h
--- a/Engine/Plugins/Runtime/MassGameplay/Source/MassSpawner/Public/MassEntityTraitBase.h
+++ b/Engine/Plugins/Runtime/MassGameplay/Source/MassSpawner/Public/MassEntityTraitBase.h
@@ -14,7 +14,7 @@ struct FMassEntityTemplateBuildContext;
  * For example, a fragment can be added based on a referenced asset, or some memory hungry settings can be
  * cached and just and index stored on a fragment.
  */
-UCLASS(Abstract)
+UCLASS(Abstract, BlueprintType, EditInlineNew, CollapseCategories)
 class MASSSPAWNER_API UMassEntityTraitBase : public UObject
 {
 	GENERATED_BODY()
