Added new basic mass entity traits: randomized velocity and simple movement. To be used for demonstration purposes and prototyping.

diff --git a/Engine/Plugins/Runtime/MassGameplay/Source/MassMovement/Private/MassSimpleMovementTrait.cpp b/Engine/Plugins/Runtime/MassGameplay/Source/MassMovement/Private/MassSimpleMovementTrait.cpp
new file mode 100644
--- /dev/null
+++ b/Engine/Plugins/Runtime/MassGameplay/Source/MassMovement/Private/MassSimpleMovementTrait.cpp
@@ -0,0 +1,61 @@
+// Copyright Epic Games, Inc. All Rights Reserved.
+
+#include "MassSimpleMovementTrait.h"
+#include "MassEntityTemplateRegistry.h"
+#include "Engine/World.h"
+#include "MassMovementTypes.h"
+#include "MassMovementFragments.h"
+#include "MassCommonFragments.h"
+#include "MassSimulationLOD.h"
+
+
+//----------------------------------------------------------------------//
+//  UMassSimpleMovementTrait
+//----------------------------------------------------------------------//
+void UMassSimpleMovementTrait::BuildTemplate(FMassEntityTemplateBuildContext& BuildContext, UWorld& World) const
+{
+	BuildContext.AddFragment<FDataFragment_Transform>();
+	BuildContext.AddFragment<FMassVelocityFragment>();
+	BuildContext.AddTag<FMassSimpleMovementTag>();	
+}
+
+//----------------------------------------------------------------------//
+//  UMassSimpleMovementProcessor
+//----------------------------------------------------------------------//
+UMassSimpleMovementProcessor::UMassSimpleMovementProcessor()
+{
+	bAutoRegisterWithProcessingPhases = true;
+	ExecutionFlags = (int32)EProcessorExecutionFlags::All;
+	ExecutionOrder.ExecuteInGroup = UE::Mass::ProcessorGroupNames::Avoidance;
+}
+
+void UMassSimpleMovementProcessor::ConfigureQueries()
+{
+	EntityQuery.AddRequirement<FMassVelocityFragment>(EMassFragmentAccess::ReadOnly);
+	EntityQuery.AddRequirement<FDataFragment_Transform>(EMassFragmentAccess::ReadWrite);
+	EntityQuery.AddTagRequirement<FMassSimpleMovementTag>(EMassFragmentPresence::All);
+
+	EntityQuery.AddRequirement<FMassSimulationLODFragment>(EMassFragmentAccess::ReadOnly, EMassFragmentPresence::Optional);
+	EntityQuery.AddChunkRequirement<FMassSimulationVariableTickChunkFragment>(EMassFragmentAccess::ReadOnly, EMassFragmentPresence::Optional);
+	EntityQuery.SetChunkFilter(&FMassSimulationVariableTickChunkFragment::ShouldTickChunkThisFrame);
+}
+
+void UMassSimpleMovementProcessor::Execute(UMassEntitySubsystem& EntitySubsystem, FMassExecutionContext& Context)
+{
+	EntityQuery.ForEachEntityChunk(EntitySubsystem, Context, ([this](FMassExecutionContext& Context)
+		{
+			const TConstArrayView<FMassVelocityFragment> VelocitiesList = Context.GetFragmentView<FMassVelocityFragment>();
+			const TArrayView<FDataFragment_Transform> TransformsList = Context.GetMutableFragmentView<FDataFragment_Transform>();
+			const TConstArrayView<FMassSimulationLODFragment> SimLODList = Context.GetFragmentView<FMassSimulationLODFragment>();
+			const bool bHasLOD = (SimLODList.Num() > 0);
+			const float WorldDeltaTime = Context.GetDeltaTimeSeconds();
+		
+			for (int32 EntityIndex = 0; EntityIndex < Context.GetNumEntities(); ++EntityIndex)
+			{
+				const FMassVelocityFragment& Velocity = VelocitiesList[EntityIndex];
+				FTransform& Transform = TransformsList[EntityIndex].GetMutableTransform();
+				const float DeltaTime = bHasLOD ? SimLODList[EntityIndex].DeltaTime : WorldDeltaTime; 
+				Transform.SetTranslation(Transform.GetTranslation() + Velocity.Value * DeltaTime);
+			}
+		}));
+}
diff --git a/Engine/Plugins/Runtime/MassGameplay/Source/MassMovement/Private/MassVelocityRandomizerTrait.cpp b/Engine/Plugins/Runtime/MassGameplay/Source/MassMovement/Private/MassVelocityRandomizerTrait.cpp
new file mode 100644
--- /dev/null
+++ b/Engine/Plugins/Runtime/MassGameplay/Source/MassMovement/Private/MassVelocityRandomizerTrait.cpp
@@ -0,0 +1,71 @@
+// Copyright Epic Games, Inc. All Rights Reserved.
+
+#include "MassVelocityRandomizerTrait.h"
+#include "MassEntityTemplateRegistry.h"
+#include "Engine/World.h"
+#include "MassMovementFragments.h"
+
+
+//----------------------------------------------------------------------//
+//  UMassVelocityRandomizerTrait
+//----------------------------------------------------------------------//
+void UMassVelocityRandomizerTrait::BuildTemplate(FMassEntityTemplateBuildContext& BuildContext, UWorld& World) const
+{
+	BuildContext.AddFragment<FMassVelocityFragment>();
+	
+	UMassRandomVelocityInitializer* VelocityInitializer = NewObject<UMassRandomVelocityInitializer>(&World);
+	check(VelocityInitializer);
+	VelocityInitializer->SetParameters(MinSpeed, MaxSpeed, bSetZComponent);
+	BuildContext.AddInitializer(*VelocityInitializer);
+}
+
+//----------------------------------------------------------------------//
+//  UMassRandomVelocityInitializer
+//----------------------------------------------------------------------//
+UMassRandomVelocityInitializer::UMassRandomVelocityInitializer()
+{
+	FragmentType = FMassVelocityFragment::StaticStruct();
+}
+
+void UMassRandomVelocityInitializer::SetParameters(const float InMinSpeed, const float InMaxSpeed, const bool bInSetZComponent)
+{
+	MinSpeed = InMinSpeed;
+	MaxSpeed = InMaxSpeed;
+	bSetZComponent = bInSetZComponent;
+}
+
+void UMassRandomVelocityInitializer::ConfigureQueries()
+{
+	EntityQuery.AddRequirement<FMassVelocityFragment>(EMassFragmentAccess::ReadWrite);
+}
+
+void UMassRandomVelocityInitializer::Execute(UMassEntitySubsystem& EntitySubsystem, FMassExecutionContext& Context)
+{
+	// note: the author is aware that the vectors produced below are not distributed uniformly, but it's good enough
+	if (bSetZComponent)
+	{
+		EntityQuery.ForEachEntityChunk(EntitySubsystem, Context, ([this](FMassExecutionContext& Context)
+			{
+				const TArrayView<FMassVelocityFragment> VelocitiesList = Context.GetMutableFragmentView<FMassVelocityFragment>();
+				for (FMassVelocityFragment& VelocityFragment : VelocitiesList)
+				{
+					const FVector RandomVector = FVector(FMath::FRandRange(-1.f, 1.f), FMath::FRandRange(-1.f, 1.f), FMath::FRandRange(-1.f, 1.f)).GetSafeNormal();
+
+					VelocityFragment.Value = RandomVector * FMath::FRandRange(MinSpeed, MaxSpeed);
+				}
+			}));
+	}
+	else
+	{
+		EntityQuery.ForEachEntityChunk(EntitySubsystem, Context, ([this](FMassExecutionContext& Context)
+			{
+				const TArrayView<FMassVelocityFragment> VelocitiesList = Context.GetMutableFragmentView<FMassVelocityFragment>();
+				for (FMassVelocityFragment& VelocityFragment : VelocitiesList)
+				{
+					const FVector RandomVector = FVector(FMath::FRandRange(-1.f, 1.f), FMath::FRandRange(-1.f, 1.f), 0).GetSafeNormal2D();
+
+					VelocityFragment.Value = RandomVector * FMath::FRandRange(MinSpeed, MaxSpeed);
+				}
+			}));
+	}	
+}
diff --git a/Engine/Plugins/Runtime/MassGameplay/Source/MassMovement/Public/MassSimpleMovementTrait.h b/Engine/Plugins/Runtime/MassGameplay/Source/MassMovement/Public/MassSimpleMovementTrait.h
new file mode 100644
--- /dev/null
+++ b/Engine/Plugins/Runtime/MassGameplay/Source/MassMovement/Public/MassSimpleMovementTrait.h
@@ -0,0 +1,40 @@
+// Copyright Epic Games, Inc. All Rights Reserved.
+
+#pragma once
+
+#include "MassEntityTypes.h"
+#include "MassEntityTraitBase.h"
+#include "MassProcessor.h"
+#include "MassSimpleMovementTrait.generated.h"
+
+
+USTRUCT()
+struct FMassSimpleMovementTag : public FMassTag
+{
+	GENERATED_BODY()
+};
+
+
+UCLASS(meta = (DisplayName = "Simple Movement"))
+class MASSMOVEMENT_API UMassSimpleMovementTrait : public UMassEntityTraitBase
+{
+	GENERATED_BODY()
+
+protected:
+	virtual void BuildTemplate(FMassEntityTemplateBuildContext& BuildContext, UWorld& World) const override;
+};
+
+
+UCLASS()
+class MASSMOVEMENT_API UMassSimpleMovementProcessor : public UMassProcessor
+{
+	GENERATED_BODY()
+public:
+	UMassSimpleMovementProcessor();
+		
+protected:	
+	virtual void ConfigureQueries() override;
+	virtual void Execute(UMassEntitySubsystem& EntitySubsystem, FMassExecutionContext& Context) override;
+
+	FMassEntityQuery EntityQuery;
+};
diff --git a/Engine/Plugins/Runtime/MassGameplay/Source/MassMovement/Public/MassVelocityRandomizerTrait.h b/Engine/Plugins/Runtime/MassGameplay/Source/MassMovement/Public/MassVelocityRandomizerTrait.h
new file mode 100644
--- /dev/null
+++ b/Engine/Plugins/Runtime/MassGameplay/Source/MassMovement/Public/MassVelocityRandomizerTrait.h
@@ -0,0 +1,57 @@
+// Copyright Epic Games, Inc. All Rights Reserved.
+
+#pragma once
+
+#include "MassEntityTraitBase.h"
+#include "MassTranslator.h"
+#include "MassVelocityRandomizerTrait.generated.h"
+
+
+class UMassRandomVelocityInitializer;
+
+UCLASS(meta = (DisplayName = "Velocity Randomizer"))
+class MASSMOVEMENT_API UMassVelocityRandomizerTrait : public UMassEntityTraitBase
+{
+	GENERATED_BODY()
+
+protected:
+	virtual void BuildTemplate(FMassEntityTemplateBuildContext& BuildContext, UWorld& World) const override;
+
+	/** The speed is expressed in UnrealUnits per second, which usually translates to 0.01m/s */
+	UPROPERTY(Category = "Velocity", EditAnywhere, meta = (UIMin = 0.0, ClampMin = 0.0))
+	float MinSpeed = 0.f;
+
+	/** The speed is expressed in UnrealUnits per second, which usually translates to 0.01m/s */
+	UPROPERTY(Category = "Velocity", EditAnywhere, meta = (UIMin = 1.0, ClampMin = 1.0))
+	float MaxSpeed = 200.f;
+
+	UPROPERTY(Category = "Velocity", EditAnywhere)
+	bool bSetZComponent = false;
+};
+
+
+UCLASS()
+class MASSMOVEMENT_API UMassRandomVelocityInitializer : public UMassFragmentInitializer
+{
+	GENERATED_BODY()
+public:
+	UMassRandomVelocityInitializer();
+
+	void SetParameters(const float InMinSpeed, const float InMaxSpeed, const bool bInSetZComponent);
+
+protected:
+	virtual void ConfigureQueries() override;
+	virtual void Execute(UMassEntitySubsystem& EntitySubsystem, FMassExecutionContext& Context) override;
+
+	FMassEntityQuery EntityQuery;
+
+	UPROPERTY()
+	float MinSpeed = 0.f;
+
+	/** The default max is set to 0 to enforce explicit configuration via SetParameters call. */
+	UPROPERTY()
+	float MaxSpeed = 0.f;
+
+	UPROPERTY()
+	bool bSetZComponent = false;
+};
