Rename LOD trait to SimulationLOD trait

#rnx

diff --git a/Engine/Plugins/Runtime/MassGameplay/Source/MassLOD/Private/MassLODModule.cpp b/Engine/Plugins/Runtime/MassGameplay/Source/MassLOD/Private/MassLODModule.cpp
--- a/Engine/Plugins/Runtime/MassGameplay/Source/MassLOD/Private/MassLODModule.cpp
+++ b/Engine/Plugins/Runtime/MassGameplay/Source/MassLOD/Private/MassLODModule.cpp
@@ -31,6 +31,8 @@ void FMassLODModule::StartupModule()
 	Redirects.Emplace(ECoreRedirectFlags::Type_Class, TEXT("/Script/MassCommon.MassProcessor_MassSimulationLOD"), TEXT("/Script/MassLOD.MassProcessor_MassSimulationLOD"));
 
 	Redirects.Emplace(ECoreRedirectFlags::Type_Class, TEXT("/Script/MassCommon.MassProcessor_MassViewerLOD"), TEXT("/Script/MassLOD.MassProcessor_MassViewerLOD"));
+
+	Redirects.Emplace(ECoreRedirectFlags::Type_Class, TEXT("/Script/MassLOD.MassLODTrait"), TEXT("/Script/MassLOD.MassSimulationLODTrait"));
 	
 	FCoreRedirects::AddRedirectList(Redirects, TEXT("MassMovement"));
 }
diff --git a/Engine/Plugins/Runtime/MassGameplay/Source/MassLOD/Private/MassLODTrait.cpp b/Engine/Plugins/Runtime/MassGameplay/Source/MassLOD/Private/MassLODTrait.cpp
--- a/Engine/Plugins/Runtime/MassGameplay/Source/MassLOD/Private/MassLODTrait.cpp
+++ b/Engine/Plugins/Runtime/MassGameplay/Source/MassLOD/Private/MassLODTrait.cpp
@@ -7,7 +7,7 @@
 #include "MassSimulationLOD.h"
 
 
-void UMassLODTrait::BuildTemplate(FMassEntityTemplateBuildContext& BuildContext, UWorld& World) const
+void UMassSimulationLODTrait::BuildTemplate(FMassEntityTemplateBuildContext& BuildContext, UWorld& World) const
 {
 	BuildContext.AddFragmentWithDefaultInitializer<FDataFragment_Transform>();
 	BuildContext.AddFragmentWithDefaultInitializer<FMassLODInfoFragment>();
diff --git a/Engine/Plugins/Runtime/MassGameplay/Source/MassLOD/Public/MassLODTrait.h b/Engine/Plugins/Runtime/MassGameplay/Source/MassLOD/Public/MassLODTrait.h
--- a/Engine/Plugins/Runtime/MassGameplay/Source/MassLOD/Public/MassLODTrait.h
+++ b/Engine/Plugins/Runtime/MassGameplay/Source/MassLOD/Public/MassLODTrait.h
@@ -5,8 +5,8 @@
 #include "MassEntityTraitBase.h"
 #include "MassLODTrait.generated.h"
 
-UCLASS(meta = (DisplayName = "LOD"))
-class MASSLOD_API UMassLODTrait : public UMassEntityTraitBase
+UCLASS(meta = (DisplayName = "SimulationLOD"))
+class MASSLOD_API UMassSimulationLODTrait : public UMassEntityTraitBase
 {
 	GENERATED_BODY()
 
