Cleanup Entity Template (used FMassArchetypeCompositionDescriptor and FMassArchetypeInitialFragmentValues)

#rnx

diff --git a/Engine/Plugins/AI/MassAI/Source/MassAIBehavior/Private/MassStateTreeTrait.cpp b/Engine/Plugins/AI/MassAI/Source/MassAIBehavior/Private/MassStateTreeTrait.cpp
--- a/Engine/Plugins/AI/MassAI/Source/MassAIBehavior/Private/MassStateTreeTrait.cpp
+++ b/Engine/Plugins/AI/MassAI/Source/MassAIBehavior/Private/MassStateTreeTrait.cpp
@@ -67,7 +67,7 @@ void UMassStateTreeTrait::ValidateTemplate(FMassEntityTemplateBuildContext& Buil
 			}
 			else if (ItemDesc.Struct->IsChildOf(FMassFragment::StaticStruct()))
 			{
-				const bool bContainsFragment = BuildContext.GetMutableFragments().GetFragments().ContainsByPredicate([&ItemDesc](const FInstancedStruct& Item) { return Item.GetScriptStruct() == ItemDesc.Struct; });
+				const bool bContainsFragment = BuildContext.HasFragment(*CastChecked<UScriptStruct>(ItemDesc.Struct));
 				UE_CVLOG(!bContainsFragment, MassStateTreeSubsystem, LogMassBehavior, Error, TEXT("StateTree %s: Could not find required fragment %s"), *GetNameSafe(StateTree), *GetNameSafe(ItemDesc.Struct));
 			}
 			else
diff --git a/Engine/Plugins/Runtime/MassEntity/Source/MassEntity/Public/MassEntityTypes.h b/Engine/Plugins/Runtime/MassEntity/Source/MassEntity/Public/MassEntityTypes.h
--- a/Engine/Plugins/Runtime/MassEntity/Source/MassEntity/Public/MassEntityTypes.h
+++ b/Engine/Plugins/Runtime/MassEntity/Source/MassEntity/Public/MassEntityTypes.h
@@ -3,6 +3,7 @@
 #pragma once
 
 #include "ScriptStructTypeBitSet.h"
+#include "MassProcessingTypes.h"
 #include "MassEntityTypes.generated.h"
 
 
@@ -126,6 +127,7 @@ struct FMassArchetypeCompositionDescriptor
 	{
 		Fragments.Reset();
 		Tags.Reset();
+		ChunkFragments.Reset();
 	}
 
 	bool IsEquivalent(const FMassFragmentBitSet& InFragmentBitSet, const FMassTagBitSet& InTagBitSet, const FMassChunkFragmentBitSet& InChunkFragmentsBitSet) const
@@ -153,6 +155,24 @@ struct FMassArchetypeCompositionDescriptor
 		return CalculateHash(Fragments, Tags, ChunkFragments);
 	}
 
+	void DebugOutputDescription(FOutputDevice& Ar) const
+	{
+#if WITH_MASSENTITY_DEBUG
+		const bool bAutoLineEnd = Ar.GetAutoEmitLineTerminator();
+		Ar.SetAutoEmitLineTerminator(false);
+
+		Ar.Logf(TEXT("Fragments:\n"));
+		Fragments.DebugGetStringDesc(Ar);
+		Ar.Logf(TEXT("Tags:\n"));
+		Tags.DebugGetStringDesc(Ar);
+		Ar.Logf(TEXT("ChunkFragments:\n"));
+		ChunkFragments.DebugGetStringDesc(Ar);
+
+		Ar.SetAutoEmitLineTerminator(bAutoLineEnd);
+#endif // WITH_MASSENTITY_DEBUG
+
+	}
+
 	FMassFragmentBitSet Fragments;
 	FMassTagBitSet Tags;
 	FMassChunkFragmentBitSet ChunkFragments;
@@ -161,13 +181,8 @@ struct FMassArchetypeCompositionDescriptor
 struct FMassArchetypeFragmentsInitialValues
 {
 	FMassArchetypeFragmentsInitialValues() = default;
-	FMassArchetypeFragmentsInitialValues(TConstArrayView<FInstancedStruct> InChunkFragments)
-	: ChunkFragments(InChunkFragments)
-	{
-		
-	}
 
-	// @todo add the regular fragments initial value here.
+	TArray<FInstancedStruct> Fragments;
 	TArray<FInstancedStruct> ChunkFragments;
 
 	// @todo add the shared fragments values here
diff --git a/Engine/Plugins/Runtime/MassGameplay/Source/MassCommon/Private/MassCommonTypes.cpp b/Engine/Plugins/Runtime/MassGameplay/Source/MassCommon/Private/MassCommonTypes.cpp
--- a/Engine/Plugins/Runtime/MassGameplay/Source/MassCommon/Private/MassCommonTypes.cpp
+++ b/Engine/Plugins/Runtime/MassGameplay/Source/MassCommon/Private/MassCommonTypes.cpp
@@ -114,43 +114,4 @@ FColor GetEntityDebugColor(FMassEntityHandle Entity)
 }
 
 } // namespace UE::Mass::Debug
-#endif // WITH_MASSGAMEPLAY_DEBUG
-
-//----------------------------------------------------------------------//
-// FMassUniqueFragmentCollection 
-//----------------------------------------------------------------------//
-void FMassUniqueFragmentCollection::Add(const FInstancedStruct& InFragment)
-{
-	const UScriptStruct* InFragmentType = InFragment.GetScriptStruct();
-	if (InFragmentType && !FragmentBitSet.Contains(*InFragmentType))
-	{
-		ensureMsgf(InFragmentType->IsChildOf(FMassFragment::StaticStruct()), TEXT("Given struct is not of a fragment type. Make sure your struct extends FMassFragment or its derivatives"));
-		FragmentInstances.Add(InFragment);
-		FragmentBitSet.Add(*InFragmentType);
-	}
-}
-
-void FMassUniqueFragmentCollection::Add(const UScriptStruct* InFragmentType)
-{
-	if (InFragmentType && !FragmentBitSet.Contains(*InFragmentType))
-	{
-		ensureMsgf(InFragmentType->IsChildOf(FMassFragment::StaticStruct()), TEXT("Given struct is not of a fragment type. Make sure your struct extends FMassFragment or its derivatives"));
-		FragmentInstances.Add(FInstancedStruct(InFragmentType));
-		FragmentBitSet.Add(*InFragmentType);
-	}
-}
-
-void FMassUniqueFragmentCollection::DebugOutputDescription(FOutputDevice& Ar) const
-{
-#if WITH_MASSGAMEPLAY_DEBUG
-	const bool bAutoLineEnd = Ar.GetAutoEmitLineTerminator();
-	Ar.SetAutoEmitLineTerminator(false);
-
-	for (const FInstancedStruct& Fragment : FragmentInstances)
-	{
-		const UScriptStruct* StructType = Fragment.GetScriptStruct();
-		Ar.Logf(TEXT("%s, "), *GetNameSafe(StructType));
-	}
-	Ar.SetAutoEmitLineTerminator(bAutoLineEnd);
-#endif // WITH_MASSGAMEPLAY_DEBUG
-}
+#endif // WITH_MASSGAMEPLAY_DEBUG
\ No newline at end of file
diff --git a/Engine/Plugins/Runtime/MassGameplay/Source/MassCommon/Public/MassCommonTypes.h b/Engine/Plugins/Runtime/MassGameplay/Source/MassCommon/Public/MassCommonTypes.h
--- a/Engine/Plugins/Runtime/MassGameplay/Source/MassCommon/Public/MassCommonTypes.h
+++ b/Engine/Plugins/Runtime/MassGameplay/Source/MassCommon/Public/MassCommonTypes.h
@@ -43,67 +43,6 @@ struct FMassNetworkID : public FSequentialIDBase
 	explicit FMassNetworkID(uint32 InID) : FSequentialIDBase(InID) {}
 };
 
-/**
- *  No Remove* functions on purpose. Collections are to be built once and stored, never manipulated
- */
-USTRUCT()
-struct MASSCOMMON_API FMassUniqueFragmentCollection
-{
-	GENERATED_BODY()
-
-	FMassUniqueFragmentCollection() = default;
-
-	/** Adds a unique fragment instance. If a fragment of this type has already been added then its values won't be 
-	 *  overridden and this call will have no effect */
-	void Add(const FInstancedStruct& InFragment);
-
-	/** Adds a unique fragment type. If a fragment of this type has already been added then no action is performed */
-	void Add(const UScriptStruct* InFragmentType);
-
-	template<typename T>
-	FORCEINLINE void Add()
-	{
-		Add(T::StaticStruct());
-	}
-
-	/** Fetches a fragment instance of a given type. If it hasn't been added before it will be created first. */
-	template<typename T>
-	T& Add_GetRef()
-	{
-		static_assert(TIsDerivedFrom<T, FMassFragment>::IsDerived, "Given struct is not of a fragment type. If you use this FMassUniqueFragmentCollection with the EntitySubsystem you'll hit some checks");
-		if (const FInstancedStruct* ElementPtr = FragmentInstances.FindByPredicate(FSameTypeScriptStructPredicate(T::StaticStruct())))
-		{
-			return ElementPtr->GetMutable<T>();
-		}
-
-		FragmentBitSet.Add<T>();
-		return FragmentInstances.Add_GetRef(FInstancedStruct(T::StaticStruct())).GetMutable<T>();
-	}
-
-	template<typename T>
-	T& GetRefChecked()
-	{
-		static_assert(TIsDerivedFrom<T, FMassFragment>::IsDerived, "Given struct is not of a fragment type. Make sure your struct extends FMassFragment or its derivatives");
-		FInstancedStruct* ElementPtr = FragmentInstances.FindByPredicate(FSameTypeScriptStructPredicate(T::StaticStruct()));
-		check(ElementPtr);
-		return ElementPtr->GetMutable<T>();
-	}
-
-	int32 GetFragmentsNum() const { return FragmentInstances.Num(); }
-	TArrayView<const FInstancedStruct> GetFragments() const { return MakeArrayView(FragmentInstances); }
-	const FMassFragmentBitSet& GetFragmentBitSet() const { return FragmentBitSet; }
-
-	bool IsEmpty() const { return FragmentInstances.IsEmpty(); }
-
-	void DebugOutputDescription(FOutputDevice& Ar) const;
-
-private:
-	UPROPERTY()
-	TArray<FInstancedStruct> FragmentInstances;
-
-	FMassFragmentBitSet FragmentBitSet;
-};
-
 /** Float encoded in int16, 1cm accuracy. */
 USTRUCT()
 struct MASSCOMMON_API FMassInt16Real
diff --git a/Engine/Plugins/Runtime/MassGameplay/Source/MassSpawner/Private/MassEntityConfigAsset.cpp b/Engine/Plugins/Runtime/MassGameplay/Source/MassSpawner/Private/MassEntityConfigAsset.cpp
--- a/Engine/Plugins/Runtime/MassGameplay/Source/MassSpawner/Private/MassEntityConfigAsset.cpp
+++ b/Engine/Plugins/Runtime/MassGameplay/Source/MassSpawner/Private/MassEntityConfigAsset.cpp
@@ -73,7 +73,7 @@ const FMassEntityTemplate* FMassEntityConfig::GetOrCreateEntityTemplate(AActor&
 		Trait->ValidateTemplate(BuildContext, *World);
 	}
 
-	if (ensureMsgf(Template.GetFragmentsNum(), TEXT("Need at least one fragment to create an Archetype")))
+	if (ensureMsgf(!Template.IsEmpty(), TEXT("Need at least one fragment to create an Archetype")))
 	{
 		Template.SetUpProcessors(BuildContext.Handlers, OwnerActor);
 		TemplateRegistry.InitializeEntityTemplate(Template);
diff --git a/Engine/Plugins/Runtime/MassGameplay/Source/MassSpawner/Private/MassEntityTemplate.cpp b/Engine/Plugins/Runtime/MassGameplay/Source/MassSpawner/Private/MassEntityTemplate.cpp
--- a/Engine/Plugins/Runtime/MassGameplay/Source/MassSpawner/Private/MassEntityTemplate.cpp
+++ b/Engine/Plugins/Runtime/MassGameplay/Source/MassSpawner/Private/MassEntityTemplate.cpp
@@ -47,8 +47,8 @@ FString FMassEntityTemplate::DebugGetDescription(UMassEntitySubsystem* EntitySub
 	}
 	else
 	{
-		Ar += TEXT("Fragments:\n");
-		FragmentCollection.DebugOutputDescription(Ar);
+		Ar += TEXT("Composition:\n");
+		Composition.DebugOutputDescription(Ar);
 	}
 
 	Ar += TEXT("\nInitialization Processors:\n");
diff --git a/Engine/Plugins/Runtime/MassGameplay/Source/MassSpawner/Private/MassEntityTemplateRegistry.cpp b/Engine/Plugins/Runtime/MassGameplay/Source/MassSpawner/Private/MassEntityTemplateRegistry.cpp
--- a/Engine/Plugins/Runtime/MassGameplay/Source/MassSpawner/Private/MassEntityTemplateRegistry.cpp
+++ b/Engine/Plugins/Runtime/MassGameplay/Source/MassSpawner/Private/MassEntityTemplateRegistry.cpp
@@ -132,7 +132,7 @@ bool UMassEntityTemplateRegistry::BuildTemplateImpl(const FStructToTemplateBuild
 	UWorld* World = GetWorld();
 	FMassEntityTemplateBuildContext Context(OutTemplate);
 	Builder.Execute(World, StructInstance, Context);
-	if (ensure(OutTemplate.GetFragmentsNum())) // need at least one fragment to create an Archetype
+	if (ensure(!OutTemplate.IsEmpty())) // need at least one fragment to create an Archetype
 	{
 		OutTemplate.SetUpProcessors(Context.Handlers, *this);
 		InitializeEntityTemplate(OutTemplate);
@@ -221,7 +221,7 @@ const FMassEntityTemplate* UMassEntityTemplateRegistry::BuildClassTemplate(const
 		FMassEntityTemplateBuildContext BuildContext(*EntityTemplate);
 		Builder->Execute(GetWorld(), Class, Instance, BuildContext);
 
-		if (ensureMsgf(EntityTemplate->GetFragmentsNum(), TEXT("Need at least one fragment to create an Archetype")))
+		if (ensureMsgf(!EntityTemplate->IsEmpty(), TEXT("Need at least one fragment to create an Archetype")))
 		{
 			EntityTemplate->SetUpProcessors(BuildContext.Handlers, *this);
 			InitializeEntityTemplate(*EntityTemplate);
@@ -248,7 +248,7 @@ UMassEntityTemplateRegistry::FClassToTemplateBuilderDelegate* UMassEntityTemplat
 void UMassEntityTemplateRegistry::InitializeEntityTemplate(FMassEntityTemplate& OutTemplate) const
 {
 	// expected to be ensured by the caller
-	check(OutTemplate.GetFragmentsNum());
+	check(!OutTemplate.IsEmpty());
 
 	UWorld* World = GetWorld();
 	// find or create template
diff --git a/Engine/Plugins/Runtime/MassGameplay/Source/MassSpawner/Private/MassSpawnerSubsystem.cpp b/Engine/Plugins/Runtime/MassGameplay/Source/MassSpawner/Private/MassSpawnerSubsystem.cpp
--- a/Engine/Plugins/Runtime/MassGameplay/Source/MassSpawner/Private/MassSpawnerSubsystem.cpp
+++ b/Engine/Plugins/Runtime/MassGameplay/Source/MassSpawner/Private/MassSpawnerSubsystem.cpp
@@ -136,7 +136,7 @@ bool UMassSpawnerSubsystem::SpawnEntities(const FMassEntityTemplate& EntityTempl
 	{ 
 		UE_CVLOG(OutEntities.Num() != NumberToSpawn, this, LogMassSpawner, Warning, TEXT("Tried to batch-create %d entities but created %d."), NumberToSpawn, OutEntities.Num());
 
-		const TArrayView<const FInstancedStruct> FragmentInstances = EntityTemplate.GetFragments();
+		const TArrayView<const FInstancedStruct> FragmentInstances = EntityTemplate.GetArchetypeFragmentsInitialValues().Fragments;
 		const FArchetypeChunkCollection ChunkCollection(EntityTemplate.GetArchetype(), OutEntities);
 		EntitySystem->BatchSetEntityFragmentsValues(ChunkCollection, FragmentInstances);
 
@@ -301,7 +301,7 @@ void UMassSpawnerSubsystem::DoSpawning(const FMassEntityTemplate& EntityTemplate
 	TArray<FMassEntityHandle> SpawnedEntities;
 	EntitySystem->BatchCreateEntities(EntityTemplate.GetArchetype(), Data.MaxNumber, SpawnedEntities);
 
-	const TArrayView<const FInstancedStruct> FragmentInstances = EntityTemplate.GetFragments();
+	const TArrayView<const FInstancedStruct> FragmentInstances = EntityTemplate.GetArchetypeFragmentsInitialValues().Fragments;
 	const FArchetypeChunkCollection ChunkCollection(EntityTemplate.GetArchetype(), SpawnedEntities);
 	EntitySystem->BatchSetEntityFragmentsValues(ChunkCollection, FragmentInstances);
 	
diff --git a/Engine/Plugins/Runtime/MassGameplay/Source/MassSpawner/Public/MassEntityTemplate.h b/Engine/Plugins/Runtime/MassGameplay/Source/MassSpawner/Public/MassEntityTemplate.h
--- a/Engine/Plugins/Runtime/MassGameplay/Source/MassSpawner/Public/MassEntityTemplate.h
+++ b/Engine/Plugins/Runtime/MassGameplay/Source/MassSpawner/Public/MassEntityTemplate.h
@@ -99,48 +99,98 @@ struct MASSSPAWNER_API FMassEntityTemplate
 	TConstArrayView<FObjectFragmentInitializerFunction> GetObjectFragmentInitializers() const { return ObjectInitializers; }
 	TArray<FObjectFragmentInitializerFunction>& GetMutableObjectFragmentInitializers() { return ObjectInitializers; }
 
-	bool IsValid() const { return Archetype.IsValid() && (FragmentCollection.IsEmpty() == false); }
+	bool IsValid() const { return Archetype.IsValid() && (Composition.IsEmpty() == false); }
+	bool IsEmpty() const { return Composition.IsEmpty(); }
 
 	void SetTemplateID(FMassEntityTemplateID InTemplateID) { TemplateID = InTemplateID; }
 	FMassEntityTemplateID GetTemplateID() const { return TemplateID; }
 
-	int32 GetFragmentsNum() const { return FragmentCollection.GetFragmentsNum(); }
-	TArrayView<const FInstancedStruct> GetFragments() const { return FragmentCollection.GetFragments(); }
-	FMassUniqueFragmentCollection& GetMutableFragmentCollection() { return FragmentCollection; }
-	TConstArrayView<FInstancedStruct> GetChunkFragments() const { return ChunkFragments; }
+	const FMassArchetypeCompositionDescriptor& GetCompositionDescriptor() const { return Composition; }
+	const FMassArchetypeFragmentsInitialValues& GetArchetypeFragmentsInitialValues() const { return InitialValues; }
 
-	FMassArchetypeCompositionDescriptor GetCompositionDescriptor() const { return FMassArchetypeCompositionDescriptor(FragmentCollection.GetFragmentBitSet(), TagBitSet, ChunkFragmentBitSet); }
-	FMassArchetypeFragmentsInitialValues GetArchetypeFragmentsInitialValues() const { return FMassArchetypeFragmentsInitialValues(ChunkFragments); }
+	template<typename T>
+	void AddFragment()
+	{
+		static_assert(TIsDerivedFrom<T, FMassFragment>::IsDerived, "Given struct doesn't represent a valid fragment type. Make sure to inherit from FMassFragment or one of its child-types.");
+		Composition.Fragments.Add<T>();
+	}
+
+	void AddFragment(const UScriptStruct& FragmentType)
+	{
+		checkf(FragmentType.IsChildOf(FMassFragment::StaticStruct()), TEXT("Given struct doesn't represent a valid fragment type. Make sure to inherit from FMassFragment or one of its child-types."));
+		Composition.Fragments.Add(FragmentType);
+	}
+
+	void AddFragment(FConstStructView Fragment)
+	{
+		const UScriptStruct* FragmentType = Fragment.GetScriptStruct();
+		checkf(FragmentType && FragmentType->IsChildOf(FMassFragment::StaticStruct()), TEXT("Given struct doesn't represent a valid fragment type. Make sure to inherit from FMassFragment or one of its child-types."));
+		if (!Composition.Fragments.Contains(*FragmentType))
+		{
+			Composition.Fragments.Add(*FragmentType);
+			InitialValues.Fragments.Add(Fragment);
+		}
+		else if (!InitialValues.Fragments.ContainsByPredicate(FSameTypeScriptStructPredicate(FragmentType)))
+		{
+			InitialValues.Fragments.Add(Fragment);
+		}
+	}
+
+	template<typename T>
+	T& AddFragment_GetRef()
+	{
+		static_assert(TIsDerivedFrom<T, FMassFragment>::IsDerived, "Given struct doesn't represent a valid fragment type. Make sure to inherit from FMassFragment or one of its child-types.");
+		if (!Composition.Fragments.Contains<T>())
+		{
+			Composition.Fragments.Add<T>();
+		}
+		else if (FInstancedStruct* Fragment = InitialValues.Fragments.FindByPredicate(FSameTypeScriptStructPredicate(T::StaticStruct())))
+		{
+			return Fragment->template GetMutable<T>();
+		}
+
+		// Add a default initial fragment value
+		const int32 Index = InitialValues.Fragments.Add(FInstancedStruct(T::StaticStruct()));
+		return InitialValues.Fragments[Index].template GetMutable<T>();
+	}
 
 	template<typename T>
 	void AddTag()
 	{
-		TagBitSet.Add<T>();
+		Composition.Tags.Add<T>();
 	}
 	
 	void AddTag(const UScriptStruct& TagType)
 	{
-		TagBitSet.Add(TagType);
+		Composition.Tags.Add(TagType);
 	}
 
-	const FMassTagBitSet& GetTags() const { return TagBitSet; }	
-	FMassTagBitSet& GetMutableTags() { return TagBitSet; }
+	const FMassTagBitSet& GetTags() const { return Composition.Tags; }
+	FMassTagBitSet& GetMutableTags() { return Composition.Tags; }
 
 	template<typename T>
 	void AddChunkFragment()
 	{
-		static_assert(TIsDerivedFrom<T, FMassChunkFragment>::IsDerived, "Given struct doesn't represent a valid chunk fragment type. Make sure to inherit from FLWChunkComonent or one of its child-types.");
-		ChunkFragments.AddUnique(FInstancedStruct(T::StaticStruct()));
-		ChunkFragmentBitSet.Add<T>();
+		static_assert(TIsDerivedFrom<T, FMassChunkFragment>::IsDerived, "Given struct doesn't represent a valid chunk fragment type. Make sure to inherit from FMassChunkFragment or one of its child-types.");
+		Composition.ChunkFragments.Add<T>();
 	}
 
 	template<typename T>
 	T& AddChunkFragment_GetRef()
 	{
-		static_assert(TIsDerivedFrom<T, FMassChunkFragment>::IsDerived, "Given struct doesn't represent a valid chunk fragment type. Make sure to inherit from FLWChunkComonent or one of its child-types.");
-		ChunkFragmentBitSet.Add<T>();
-		const int32 Index = ChunkFragments.AddUnique(FInstancedStruct(T::StaticStruct()));
-		return ChunkFragments[Index].template GetMutable<T>();
+		static_assert(TIsDerivedFrom<T, FMassChunkFragment>::IsDerived, "Given struct doesn't represent a valid chunk fragment type. Make sure to inherit from FMassChunkFragment or one of its child-types.");
+		if( !Composition.ChunkFragments.Contains<T>() )
+		{
+			Composition.ChunkFragments.Add<T>();
+		}
+		else if (FInstancedStruct* ChunkFragment = InitialValues.ChunkFragments.FindByPredicate(FSameTypeScriptStructPredicate(T::StaticStruct())))
+		{
+			return ChunkFragment->template GetMutable<T>();
+		}
+
+		// Add a default initial chunk fragment value
+		const int32 Index = InitialValues.ChunkFragments.Add(FInstancedStruct(T::StaticStruct()));
+		return InitialValues.ChunkFragments[Index].template GetMutable<T>();
 	}
 
 	FString DebugGetDescription(UMassEntitySubsystem* EntitySubsystem = nullptr) const;
@@ -149,27 +199,31 @@ struct MASSSPAWNER_API FMassEntityTemplate
 	template<typename T>
 	bool HasFragment() const
 	{
-		return FragmentCollection.GetFragmentBitSet().Contains(*T::StaticStruct());
+		return Composition.Fragments.Contains<T>();
 	}
 	
+	bool HasFragment(const UScriptStruct& ScriptStruct) const
+	{
+		return Composition.Fragments.Contains(ScriptStruct);
+	}
+
 	template<typename T>
 	bool HasTag() const
 	{
-		return TagBitSet.Contains(*T::StaticStruct());
+		return Composition.Tags.Contains<T>();
 	}
 
 	template<typename T>
 	bool HasChunkFragment() const
 	{
-		return ChunkFragmentBitSet.Contains(*T::StaticStruct());
+		return Composition.ChunkFragments.Contains<T>();
 	}
 
 private:
 	FArchetypeHandle Archetype;
 
-	FMassUniqueFragmentCollection FragmentCollection;
-	FMassTagBitSet TagBitSet;
-	FMassChunkFragmentBitSet ChunkFragmentBitSet;
+	FMassArchetypeCompositionDescriptor Composition;
+	FMassArchetypeFragmentsInitialValues InitialValues;
 	
 	UPROPERTY()
 	FMassRuntimePipeline InitializationPipeline;
@@ -177,9 +231,6 @@ private:
 	UPROPERTY()
 	FMassRuntimePipeline DeinitializationPipeline;
 
-	UPROPERTY()
-	TArray<FInstancedStruct> ChunkFragments;
-
 	// These functions will be called to initialize entity's UObject-based fragments
 	TArray<FObjectFragmentInitializerFunction> ObjectInitializers;
 
diff --git a/Engine/Plugins/Runtime/MassGameplay/Source/MassSpawner/Public/MassEntityTemplateRegistry.h b/Engine/Plugins/Runtime/MassGameplay/Source/MassSpawner/Public/MassEntityTemplateRegistry.h
--- a/Engine/Plugins/Runtime/MassGameplay/Source/MassSpawner/Public/MassEntityTemplateRegistry.h
+++ b/Engine/Plugins/Runtime/MassGameplay/Source/MassSpawner/Public/MassEntityTemplateRegistry.h
@@ -33,7 +33,7 @@ struct FMassEntityTemplateBuildContext
 	template<typename T>
 	T& AddFragment_GetRef()
 	{
-		return Template.GetMutableFragmentCollection().Add_GetRef<T>();
+		return Template.AddFragment_GetRef<T>();
 	}
 
 	template<typename T>
@@ -53,14 +53,12 @@ struct FMassEntityTemplateBuildContext
 	template<typename T>
 	void AddFragment()
 	{
-		static_assert(TIsDerivedFrom<T, FMassFragment>::IsDerived, "Given struct doesn't represent a valid fragment type. Make sure to inherit from FMassFragment or one of its child-types.");
-		Template.GetMutableFragmentCollection().Add(T::StaticStruct());
+		Template.AddFragment<T>();
 	}
 
 	template<typename T>
 	void AddFragmentWithDefaultInitializer()
 	{
-		static_assert(TIsDerivedFrom<T, FMassFragment>::IsDerived, "Given struct doesn't represent a valid fragment type. Make sure to inherit from FMassFragment or one of its child-types.");
 		AddFragment<T>();
 		AddDefaultInitializer<T>();
 	}
@@ -68,17 +66,16 @@ struct FMassEntityTemplateBuildContext
 	template<typename T>
 	void AddFragmentWithInitializer(const UMassProcessor& Initializer)
 	{
-		static_assert(TIsDerivedFrom<T, FMassFragment>::IsDerived, "Given struct doesn't represent a valid fragment type. Make sure to inherit from FMassFragment or one of its child-types.");
 		AddFragment<T>();
 		AddInitializer(Initializer);
 	}
 
-	void AddFragment(const FStructView InFragment)
+	void AddFragment(FConstStructView InFragment)
 	{ 
-		Template.GetMutableFragmentCollection().Add(InFragment);
+		Template.AddFragment(InFragment);
 	}
 
-	void AddFragmentWithDefaultInitializer(FStructView InFragment)
+	void AddFragmentWithDefaultInitializer(FConstStructView InFragment)
 	{
 		AddFragment(InFragment);
 		if (InFragment.GetScriptStruct())
@@ -87,7 +84,7 @@ struct FMassEntityTemplateBuildContext
 		}
 	}
 
-	void AddFragmentWithInitializer(FStructView InFragment, const UMassProcessor& Initializer)
+	void AddFragmentWithInitializer(FConstStructView InFragment, const UMassProcessor& Initializer)
 	{
 		AddFragment(InFragment);
 		AddInitializer(Initializer);
@@ -122,6 +119,11 @@ struct FMassEntityTemplateBuildContext
 		return Template.HasFragment<T>();
 	}
 	
+	bool HasFragment(const UScriptStruct& ScriptStruct) const
+	{
+		return Template.HasFragment(ScriptStruct);
+	}
+
 	template<typename T>
 	bool HasTag() const
 	{
@@ -196,7 +198,6 @@ struct FMassEntityTemplateBuildContext
 	// Template access
 	//----------------------------------------------------------------------//
 	FMassEntityTemplateID GetTemplateID() const { return Template.GetTemplateID(); }
-	FMassUniqueFragmentCollection& GetMutableFragments() { return Template.GetMutableFragmentCollection(); }
 	TArray<FMassEntityTemplate::FObjectFragmentInitializerFunction>& GetMutableObjectFragmentInitializers() { return Template.GetMutableObjectFragmentInitializers(); }
 
 	FMassObjectHandlers Handlers;
diff --git a/Engine/Plugins/Runtime/MassGameplay/Source/MassSpawner/Public/MassTranslatorRegistry.h b/Engine/Plugins/Runtime/MassGameplay/Source/MassSpawner/Public/MassTranslatorRegistry.h
--- a/Engine/Plugins/Runtime/MassGameplay/Source/MassSpawner/Public/MassTranslatorRegistry.h
+++ b/Engine/Plugins/Runtime/MassGameplay/Source/MassSpawner/Public/MassTranslatorRegistry.h
@@ -16,7 +16,6 @@ class UMassTranslator;
 class UMassObjectFragmentCreator;
 class UMassFragmentInitializer;
 class UMassFragmentDestructor;
-struct FMassUniqueFragmentCollection;
 
 /** @todo rename the class, no longer hosting Translators, just Initializers and Deinitializers/Destructors
  *  @todo brief description of the class
