Modified Simulation LOD to use SharedFragments for Config and runtime data.

#rnx

diff --git a/Engine/Plugins/Runtime/MassGameplay/Source/MassLOD/Private/MassLODTrait.cpp b/Engine/Plugins/Runtime/MassGameplay/Source/MassLOD/Private/MassLODTrait.cpp
--- a/Engine/Plugins/Runtime/MassGameplay/Source/MassLOD/Private/MassLODTrait.cpp
+++ b/Engine/Plugins/Runtime/MassGameplay/Source/MassLOD/Private/MassLODTrait.cpp
@@ -4,8 +4,7 @@
 #include "MassEntityTemplateRegistry.h"
 #include "Engine/World.h"
 #include "MassCommonFragments.h"
-#include "MassSimulationLOD.h"
-
+#include "StructUtilsTypes.h"
 
 void UMassSimulationLODTrait::BuildTemplate(FMassEntityTemplateBuildContext& BuildContext, UWorld& World) const
 {
@@ -15,4 +14,12 @@ void UMassSimulationLODTrait::BuildTemplate(FMassEntityTemplateBuildContext& Bui
 	BuildContext.AddFragmentWithDefaultInitializer<FDataFragment_MassSimulationLODInfo>();
 	BuildContext.AddTag<FMassOffLODTag>();
 	BuildContext.AddChunkFragment<FMassSimulationVariableTickChunkFragment>();
+
+	UMassEntitySubsystem* EntitySubsystem = UWorld::GetSubsystem<UMassEntitySubsystem>(&World);
+
+	uint32 ConfigHash = UE::StructUtils::GetStructCrc32(FConstStructView::Make(Config));
+	FConstSharedStruct ConfigFragment = EntitySubsystem->GetOrCreateConstSharedFragment(ConfigHash, Config);
+	BuildContext.AddConstSharedFragment(ConfigFragment);
+	FSharedStruct SharedFragment = EntitySubsystem->GetOrCreateSharedFragment<FMassSimulationLODSharedFragment>(ConfigHash, Config);
+	BuildContext.AddSharedFragment(SharedFragment);
 }
diff --git a/Engine/Plugins/Runtime/MassGameplay/Source/MassLOD/Private/MassSimulationLOD.cpp b/Engine/Plugins/Runtime/MassGameplay/Source/MassLOD/Private/MassSimulationLOD.cpp
--- a/Engine/Plugins/Runtime/MassGameplay/Source/MassLOD/Private/MassSimulationLOD.cpp
+++ b/Engine/Plugins/Runtime/MassGameplay/Source/MassLOD/Private/MassSimulationLOD.cpp
@@ -65,6 +65,15 @@ FMassSimulationLODConfig::FMassSimulationLODConfig()
 	TickRates[EMassLOD::Off] = 1.5f;
 }
 
+//-----------------------------------------------------------------------------
+// FMassSimulationLODSharedFragment
+//-----------------------------------------------------------------------------
+FMassSimulationLODSharedFragment::FMassSimulationLODSharedFragment(const FMassSimulationLODConfig& LODConfig)
+{
+	LODCalculator.Initialize(LODConfig.LODDistance, LODConfig.BufferHysteresisOnDistancePercentage / 100.0f, LODConfig.LODMaxCount);
+	LODTickRateController.Initialize(LODConfig.TickRates, LODConfig.bSpreadFirstSimulationUpdate);
+}
+
 //-----------------------------------------------------------------------------
 // UMassSimulationLODProcessor
 //-----------------------------------------------------------------------------
@@ -80,99 +89,89 @@ UMassSimulationLODProcessor::UMassSimulationLODProcessor()
 	ExecutionFlags = (int32)EProcessorExecutionFlags::All;
 	ExecutionOrder.ExecuteInGroup = UE::Mass::ProcessorGroupNames::LOD;
 	ExecutionOrder.ExecuteAfter.Add(UE::Mass::ProcessorGroupNames::LODCollector);
-	LODConfigs.AddDefaulted();
 }
 
 void UMassSimulationLODProcessor::ConfigureQueries()
 {
-	if (LODConfigs.Num() == 0)
+	EntityQuery.AddRequirement<FDataFragment_Transform>(EMassFragmentAccess::ReadOnly);
+	EntityQuery.AddRequirement<FMassLODInfoFragment>(EMassFragmentAccess::ReadOnly);
+	EntityQuery.AddRequirement<FMassSimulationLODFragment>(EMassFragmentAccess::ReadWrite);
+	EntityQuery.AddChunkRequirement<FMassSimulationVariableTickChunkFragment>(EMassFragmentAccess::ReadOnly);
+	EntityQuery.AddConstSharedRequirement<FMassSimulationLODConfig>();
+	EntityQuery.AddSharedRequirement<FMassSimulationLODSharedFragment>(EMassFragmentAccess::ReadWrite);
+
+	EntityQueryCalculateLOD = EntityQuery;
+	EntityQueryCalculateLOD.SetChunkFilter([](const FMassExecutionContext& Context)
 	{
-		UE_VLOG_UELOG(this, LogMassLOD, Warning, TEXT("Using MassSimulationLODProcessor while no LOD configs set (see MassSimulationLODProcessor.LODConfigs)."));
-	}
+		const FMassSimulationLODSharedFragment& LODSharedFragment = Context.GetSharedFragment<FMassSimulationLODSharedFragment>();
+		return LODSharedFragment.LODTickRateController.ShouldCalculateLODForChunk(Context);
+	});
 
-	for (FMassSimulationLODConfig& LODConfig : LODConfigs)
+	EntityQueryAdjustDistances = EntityQuery;
+	EntityQueryAdjustDistances.SetArchetypeFilter([](const FMassExecutionContext& Context)
 	{
-		if(LODConfig.TagFilter.GetScriptStruct())
-		{
-			LODConfig.EntityQuery.AddTagRequirement(*LODConfig.TagFilter.GetScriptStruct(), EMassFragmentPresence::All);
-		}
-		LODConfig.EntityQuery.AddRequirement<FDataFragment_Transform>(EMassFragmentAccess::ReadOnly);
-		LODConfig.EntityQuery.AddRequirement<FMassLODInfoFragment>(EMassFragmentAccess::ReadOnly);
-		LODConfig.EntityQuery.AddRequirement<FMassSimulationLODFragment>(EMassFragmentAccess::ReadWrite);
-		LODConfig.EntityQuery.AddChunkRequirement<FMassSimulationVariableTickChunkFragment>(EMassFragmentAccess::ReadOnly);
-	}
-}
-
-void UMassSimulationLODProcessor::Initialize(UObject& InOwner)
-{
-	for (FMassSimulationLODConfig& LODConfig : LODConfigs)
+		const FMassSimulationLODSharedFragment& LODSharedFragment = Context.GetSharedFragment<FMassSimulationLODSharedFragment>();
+		return LODSharedFragment.bHasAdjustedDistancesFromCount;
+	});
+	EntityQueryCalculateLOD.SetChunkFilter([](const FMassExecutionContext& Context)
 	{
-		LODConfig.LODCalculator.Initialize(LODConfig.LODDistance, LODConfig.BufferHysteresisOnDistancePercentage / 100.0f, LODConfig.LODMaxCount);
-		LODConfig.LODTickRateController.Initialize(LODConfig.TickRates, LODConfig.bSpreadFirstSimulationUpdate);
-	}
-
-	Super::Initialize(InOwner);
+		const FMassSimulationLODSharedFragment& LODSharedFragment = Context.GetSharedFragment<FMassSimulationLODSharedFragment>();
+		return LODSharedFragment.LODTickRateController.ShouldAdjustLODFromCountForChunk(Context);
+	});
 }
 
 void UMassSimulationLODProcessor::Execute(UMassEntitySubsystem& EntitySubsystem, FMassExecutionContext& Context)
 {
 	TRACE_CPUPROFILER_EVENT_SCOPE(TEXT("SimulationLOD"))
 
-	for (FMassSimulationLODConfig& LODConfig : LODConfigs)
 	{
-		CalculateLODForConfig(EntitySubsystem, Context, LODConfig);
-	}
-}
+		TRACE_CPUPROFILER_EVENT_SCOPE(TEXT("PrepareExecution"));
+		check(LODManager);
+		const TArray<FViewerInfo>& Viewers = LODManager->GetViewers();
 
-void UMassSimulationLODProcessor::CalculateLODForConfig(UMassEntitySubsystem& EntitySubsystem, FMassExecutionContext& Context, FMassSimulationLODConfig& LODConfig)
-{
-	check(LODManager);
-	const TArray<FViewerInfo>& Viewers = LODManager->GetViewers();
-	LODConfig.LODCalculator.PrepareExecution(Viewers);
+		EntitySubsystem.ForEachSharedFragment<FMassSimulationLODSharedFragment>([&Viewers](FMassSimulationLODSharedFragment& LODSharedFragment)
+		{
+			LODSharedFragment.LODCalculator.PrepareExecution(Viewers);
+		});
+	}
 
 	{
 		TRACE_CPUPROFILER_EVENT_SCOPE(TEXT("CalculateLOD"));
-
-		LODConfig.EntityQuery.ForEachEntityChunk(EntitySubsystem, Context, [&LODConfig](FMassExecutionContext& Context)
+		EntityQueryCalculateLOD.ForEachEntityChunk(EntitySubsystem, Context, [](FMassExecutionContext& Context)
 		{
-			if (!LODConfig.LODTickRateController.ShouldCalculateLODForChunk(Context))
-			{
-				return;
-			}
-
-			const TConstArrayView<FMassLODInfoFragment> ViewersInfoList = Context.GetFragmentView<FMassLODInfoFragment>();
-			const TArrayView<FMassSimulationLODFragment> SimulationLODFragments = Context.GetMutableFragmentView<FMassSimulationLODFragment>();
-			LODConfig.LODCalculator.CalculateLOD(Context, SimulationLODFragments, ViewersInfoList);
+			FMassSimulationLODSharedFragment& LODSharedFragment = Context.GetMutableSharedFragment<FMassSimulationLODSharedFragment>();
+			TConstArrayView<FMassLODInfoFragment> ViewersInfoList = Context.GetFragmentView<FMassLODInfoFragment>();
+			TArrayView<FMassSimulationLODFragment> SimulationLODFragments = Context.GetMutableFragmentView<FMassSimulationLODFragment>();
+			LODSharedFragment.LODCalculator.CalculateLOD(Context, SimulationLODFragments, ViewersInfoList);
 		});
 	}
 
 	{
 		TRACE_CPUPROFILER_EVENT_SCOPE(TEXT("AdjustDistancesAndLODFromCount"));
+		EntitySubsystem.ForEachSharedFragment<FMassSimulationLODSharedFragment>([](FMassSimulationLODSharedFragment& LODSharedFragment)
+		{
+			LODSharedFragment.bHasAdjustedDistancesFromCount = LODSharedFragment.LODCalculator.AdjustDistancesFromCount();
+		});
 
-		if (LODConfig.LODCalculator.AdjustDistancesFromCount())
+		EntityQueryAdjustDistances.ForEachEntityChunk(EntitySubsystem, Context, [](FMassExecutionContext& Context)
 		{
-			LODConfig.EntityQuery.ForEachEntityChunk(EntitySubsystem, Context, [&LODConfig](FMassExecutionContext& Context)
-			{
-				if (!LODConfig.LODTickRateController.ShouldAdjustLODFromCountForChunk(Context))
-				{
-					return;
-				}
-
-				const TConstArrayView<FMassLODInfoFragment> ViewersInfoList = Context.GetFragmentView<FMassLODInfoFragment>();
-				const TArrayView<FMassSimulationLODFragment> SimulationLODFragments = Context.GetMutableFragmentView<FMassSimulationLODFragment>();
-				LODConfig.LODCalculator.AdjustLODFromCount(Context, SimulationLODFragments, ViewersInfoList);
-			});
-		}
+			FMassSimulationLODSharedFragment& LODSharedFragment = Context.GetMutableSharedFragment<FMassSimulationLODSharedFragment>();
+			TConstArrayView<FMassLODInfoFragment> ViewersInfoList = Context.GetFragmentView<FMassLODInfoFragment>();
+			TArrayView<FMassSimulationLODFragment> SimulationLODFragments = Context.GetMutableFragmentView<FMassSimulationLODFragment>();
+			LODSharedFragment.LODCalculator.AdjustLODFromCount(Context, SimulationLODFragments, ViewersInfoList);
+		});
 	}
 
 	{
 		TRACE_CPUPROFILER_EVENT_SCOPE(TEXT("VariableTickRates"))
 		check(World);
 		const float Time = World->GetTimeSeconds();
-		LODConfig.EntityQuery.ForEachEntityChunk(EntitySubsystem, Context, [&LODConfig, Time](FMassExecutionContext& Context)
+		EntityQuery.ForEachEntityChunk(EntitySubsystem, Context, [Time](FMassExecutionContext& Context)
 		{
+			FMassSimulationLODSharedFragment& LODSharedFragment = Context.GetMutableSharedFragment<FMassSimulationLODSharedFragment>();
 			TArrayView<FMassSimulationLODFragment> SimulationLODFragments = Context.GetMutableFragmentView<FMassSimulationLODFragment>();
-			LODConfig.LODTickRateController.UpdateTickRateFromLOD(Context, SimulationLODFragments, Time);
+
+			LODSharedFragment.LODTickRateController.UpdateTickRateFromLOD(Context, SimulationLODFragments, Time);
 		});
 	}
 
@@ -180,12 +179,12 @@ void UMassSimulationLODProcessor::CalculateLODForConfig(UMassEntitySubsystem& En
 	if (UE::MassLOD::bDebugSimulationLOD)
 	{
 		TRACE_CPUPROFILER_EVENT_SCOPE(TEXT("DebugDisplayLOD"));
-
-		LODConfig.EntityQuery.ForEachEntityChunk(EntitySubsystem, Context, [this, &LODConfig](FMassExecutionContext& Context)
+		EntityQuery.ForEachEntityChunk(EntitySubsystem, Context, [this](FMassExecutionContext& Context)
 		{
-			const TConstArrayView<FDataFragment_Transform> LocationList = Context.GetFragmentView<FDataFragment_Transform>();
-			const TConstArrayView<FMassSimulationLODFragment> SimulationLODList = Context.GetFragmentView<FMassSimulationLODFragment>();
-			LODConfig.LODCalculator.DebugDisplayLOD(Context, SimulationLODList, LocationList, World);
+			FMassSimulationLODSharedFragment& LODSharedFragment = Context.GetMutableSharedFragment<FMassSimulationLODSharedFragment>();
+			TConstArrayView<FDataFragment_Transform> LocationList = Context.GetFragmentView<FDataFragment_Transform>();
+			TConstArrayView<FMassSimulationLODFragment> SimulationLODList = Context.GetFragmentView<FMassSimulationLODFragment>();
+			LODSharedFragment.LODCalculator.DebugDisplayLOD(Context, SimulationLODList, LocationList, World);
 		});
 	}
 }
\ No newline at end of file
diff --git a/Engine/Plugins/Runtime/MassGameplay/Source/MassLOD/Public/MassLODTrait.h b/Engine/Plugins/Runtime/MassGameplay/Source/MassLOD/Public/MassLODTrait.h
--- a/Engine/Plugins/Runtime/MassGameplay/Source/MassLOD/Public/MassLODTrait.h
+++ b/Engine/Plugins/Runtime/MassGameplay/Source/MassLOD/Public/MassLODTrait.h
@@ -3,6 +3,8 @@
 #pragma once 
 
 #include "MassEntityTraitBase.h"
+#include "MassSimulationLOD.h"
+
 #include "MassLODTrait.generated.h"
 
 UCLASS(meta = (DisplayName = "SimulationLOD"))
@@ -10,6 +12,9 @@ class MASSLOD_API UMassSimulationLODTrait : public UMassEntityTraitBase
 {
 	GENERATED_BODY()
 
+	UPROPERTY(Category = "Config", EditAnywhere)
+	FMassSimulationLODConfig Config;
+
 protected:
 	virtual void BuildTemplate(FMassEntityTemplateBuildContext& BuildContext, UWorld& World) const override;
 };
diff --git a/Engine/Plugins/Runtime/MassGameplay/Source/MassLOD/Public/MassSimulationLOD.h b/Engine/Plugins/Runtime/MassGameplay/Source/MassLOD/Public/MassSimulationLOD.h
--- a/Engine/Plugins/Runtime/MassGameplay/Source/MassLOD/Public/MassSimulationLOD.h
+++ b/Engine/Plugins/Runtime/MassGameplay/Source/MassLOD/Public/MassSimulationLOD.h
@@ -92,16 +92,12 @@ struct FMassSimulationVariableTickChunkFragment : public FMassVariableTickChunkF
 };
 
 USTRUCT()
-struct FMassSimulationLODConfig
+struct FMassSimulationLODConfig : public FMassSharedFragment
 {
 	GENERATED_BODY()
 
 	FMassSimulationLODConfig();
 
-	/** Component Tag that will be used to associate LOD config */
-	UPROPERTY(EditAnywhere, Category = "Mass|LOD", config, meta = (BaseStruct = "MassTag"))
-	FInstancedStruct TagFilter;
-
 	/** Distance where each LOD becomes relevant */
 	UPROPERTY(EditAnywhere, Category = "Mass|LOD", config)
 	float LODDistance[EMassLOD::Max];
@@ -121,9 +117,19 @@ struct FMassSimulationLODConfig
 	/** If true, will spread the first simulation update over TickRate period */
 	UPROPERTY(EditAnywhere, Category = "Mass|Representation", config)
 	bool bSpreadFirstSimulationUpdate = false;
+};
+
+
+USTRUCT()
+struct FMassSimulationLODSharedFragment : public FMassSharedFragment
+{
+	GENERATED_BODY()
+
+	FMassSimulationLODSharedFragment() = default;
+	FMassSimulationLODSharedFragment(const FMassSimulationLODConfig& Config);
 
 	/** Runtime data for matching the LOD config */
-	FMassEntityQuery EntityQuery;
+	bool bHasAdjustedDistancesFromCount = false;
 	TMassLODCalculator<FMassSimulationLODLogic> LODCalculator;
 	TMassLODTickRateController<FMassSimulationVariableTickChunkFragment, FMassSimulationLODLogic> LODTickRateController;
 };
@@ -139,11 +145,11 @@ public:
 protected:
 
 	virtual void ConfigureQueries() override;
-	virtual void Initialize(UObject& InOwner) override;
 	virtual void Execute(UMassEntitySubsystem& EntitySubsystem, FMassExecutionContext& Context) override;
 
 	void CalculateLODForConfig(UMassEntitySubsystem& EntitySubsystem, FMassExecutionContext& Context, FMassSimulationLODConfig& LODConfig);
 
-	UPROPERTY(EditAnywhere, Category = "Mass|LOD", config)
-	TArray<FMassSimulationLODConfig> LODConfigs;
+	FMassEntityQuery EntityQuery;
+	FMassEntityQuery EntityQueryCalculateLOD;
+	FMassEntityQuery EntityQueryAdjustDistances;
 };
