Updated LOD collector to use Camera FOV instead of a config value

#rnx

diff --git a/Engine/Plugins/Runtime/MassGameplay/Source/MassLOD/Private/MassLODCollectorProcessor.cpp b/Engine/Plugins/Runtime/MassGameplay/Source/MassLOD/Private/MassLODCollectorProcessor.cpp
--- a/Engine/Plugins/Runtime/MassGameplay/Source/MassLOD/Private/MassLODCollectorProcessor.cpp
+++ b/Engine/Plugins/Runtime/MassGameplay/Source/MassLOD/Private/MassLODCollectorProcessor.cpp
@@ -20,9 +20,9 @@ void UMassLODCollectorProcessor::Initialize(UObject& InOwner)
 
 	for (FMassCollectorLODConfig& LODConfig : LODConfigs)
 	{
-		LODConfig.RepresentationLODCollector.Initialize(LODConfig.FOVAnglesToDriveVisibility, LODConfig.BufferHysteresisOnFOVPercentage / 100.f);
+		LODConfig.RepresentationLODCollector.Initialize(LODConfig.DistanceToFrustum, LODConfig.DistanceToFrustumHysteresis, false);
 		LODConfig.SimulationLODCollector.Initialize();
-		LODConfig.CombinedLODCollector.Initialize(LODConfig.FOVAnglesToDriveVisibility, LODConfig.BufferHysteresisOnFOVPercentage / 100.f);
+		LODConfig.CombinedLODCollector.Initialize(LODConfig.DistanceToFrustum, LODConfig.DistanceToFrustumHysteresis, false);
 	}
 }
 
@@ -76,9 +76,9 @@ void UMassLODCollectorProcessor::ConfigureQueries()
 }
 
 template <typename TCollector>
-void CollectLOD(UMassEntitySubsystem& EntitySubsystem, FMassExecutionContext& Context, const TArray<FViewerInfo>& Viewers, TCollector& Collector, TArrayView<FMassEntityQuery> HighFrequencyTickingEntityQueries, FMassEntityQuery& LowFrequencyTickingEntityQuery)
+void CollectLOD(UMassEntitySubsystem& EntitySubsystem, FMassExecutionContext& Context, const TArray<FViewerInfo>& ViewersInfo, TCollector& Collector, TArrayView<FMassEntityQuery> HighFrequencyTickingEntityQueries, FMassEntityQuery& LowFrequencyTickingEntityQuery)
 {
-	Collector.PrepareExecution(Viewers);
+	Collector.PrepareExecution(ViewersInfo);
 
 	auto CollectLODInfo = [&Collector](FMassExecutionContext& Context)
 	{
diff --git a/Engine/Plugins/Runtime/MassGameplay/Source/MassLOD/Private/MassLODLogic.cpp b/Engine/Plugins/Runtime/MassGameplay/Source/MassLOD/Private/MassLODLogic.cpp
--- a/Engine/Plugins/Runtime/MassGameplay/Source/MassLOD/Private/MassLODLogic.cpp
+++ b/Engine/Plugins/Runtime/MassGameplay/Source/MassLOD/Private/MassLODLogic.cpp
@@ -4,26 +4,44 @@
 #include "Engine/World.h"
 #include "GameFramework/PlayerController.h"
 
-void FMassLODBaseLogic::CacheViewerInformation(TConstArrayView<FViewerInfo> Viewers, const bool bLocalViewersOnly)
+void FMassLODBaseLogic::CacheViewerInformation(TConstArrayView<FViewerInfo> ViewerInfos, const bool bLocalViewersOnly)
 {
-	check(Viewers.Num() <= UE::MassLOD::MaxNumOfViewers);
-	NumOfViewers = FMath::Min(Viewers.Num(), UE::MassLOD::MaxNumOfViewers);
+	if(Viewers.Num() < ViewerInfos.Num())
+	{
+		Viewers.AddDefaulted(ViewerInfos.Num() - Viewers.Num());
+	}
 
 	// Cache viewer info
-	for (int ViewerIdx = 0; ViewerIdx < NumOfViewers; ++ViewerIdx)
+	for (int ViewerIdx = 0; ViewerIdx < Viewers.Num(); ++ViewerIdx)
 	{
-		const FViewerInfo& Viewer = Viewers[ViewerIdx];
-		const APlayerController* PlayerController = Viewer.PlayerController;
-		const FMassViewerHandle ViewerHandle = Viewer.bEnabled && (!bLocalViewersOnly || Viewer.IsLocal()) ? Viewer.Handle : FMassViewerHandle();
+		const FViewerInfo& Viewer =  ViewerInfos.IsValidIndex(ViewerIdx) ? ViewerInfos[ViewerIdx] : FViewerInfo();
+		const FMassViewerHandle ViewerHandle =  Viewer.bEnabled && (!bLocalViewersOnly || Viewer.IsLocal()) ? Viewer.Handle : FMassViewerHandle();
 
 		// Check if it is the same client as before
-		bClearViewerData[ViewerIdx] = ViewersHandles[ViewerIdx] != ViewerHandle;
-		ViewersHandles[ViewerIdx] = ViewerHandle;
+		FViewerLODInfo& ViewerLOD = Viewers[ViewerIdx];
+		ViewerLOD.bClearData = Viewers[ViewerIdx].Handle != ViewerHandle;
+		ViewerLOD.Handle = ViewerHandle;
 
 		if (ViewerHandle.IsValid())
 		{
-			ViewersLocation[ViewerIdx] = Viewer.Location;
-			ViewersDirection[ViewerIdx] = Viewer.Direction;
+			ViewerLOD.Location = Viewer.Location;
+			ViewerLOD.Direction = Viewer.Rotation.Vector();
+
+			const float HalfHorizontalFOVAngle = Viewer.FOV * 0.5f;
+			const float HalfVerticalFOVAngle = FMath::RadiansToDegrees(FMath::Atan(FMath::Tan(FMath::DegreesToRadians(HalfHorizontalFOVAngle)) * Viewer.AspectRatio));
+
+			const FVector RightPlaneNormal = Viewer.Rotation.RotateVector(FRotator(0.0f, HalfHorizontalFOVAngle, 0.0f).RotateVector(FVector::RightVector));
+			const FVector LeftPlaneNormal = Viewer.Rotation.RotateVector(FRotator(0.0f, -HalfHorizontalFOVAngle, 0.0f).RotateVector(FVector::LeftVector));
+			const FVector TopPlaneNormal = Viewer.Rotation.RotateVector(FRotator(HalfVerticalFOVAngle, 0.0f, 0.0f).RotateVector(FVector::UpVector));
+			const FVector BottomPlaneNormal = Viewer.Rotation.RotateVector(FRotator(-HalfVerticalFOVAngle, 0.0f, 0.0f).RotateVector(FVector::DownVector));
+			
+			TArray<FPlane, TInlineAllocator<6>> Planes;
+			Planes.Emplace(ViewerLOD.Location, RightPlaneNormal);
+			Planes.Emplace(ViewerLOD.Location, LeftPlaneNormal);
+			Planes.Emplace(ViewerLOD.Location, TopPlaneNormal);
+			Planes.Emplace(ViewerLOD.Location, BottomPlaneNormal);
+
+			ViewerLOD.Frustum = FConvexVolume(Planes);
 		}
 	}
 }
\ No newline at end of file
diff --git a/Engine/Plugins/Runtime/MassGameplay/Source/MassLOD/Private/MassLODManager.cpp b/Engine/Plugins/Runtime/MassGameplay/Source/MassLOD/Private/MassLODManager.cpp
--- a/Engine/Plugins/Runtime/MassGameplay/Source/MassLOD/Private/MassLODManager.cpp
+++ b/Engine/Plugins/Runtime/MassGameplay/Source/MassLOD/Private/MassLODManager.cpp
@@ -219,7 +219,16 @@ void UMassLODManager::SynchronizeViewers()
 			FRotator PlayerCameraRotation(FRotator::ZeroRotator);
 			ViewerInfo.PlayerController->GetPlayerViewPoint(PlayerCameraLocation, PlayerCameraRotation);
 			ViewerInfo.Location = PlayerCameraLocation;
-			ViewerInfo.Direction = PlayerCameraRotation.Vector();
+			ViewerInfo.Rotation = PlayerCameraRotation;
+
+			// Try to fetch a more precise FOV
+			if(ViewerInfo.PlayerController->PlayerCameraManager)
+			{
+				ViewerInfo.FOV = ViewerInfo.PlayerController->PlayerCameraManager->GetFOVAngle();
+
+				// @todo need to find a way to retrieve aspect ratio, this does not seems to work
+				//ViewerInfo.AspectRatio = MinViewInfo.AspectRatio;
+			}
 		}
 		else
 		{
@@ -230,7 +239,7 @@ void UMassLODManager::SynchronizeViewers()
 			if (StreamingSource)
 			{
 				ViewerInfo.Location = StreamingSource->Location;
-				ViewerInfo.Direction = StreamingSource->Rotation.Vector();
+				ViewerInfo.Rotation = StreamingSource->Rotation;
 			}
 		}
 	}
diff --git a/Engine/Plugins/Runtime/MassGameplay/Source/MassLOD/Public/MassLODCalculator.h b/Engine/Plugins/Runtime/MassGameplay/Source/MassLOD/Public/MassLODCalculator.h
--- a/Engine/Plugins/Runtime/MassGameplay/Source/MassLOD/Public/MassLODCalculator.h
+++ b/Engine/Plugins/Runtime/MassGameplay/Source/MassLOD/Public/MassLODCalculator.h
@@ -207,16 +207,16 @@ void TMassLODCalculator<FLODLogic>::Initialize(const float InBaseLODDistance[EMa
 }
 
 template <typename FLODLogic>
-void TMassLODCalculator<FLODLogic>::PrepareExecution(TConstArrayView<FViewerInfo> Viewers)
+void TMassLODCalculator<FLODLogic>::PrepareExecution(TConstArrayView<FViewerInfo> ViewersInfo)
 {
-	CacheViewerInformation(Viewers, FLODLogic::bLocalViewersOnly);
+	CacheViewerInformation(ViewersInfo, FLODLogic::bLocalViewersOnly);
 
 	if (FLODLogic::bMaximizeCountPerViewer)
 	{
-		for (int ViewerIdx = 0; ViewerIdx < NumOfViewers; ++ViewerIdx)
+		for (int ViewerIdx = 0; ViewerIdx < Viewers.Num(); ++ViewerIdx)
 		{
 			// Reset viewer data
-			if (ViewersHandles[ViewerIdx].IsValid())
+			if (Viewers[ViewerIdx].Handle.IsValid())
 			{
 				RuntimeDataPerViewer[ViewerIdx].Reset(BaseLODDistance, VisibleLODDistance);
 			}
@@ -285,14 +285,15 @@ void TMassLODCalculator<FLODLogic>::CalculateLOD(FMassExecutionContext& Context,
 		bool bIsVisibleByAViewer = false;
 		bool bIsInAVisibleRange = false;
 
-		for (int ViewerIdx = 0; ViewerIdx < NumOfViewers; ++ViewerIdx)
+		for (int ViewerIdx = 0; ViewerIdx < Viewers.Num(); ++ViewerIdx)
 		{
-			if (bClearViewerData[ViewerIdx])
+			const FViewerLODInfo& Viewer = Viewers[ViewerIdx];
+			if (Viewer.bClearData)
 			{
 				SetLODPerViewer<FLODLogic::bStoreLODPerViewer>(EntityLOD, ViewerIdx, EMassLOD::Max);
 				SetPrevLODPerViewer<FLODLogic::bStoreLODPerViewer>(EntityLOD, ViewerIdx, EMassLOD::Max);
 			}
-			if (ViewersHandles[ViewerIdx].IsValid())
+			if (Viewer.Handle.IsValid())
 			{
 				const bool bIsVisibleByViewer = GetbIsVisibleByViewer<FLODLogic::bDoVisibilityLogic>(EntityViewersInfo, ViewerIdx, false);
 				bIsVisibleByAViewer |= bIsVisibleByViewer;
@@ -456,9 +457,9 @@ bool TMassLODCalculator<FLODLogic>::AdjustDistancesFromCount()
 	bool bDistanceAdjusted = false;
 	if (FLODLogic::bMaximizeCountPerViewer)
 	{
-		for (int ViewerIdx = 0; ViewerIdx < NumOfViewers; ++ViewerIdx)
+		for (int ViewerIdx = 0; ViewerIdx < Viewers.Num(); ++ViewerIdx)
 		{
-			if (!ViewersHandles[ViewerIdx].IsValid())
+			if (!Viewers[ViewerIdx].Handle.IsValid())
 			{
 				continue;
 			}
@@ -484,9 +485,9 @@ void TMassLODCalculator<FLODLogic>::AdjustLODFromCount(FMassExecutionContext& Co
 		EMassLOD::Type HighestViewerLOD = EMassLOD::Off;
 		if (FLODLogic::bMaximizeCountPerViewer)
 		{
-			for (int ViewerIdx = 0; ViewerIdx < NumOfViewers; ++ViewerIdx)
+			for (int ViewerIdx = 0; ViewerIdx < Viewers.Num(); ++ViewerIdx)
 			{
-				if (!ViewersHandles[ViewerIdx].IsValid())
+				if (!Viewers[ViewerIdx].Handle.IsValid())
 				{
 					continue;
 				}
diff --git a/Engine/Plugins/Runtime/MassGameplay/Source/MassLOD/Public/MassLODCollector.h b/Engine/Plugins/Runtime/MassGameplay/Source/MassLOD/Public/MassLODCollector.h
--- a/Engine/Plugins/Runtime/MassGameplay/Source/MassLOD/Public/MassLODCollector.h
+++ b/Engine/Plugins/Runtime/MassGameplay/Source/MassLOD/Public/MassLODCollector.h
@@ -7,7 +7,7 @@
 
 /**
  * This is the expected member variables for the TArrayView<FMassLODSourceInfo> when calling TMassLODCollector methods
- */
+ *
 struct FMassLODSourceInfo
 {
 	// Distance to viewer (Always needed)
@@ -16,6 +16,7 @@ struct FMassLODSourceInfo
 	// Visibility information (Only when FLODLogic::bDoVisibilityLogic is enabled)
 	TStaticArray<bool, UE::MassLOD::MaxNumOfViewers> bIsVisibleByViewer;
 };
+*/
 
 /**
  * Helper struct to collect needed information on each agent that will be needed later for LOD calculation
@@ -27,10 +28,10 @@ struct TMassLODCollector : public FMassLODBaseLogic
 {
 	/**
 	 * Initializes the LOD collector, needed to be called once at initialization time
-	 * @Param InFOVAngleToDriveVisibility the FOV angle that will be use for visibility check in degrees (Only when FLODLogic::bDoVisibilityLogic is enabled)
-	 * @Param InBufferHysteresisOnFOVRatio the FOV angle hysteresis that will be use for visibility check (Only when FLODLogic::bDoVisibilityLogic is enabled)
+	 * @Param InDistanceToFrustum the extra distance to frustrum to use to know if the entity is visible (Only when FLODLogic::bDoVisibilityLogic is enabled)
+	 * @Param InDistanceToFrustumHysteresis the hyteresis on the distance to frustrum to use when entity was previously visible to become not visible (Only when FLODLogic::bDoVisibilityLogic is enabled)
 	 */
-	void Initialize(const float InFOVAngleToDriveVisibility = -1.0f, const float InBufferHysteresisOnFOVRatio = -1.0f);
+	void Initialize(const float InDistanceToFrustum = 0.0f, const float InDistanceToFrustumHysteresis = 0.0f, bool bTemp = false);
 
 	/**
 	 * Prepares execution for the current frame, needed to be called before every execution
@@ -41,53 +42,52 @@ struct TMassLODCollector : public FMassLODBaseLogic
 	/**
 	 * Collects the information for LOD calculation, called for each entity chunks
 	 * @Param Context of the chunk execution
-	 * @Param LocationList is the fragment transforms of the entities
+	 * @Param TranformList is the fragment transforms of the entities
 	 * @Param ViewersInfoList is the fragment where to store source information for LOD calculation
 	 */
 	template< typename FMassTransform, typename FMassLODSourceInfo >
-	void CollectLODInfo(FMassExecutionContext& Context, TConstArrayView<FMassTransform> LocationList, TArrayView<FMassLODSourceInfo> ViewersInfoList);
+	void CollectLODInfo(FMassExecutionContext& Context, TConstArrayView<FMassTransform> TranformList, TArrayView<FMassLODSourceInfo> ViewersInfoList);
 
 protected:
 
-	float FOVAngleToDriveVisibilityRad = 0.0f;
-	float CosFOVAngleToDriveVisibility = 0.0f;
-	float BufferHysteresisOnFOVRatio = 0.1f;
+	float DistanceToFrustum = 0.0f;
+	float DistanceToFrustumHysteresis = 0.0f;
 };
 
 template <typename FLODLogic>
-void TMassLODCollector<FLODLogic>::Initialize(const float InFOVAngleToDriveVisibility/*= -1.0f*/, const float InBufferHysteresisOnFOVRatio /*= -1.0f*/)
+void TMassLODCollector<FLODLogic>::Initialize(const float InDistanceToFrustum /*= 0.0f*/, const float InDistanceToFrustumHysteresis /*= 0.0f*/, bool bTemp /*= false*/)
 {
-	checkf(!FLODLogic::bDoVisibilityLogic || (InFOVAngleToDriveVisibility > 0.0f && BufferHysteresisOnFOVRatio > 0.0f), TEXT("Expecting FOVAngleToDriveVisibility and BufferHysteresisOnFOVRatio when bDoVisibilityLogic is true"));
-	FOVAngleToDriveVisibilityRad = FMath::DegreesToRadians(InFOVAngleToDriveVisibility);
-	CosFOVAngleToDriveVisibility = FMath::Cos(FOVAngleToDriveVisibilityRad);
-	BufferHysteresisOnFOVRatio = InBufferHysteresisOnFOVRatio;
+	DistanceToFrustum = InDistanceToFrustum;
+	DistanceToFrustumHysteresis = InDistanceToFrustumHysteresis;
 }
 
 template <typename FLODLogic>
-void TMassLODCollector<FLODLogic>::PrepareExecution(TConstArrayView<FViewerInfo> Viewers)
+void TMassLODCollector<FLODLogic>::PrepareExecution(TConstArrayView<FViewerInfo> ViewersInfo)
 {
-	CacheViewerInformation(Viewers, FLODLogic::bLocalViewersOnly);
+	CacheViewerInformation(ViewersInfo, FLODLogic::bLocalViewersOnly);
 }
 
 template <typename FLODLogic>
 template< typename FMassTransform, typename FMassLODSourceInfo >
-void TMassLODCollector<FLODLogic>::CollectLODInfo(FMassExecutionContext& Context, TConstArrayView<FMassTransform> LocationList, TArrayView<FMassLODSourceInfo> ViewersInfoList)
+void TMassLODCollector<FLODLogic>::CollectLODInfo(FMassExecutionContext& Context, TConstArrayView<FMassTransform> TranformList, TArrayView<FMassLODSourceInfo> ViewersInfoList)
 {
 	const int32 NumEntities = Context.GetNumEntities();
 	for (int EntityIdx = 0; EntityIdx < NumEntities; EntityIdx++)
 	{
-		const FMassTransform& EntityLocation = LocationList[EntityIdx];
+		const FMassTransform& EntityTransform = TranformList[EntityIdx];
 		FMassLODSourceInfo& EntityViewerInfo = ViewersInfoList[EntityIdx];
-		for (int ViewerIdx = 0; ViewerIdx < NumOfViewers; ++ViewerIdx)
+		for (int ViewerIdx = 0; ViewerIdx < Viewers.Num(); ++ViewerIdx)
 		{
-			if (bClearViewerData[ViewerIdx])
+			const FViewerLODInfo& Viewer = Viewers[ViewerIdx];
+			if (Viewer.bClearData)
 			{
 				EntityViewerInfo.DistanceToViewerSq[ViewerIdx] = FLT_MAX;
 				SetbIsVisibleByViewer<FLODLogic::bDoVisibilityLogic>(EntityViewerInfo, ViewerIdx, false);
 			}
-			if (ViewersHandles[ViewerIdx].IsValid())
+			if (Viewer.Handle.IsValid())
 			{
-				const FVector ViewerToEntity = EntityLocation.GetTransform().GetLocation() - ViewersLocation[ViewerIdx];
+				const FVector& EntityLocation = EntityTransform.GetTransform().GetLocation();
+				const FVector ViewerToEntity = EntityLocation - Viewer.Location;
 				EntityViewerInfo.DistanceToViewerSq[ViewerIdx] = ViewerToEntity.SizeSquared();
 
 				if (FLODLogic::bDoVisibilityLogic)
@@ -101,19 +101,7 @@ void TMassLODCollector<FLODLogic>::CollectLODInfo(FMassExecutionContext& Context
 					}
 					else
 					{
-						const FVector ViewerToEntityNorm = ViewerToEntity * FMath::InvSqrt(EntityViewerInfo.DistanceToViewerSq[ViewerIdx]);
-						const float Dot = ViewersDirection[ViewerIdx] | ViewerToEntityNorm;
-
-						// If we were visible and are now exiting defined FOV angle, consider a buffer hysteresis to give some room before toggling off visibility to prevent oscillating LOD states
-						if (bWasVisibleByViewer && Dot < CosFOVAngleToDriveVisibility)
-						{
-							const float FOVToCos = FMath::Cos(FMath::Clamp(FOVAngleToDriveVisibilityRad * (1.f + BufferHysteresisOnFOVRatio), 0.f, PI));
-							bIsVisibleByViewer = Dot > FOVToCos;
-						}
-						else
-						{
-							bIsVisibleByViewer = Dot >= CosFOVAngleToDriveVisibility;
-						}
+						bIsVisibleByViewer = Viewer.Frustum.IntersectSphere(EntityLocation, bWasVisibleByViewer ? DistanceToFrustum + DistanceToFrustumHysteresis : DistanceToFrustum );
 					}
 
 					SetbIsVisibleByViewer<FLODLogic::bDoVisibilityLogic>(EntityViewerInfo, ViewerIdx, bIsVisibleByViewer);
diff --git a/Engine/Plugins/Runtime/MassGameplay/Source/MassLOD/Public/MassLODCollectorProcessor.h b/Engine/Plugins/Runtime/MassGameplay/Source/MassLOD/Public/MassLODCollectorProcessor.h
--- a/Engine/Plugins/Runtime/MassGameplay/Source/MassLOD/Public/MassLODCollectorProcessor.h
+++ b/Engine/Plugins/Runtime/MassGameplay/Source/MassLOD/Public/MassLODCollectorProcessor.h
@@ -18,11 +18,12 @@ struct FMassCollectorLODConfig
 	UPROPERTY(EditAnywhere, Category = "Mass|LOD", config, meta = (BaseStruct = "MassTag"))
 	FInstancedStruct TagFilter;
 
-	/** Set FOV Angles to define the limits between using Visible or Base LOD Distances */
-	UPROPERTY(EditAnywhere, Category = "Mass|LOD", meta = (ClampMin = "0.0", ClampMax = "180.0", UIMin = "0.0", UIMax = "180.0"), config)
-	float FOVAnglesToDriveVisibility = 45.0f;
+	/** How far away from frustrum does this entities are considered visible */
 	UPROPERTY(EditAnywhere, Category = "Mass|LOD", meta = (ClampMin = "0.0", UIMin = "0.0"), config)
-	float BufferHysteresisOnFOVPercentage = 10.0f;
+	float DistanceToFrustum = 0.0f;
+	// Once visible how much further than DistanceToFrustum does the entities need to be before being cull again
+	UPROPERTY(EditAnywhere, Category = "Mass|LOD", meta = (ClampMin = "0.0", UIMin = "0.0"), config)
+	float DistanceToFrustumHysteresis = 0.0f;
 
 	/** Runtime data for matching the LOD config */
 	TMassLODCollector<FMassRepresentationLODLogic> RepresentationLODCollector;
diff --git a/Engine/Plugins/Runtime/MassGameplay/Source/MassLOD/Public/MassLODLogic.h b/Engine/Plugins/Runtime/MassGameplay/Source/MassLOD/Public/MassLODLogic.h
--- a/Engine/Plugins/Runtime/MassGameplay/Source/MassLOD/Public/MassLODLogic.h
+++ b/Engine/Plugins/Runtime/MassGameplay/Source/MassLOD/Public/MassLODLogic.h
@@ -4,6 +4,7 @@
 
 #include "MassLODManager.h"
 #include "Containers/StaticArray.h"
+#include "ConvexVolume.h"
 
 
 #define DECLARE_CONDITIONAL_MEMBER_ACCESSORS( Condition, MemberType, MemberName ) \
@@ -95,13 +96,29 @@ struct FMassLODFragment
 };
 */
 
+struct FViewerLODInfo
+{
+	/* Boolean indicating the viewer data needs to be cleared */
+	bool bClearData = false;
+
+	/** The handle to the viewer */
+	FMassViewerHandle Handle;
+
+	/** Viewer location and looking direction */
+	FVector Location;
+	FVector Direction;
+
+	/** Viewer frustum (will not include near and far planes) */
+	FConvexVolume Frustum;
+};
+
 /**
  * Base struct for the LOD calculation helpers
  */
 struct MASSLOD_API FMassLODBaseLogic
 {
 protected:
-	void CacheViewerInformation(TConstArrayView<FViewerInfo> Viewers, const bool bLocalViewersOnly);
+	void CacheViewerInformation(TConstArrayView<FViewerInfo> ViewerInfos, const bool bLocalViewersOnly);
 
 	DECLARE_CONDITIONAL_MEMBER_ACCESSORS(Condition, bool, bIsVisibleByAViewer);
 	DECLARE_CONDITIONAL_MEMBER_ACCESSORS(Condition, bool, bWasVisibleByAViewer);
@@ -114,16 +131,5 @@ protected:
 	DECLARE_CONDITIONAL_MEMBER_ARRAY_ACCESSORS(Condition, EMassLOD::Type, LODPerViewer);
 	DECLARE_CONDITIONAL_MEMBER_ARRAY_ACCESSORS(Condition, EMassLOD::Type, PrevLODPerViewer);
 
-	/* Number of viewers handled */
-	int32 NumOfViewers = 0;
-
-	/* Boolean indicating the viewer data needs to be cleared */
-	TStaticArray<bool, UE::MassLOD::MaxNumOfViewers> bClearViewerData;
-
-	/** The handle to each viewers */
-	TStaticArray<FMassViewerHandle, UE::MassLOD::MaxNumOfViewers> ViewersHandles;
-
-	/** Viewer location and looking direction */
-	TStaticArray<FVector, UE::MassLOD::MaxNumOfViewers> ViewersLocation;
-	TStaticArray<FVector, UE::MassLOD::MaxNumOfViewers> ViewersDirection;
+	TArray<FViewerLODInfo> Viewers;
 };
\ No newline at end of file
diff --git a/Engine/Plugins/Runtime/MassGameplay/Source/MassLOD/Public/MassLODManager.h b/Engine/Plugins/Runtime/MassGameplay/Source/MassLOD/Public/MassLODManager.h
--- a/Engine/Plugins/Runtime/MassGameplay/Source/MassLOD/Public/MassLODManager.h
+++ b/Engine/Plugins/Runtime/MassGameplay/Source/MassLOD/Public/MassLODManager.h
@@ -56,7 +56,9 @@ struct FViewerInfo
 	uint32 HashValue = 0;
 
 	FVector Location;
-	FVector Direction;
+	FRotator Rotation;
+	float FOV = 90.0f;
+	float AspectRatio = 16.0f / 9.0f;
 
 	bool bEnabled = true;
 
