Modified MassSpawner.SpawnAtLocations to enable using spawner's location as as fallback spawn location.

diff --git a/Engine/Plugins/Runtime/MassGameplay/Source/MassSpawner/Private/MassSpawner.cpp b/Engine/Plugins/Runtime/MassGameplay/Source/MassSpawner/Private/MassSpawner.cpp
--- a/Engine/Plugins/Runtime/MassGameplay/Source/MassSpawner/Private/MassSpawner.cpp
+++ b/Engine/Plugins/Runtime/MassGameplay/Source/MassSpawner/Private/MassSpawner.cpp
@@ -274,10 +274,19 @@ void AMassSpawner::DoSpawning()
 	const int32 SpawnCount = GetSpawnCount();
 	if (SpawnCount > 0)
 	{
+		// no spawn point generators configured. Let user know and fall back to the spawner's location
+		if (SpawnPointsGenerators.Num() == 0)
+		{
+			UE_VLOG_UELOG(this, LogMassSpawner, Warning, TEXT("No Spawn Points Generators configured. Falling back to the MassSpawner\'s location."));
+			const FVector SpawnerLocation = GetActorLocation();
+			SpawnAtLocations(MakeArrayView(&SpawnerLocation, 1));
+			return;
+		}
+
 		float TotalProportion = 0.0f;
 		for (FMassSpawnPointGenerator& SpawnPointsGenerator : SpawnPointsGenerators)
 		{
-		    if (SpawnPointsGenerator.GeneratorInstance)
+			if (SpawnPointsGenerator.GeneratorInstance)
 			{
 				SpawnPointsGenerator.bPointsGenerated = false;
 				TotalProportion += SpawnPointsGenerator.Proportion;
@@ -382,7 +391,7 @@ int32 AMassSpawner::GetSpawnCount() const
 	return int32(FinalSpawningCountScale * Count);
 }
 
-void AMassSpawner::SpawnAtLocations(const TArray<FVector>& Locations)
+void AMassSpawner::SpawnAtLocations(TConstArrayView<FVector> Locations)
 {
 	UMassSpawnerSubsystem* SpawnerSystem = UWorld::GetSubsystem<UMassSpawnerSubsystem>(GetWorld());
 	if (SpawnerSystem == nullptr)
@@ -407,7 +416,7 @@ void AMassSpawner::SpawnAtLocations(const TArray<FVector>& Locations)
 		if (FVisualLogEntry* LogEntry = FVisualLogger::Get().GetLastEntryForObject(this))
 		{
 			FVisualLogShapeElement Element(TEXT(""), FColor::Orange, /*Thickness*/20, LogMassSpawner.GetCategoryName());
-			Element.Points += Locations;
+			Element.Points.Append(Locations.GetData(), Locations.Num());
 			Element.Type = EVisualLoggerShapeElement::SinglePoint;
 			Element.Verbosity = ELogVerbosity::Display;
 			LogEntry->AddElement(Element);
diff --git a/Engine/Plugins/Runtime/MassGameplay/Source/MassSpawner/Public/MassSpawner.h b/Engine/Plugins/Runtime/MassGameplay/Source/MassSpawner/Public/MassSpawner.h
--- a/Engine/Plugins/Runtime/MassGameplay/Source/MassSpawner/Public/MassSpawner.h
+++ b/Engine/Plugins/Runtime/MassGameplay/Source/MassSpawner/Public/MassSpawner.h
@@ -82,7 +82,7 @@ public:
 protected:
 	void OnAdjustTickSchematics(UWorld* World, TArray<TSoftObjectPtr<UMassSchematic>>& InOutTickSchematics);
 	void OnPostWorldInit(UWorld* World, const UWorld::InitializationValues);
-	void SpawnAtLocations(const TArray<FVector>& Locations);
+	void SpawnAtLocations(TConstArrayView<FVector> Locations);
 	void OnSpawnPointGenerationFinished(const TArray<FVector>& Locations, FMassSpawnPointGenerator* FinishedGenerator);
 
 	int32 GetSpawnCount() const;
