Made UMassDebugVisualizationTrait add basic debuggability fragments regardless of whether debug mesh has been set

diff --git a/Engine/Plugins/Runtime/MassGameplay/Source/MassGameplayDebug/Private/MassDebugVisualizationTrait.cpp b/Engine/Plugins/Runtime/MassGameplay/Source/MassGameplayDebug/Private/MassDebugVisualizationTrait.cpp
--- a/Engine/Plugins/Runtime/MassGameplay/Source/MassGameplayDebug/Private/MassDebugVisualizationTrait.cpp
+++ b/Engine/Plugins/Runtime/MassGameplay/Source/MassGameplayDebug/Private/MassDebugVisualizationTrait.cpp
@@ -14,8 +14,6 @@ void UMassDebugVisualizationTrait::BuildTemplate(FMassEntityTemplateBuildContext
 #else
 	const UStaticMesh* const DebugMesh = nullptr;
 #endif
-
-	BuildContext.AddTag<FMassDebuggableTag>();
 	
 	if (DebugMesh)
 	{
@@ -37,16 +35,15 @@ void UMassDebugVisualizationTrait::BuildTemplate(FMassEntityTemplateBuildContext
 	}
 	// add fragments needed whenever we have debugging capabilities
 #if !(UE_BUILD_SHIPPING || UE_BUILD_TEST)
-	else // Data.DebugShape.Mesh == null
-	{
+	BuildContext.AddTag<FMassDebuggableTag>();
 #if WITH_EDITORONLY_DATA
-		BuildContext.AddFragmentWithDefaultInitializer_GetRef<FDataFragment_DebugVis>().Shape = DebugShape.WireShape;
+	BuildContext.AddFragmentWithDefaultInitializer_GetRef<FDataFragment_DebugVis>().Shape = DebugShape.WireShape;
 #else
-		// DebugShape unavailable, will used default instead
-		BuildContext.AddFragmentWithDefaultInitializer<FDataFragment_DebugVis>();
+	// DebugShape unavailable, will used default instead
+	BuildContext.AddFragmentWithDefaultInitializer<FDataFragment_DebugVis>();
 #endif // WITH_EDITORONLY_DATA
-		BuildContext.AddFragmentWithDefaultInitializer<FDataFragment_AgentRadius>();
-	}
+	BuildContext.AddFragmentWithDefaultInitializer<FDataFragment_AgentRadius>();
+
 	BuildContext.AddFragmentWithDefaultInitializer<FDataFragment_Transform>();
 #endif // if !(UE_BUILD_SHIPPING || UE_BUILD_TEST)
 
