[SmartObjects] integration of (18039396, 18230600, 18263205, 18300133, 18301080) from UE5

diff --git a/Engine/Plugins/Runtime/MassGameplay/Source/MassSmartObjects/Private/MassSmartObjectsModule.cpp b/Engine/Plugins/Runtime/MassGameplay/Source/MassSmartObjects/Private/MassSmartObjectsModule.cpp
--- a/Engine/Plugins/Runtime/MassGameplay/Source/MassSmartObjects/Private/MassSmartObjectsModule.cpp
+++ b/Engine/Plugins/Runtime/MassGameplay/Source/MassSmartObjects/Private/MassSmartObjectsModule.cpp
@@ -21,7 +21,6 @@ void FMassSmartObjectsModule::StartupModule()
 	// This code will execute after your module is loaded into memory (but after global variables are initialized, of course.)
 
 	TArray<FCoreRedirect> Redirects;
-	Redirects.Emplace(ECoreRedirectFlags::Type_Struct, TEXT("DataFragment_SmartObjectUser"), TEXT("MassSmartObjectUserFragment"));
 	Redirects.Emplace(ECoreRedirectFlags::Type_Class, TEXT("SmartObjectMassBehaviorConfig"), TEXT("SmartObjectMassBehaviorDefinition"));
 
 	FCoreRedirects::AddRedirectList(Redirects, TEXT("MassSmartObject"));
