Change LODcollector to collect distance to frustum instead of calculating visibility and move the visibility calculation in the  CalculateLOD.

#rnx

diff --git a/Engine/Plugins/AI/MassCrowd/Source/MassCrowd/Private/MassCrowdReplicationProcessor.cpp b/Engine/Plugins/AI/MassCrowd/Source/MassCrowd/Private/MassCrowdReplicationProcessor.cpp
--- a/Engine/Plugins/AI/MassCrowd/Source/MassCrowd/Private/MassCrowdReplicationProcessor.cpp
+++ b/Engine/Plugins/AI/MassCrowd/Source/MassCrowd/Private/MassCrowdReplicationProcessor.cpp
@@ -61,7 +61,7 @@ void UMassCrowdReplicationProcessor::Execute(UMassEntitySubsystem& EntitySubsyst
 		EntityQuery.ForEachEntityChunk(EntitySubsystem, Context, [this](FMassExecutionContext& Context)
 			{
 				const TConstArrayView<FDataFragment_Transform> TransformList = Context.GetFragmentView<FDataFragment_Transform>();
-				const TConstArrayView<FMassReplicationViewerLODFragment> ViewerLODList = Context.GetFragmentView<FMassReplicationViewerLODFragment>();
+				const TConstArrayView<FMassReplicationLODFragment> ViewerLODList = Context.GetFragmentView<FMassReplicationLODFragment>();
 				LODCalculator.DebugDisplayLOD(Context, ViewerLODList, TransformList, World);
 			});
 	}
diff --git a/Engine/Plugins/AI/MassCrowd/Source/MassCrowd/Private/MassCrowdServerRepresentationLODProcessor.cpp b/Engine/Plugins/AI/MassCrowd/Source/MassCrowd/Private/MassCrowdServerRepresentationLODProcessor.cpp
--- a/Engine/Plugins/AI/MassCrowd/Source/MassCrowd/Private/MassCrowdServerRepresentationLODProcessor.cpp
+++ b/Engine/Plugins/AI/MassCrowd/Source/MassCrowd/Private/MassCrowdServerRepresentationLODProcessor.cpp
@@ -32,7 +32,7 @@ void UMassCrowdServerRepresentationLODProcessor::ConfigureQueries()
 {
 	EntityQuery.AddTagRequirement<FTagFragment_MassCrowd>(EMassFragmentPresence::All);
 	EntityQuery.AddRequirement<FDataFragment_Transform>(EMassFragmentAccess::ReadOnly);
-	EntityQuery.AddRequirement<FMassLODInfoFragment>(EMassFragmentAccess::ReadOnly);
+	EntityQuery.AddRequirement<FMassViewerInfoFragment>(EMassFragmentAccess::ReadOnly);
 	EntityQuery.AddRequirement<FMassRepresentationLODFragment>(EMassFragmentAccess::ReadWrite);
 }
 
@@ -56,7 +56,7 @@ void UMassCrowdServerRepresentationLODProcessor::Execute(UMassEntitySubsystem& E
 		
 		EntityQuery.ForEachEntityChunk(EntitySubsystem, Context, [this](FMassExecutionContext& Context)
 		{
-			const TConstArrayView<FMassLODInfoFragment> ViewersInfoList = Context.GetFragmentView<FMassLODInfoFragment>();
+			const TConstArrayView<FMassViewerInfoFragment> ViewersInfoList = Context.GetFragmentView<FMassViewerInfoFragment>();
 			const TArrayView<FMassRepresentationLODFragment> RepresentationLODFragments = Context.GetMutableFragmentView<FMassRepresentationLODFragment>();
 			LODCalculator.CalculateLOD(Context, RepresentationLODFragments, ViewersInfoList);
 		});
@@ -69,7 +69,7 @@ void UMassCrowdServerRepresentationLODProcessor::Execute(UMassEntitySubsystem& E
 		{
 			EntityQuery.ForEachEntityChunk(EntitySubsystem, Context, [this](FMassExecutionContext& Context)
 			{
-				const TConstArrayView<FMassLODInfoFragment> ViewersInfoList = Context.GetFragmentView<FMassLODInfoFragment>();
+				const TConstArrayView<FMassViewerInfoFragment> ViewersInfoList = Context.GetFragmentView<FMassViewerInfoFragment>();
 				const TArrayView<FMassRepresentationLODFragment> RepresentationLODFragments = Context.GetMutableFragmentView<FMassRepresentationLODFragment>();
 				LODCalculator.AdjustLODFromCount(Context, RepresentationLODFragments, ViewersInfoList);
 			});
diff --git a/Engine/Plugins/AI/MassCrowd/Source/MassCrowd/Private/MassCrowdVisualizationLODProcessor.cpp b/Engine/Plugins/AI/MassCrowd/Source/MassCrowd/Private/MassCrowdVisualizationLODProcessor.cpp
--- a/Engine/Plugins/AI/MassCrowd/Source/MassCrowd/Private/MassCrowdVisualizationLODProcessor.cpp
+++ b/Engine/Plugins/AI/MassCrowd/Source/MassCrowd/Private/MassCrowdVisualizationLODProcessor.cpp
@@ -59,7 +59,7 @@ void UMassCrowdVisualizationLODProcessor::ConfigureQueries()
 	// Do not call super as we do have our own LODInfo fragment, so need to duplicate
 
 	CloseEntityQuery.AddTagRequirement<FTagFragment_MassCrowd>(EMassFragmentPresence::All);
-	CloseEntityQuery.AddRequirement<FMassLODInfoFragment>(EMassFragmentAccess::ReadOnly);
+	CloseEntityQuery.AddRequirement<FMassViewerInfoFragment>(EMassFragmentAccess::ReadOnly);
 	CloseEntityQuery.AddRequirement<FMassRepresentationFragment>(EMassFragmentAccess::ReadOnly);
 	CloseEntityQuery.AddRequirement<FMassRepresentationLODFragment>(EMassFragmentAccess::ReadWrite);
 	CloseEntityQuery.AddRequirement<FDataFragment_Transform>(EMassFragmentAccess::ReadOnly);
@@ -76,7 +76,7 @@ void UMassCrowdVisualizationLODProcessor::Execute(UMassEntitySubsystem& EntitySu
 
 	TRACE_CPUPROFILER_EVENT_SCOPE(TEXT("CrowdVisualizationLOD"))
 
-	Super::ExecuteInternal<FMassLODInfoFragment>(EntitySubsystem, Context);
+	Super::ExecuteInternal<FMassViewerInfoFragment>(EntitySubsystem, Context);
 	
 	if (UE::MassCrowd::bDebugCrowdVisualizationLOD)
 	{
@@ -101,14 +101,14 @@ void UMassCrowdVisualizationLODProcessor::Execute(UMassEntitySubsystem& EntitySu
 		{
 			const TConstArrayView<FDataFragment_Transform> LocationList = Context.GetFragmentView<FDataFragment_Transform>();
 			const TConstArrayView<FMassRepresentationFragment> RepresentationFragmentList = Context.GetFragmentView<FMassRepresentationFragment>();
-			const TConstArrayView<FMassRepresentationLODFragment> RepresentationLODFragmentList = Context.GetFragmentView<FMassRepresentationLODFragment>();
+			const TConstArrayView<FMassViewerInfoFragment> LODInfoFragmentList = Context.GetFragmentView<FMassViewerInfoFragment>();
 			const int32 NumEntities = Context.GetNumEntities();
 			const float SpecifiedRangeSquaredCentimeters = FMath::Square(UE::MassCrowd::bDebugShowISMUnderSpecifiedRange * 100);
 			for (int EntityIdx = 0; EntityIdx < NumEntities; EntityIdx++)
 			{
 				const FMassRepresentationFragment& RepresentationFragment = RepresentationFragmentList[EntityIdx];
-				const FMassRepresentationLODFragment& RepresentationLODFragment = RepresentationLODFragmentList[EntityIdx];
-				if (RepresentationFragment.CurrentRepresentation == ERepresentationType::StaticMeshInstance && SpecifiedRangeSquaredCentimeters > RepresentationLODFragment.ClosestViewerDistanceSq)
+				const FMassViewerInfoFragment& LODInfoFragment = LODInfoFragmentList[EntityIdx];
+				if (RepresentationFragment.CurrentRepresentation == ERepresentationType::StaticMeshInstance && SpecifiedRangeSquaredCentimeters > LODInfoFragment.ClosestViewerDistanceSq)
 				{
 					const FDataFragment_Transform& EntityLocation = LocationList[EntityIdx];
 					DrawDebugSolidBox(World, EntityLocation.GetTransform().GetLocation() + FVector(0.0f, 0.0f, 150.0f), FVector(50.0f), FColor::Red);
@@ -131,7 +131,7 @@ void UMassCrowdVisualizationLODProcessor::Execute(UMassEntitySubsystem& EntitySu
 			LODMaxCount[0] = GetClass()->GetDefaultObject<UMassCrowdVisualizationLODProcessor>()->LODMaxCount[0];
 		}
 		UE::MassCrowd::LastDebugCrowdMaxCountHigh = UE::MassCrowd::DebugCrowdMaxCountHigh;
-		LODCalculator.Initialize(BaseLODDistance, BufferHysteresisOnDistancePercentage / 100.f, LODMaxCount, nullptr, VisibleLODDistance);
+		LODCalculator.Initialize(BaseLODDistance, BufferHysteresisOnDistancePercentage / 100.f, LODMaxCount, nullptr, DistanceToFrustum, DistanceToFrustumHysteresis, VisibleLODDistance);
 	}
 	if (UE::MassCrowd::LastDebugCrowdMaxCountMedium != UE::MassCrowd::DebugCrowdMaxCountMedium)
 	{
@@ -144,7 +144,7 @@ void UMassCrowdVisualizationLODProcessor::Execute(UMassEntitySubsystem& EntitySu
 			LODMaxCount[1] = GetClass()->GetDefaultObject<UMassCrowdVisualizationLODProcessor>()->LODMaxCount[1];
 		}
 		UE::MassCrowd::LastDebugCrowdMaxCountMedium = UE::MassCrowd::DebugCrowdMaxCountMedium;
-		LODCalculator.Initialize(BaseLODDistance, BufferHysteresisOnDistancePercentage / 100.f, LODMaxCount, nullptr, VisibleLODDistance);
+		LODCalculator.Initialize(BaseLODDistance, BufferHysteresisOnDistancePercentage / 100.f, LODMaxCount, nullptr, DistanceToFrustum, DistanceToFrustumHysteresis, VisibleLODDistance);
 	}
 	if (UE::MassCrowd::LastDebugCrowdMaxCountLow != UE::MassCrowd::DebugCrowdMaxCountLow)
 	{
@@ -157,7 +157,7 @@ void UMassCrowdVisualizationLODProcessor::Execute(UMassEntitySubsystem& EntitySu
 			LODMaxCount[2] = GetClass()->GetDefaultObject<UMassCrowdVisualizationLODProcessor>()->LODMaxCount[2];
 		}
 		UE::MassCrowd::LastDebugCrowdMaxCountLow = UE::MassCrowd::DebugCrowdMaxCountLow;
-		LODCalculator.Initialize(BaseLODDistance, BufferHysteresisOnDistancePercentage / 100.f, LODMaxCount, nullptr, VisibleLODDistance);
+		LODCalculator.Initialize(BaseLODDistance, BufferHysteresisOnDistancePercentage / 100.f, LODMaxCount, nullptr, DistanceToFrustum, DistanceToFrustumHysteresis, VisibleLODDistance);
 	}
 #endif // WITH_EDITOR
 }
diff --git a/Engine/Plugins/Runtime/MassGameplay/Source/MassLOD/Private/MassLODCollectorProcessor.cpp b/Engine/Plugins/Runtime/MassGameplay/Source/MassLOD/Private/MassLODCollectorProcessor.cpp
--- a/Engine/Plugins/Runtime/MassGameplay/Source/MassLOD/Private/MassLODCollectorProcessor.cpp
+++ b/Engine/Plugins/Runtime/MassGameplay/Source/MassLOD/Private/MassLODCollectorProcessor.cpp
@@ -14,18 +14,6 @@ UMassLODCollectorProcessor::UMassLODCollectorProcessor()
 	ExecutionOrder.ExecuteAfter.Add(UE::Mass::ProcessorGroupNames::SyncWorldToMass);
 }
 
-void UMassLODCollectorProcessor::Initialize(UObject& InOwner)
-{
-	Super::Initialize(InOwner);
-
-	for (FMassCollectorLODConfig& LODConfig : LODConfigs)
-	{
-		LODConfig.RepresentationLODCollector.Initialize(LODConfig.DistanceToFrustum, LODConfig.DistanceToFrustumHysteresis, false);
-		LODConfig.SimulationLODCollector.Initialize();
-		LODConfig.CombinedLODCollector.Initialize(LODConfig.DistanceToFrustum, LODConfig.DistanceToFrustumHysteresis, false);
-	}
-}
-
 void UMassLODCollectorProcessor::ConfigureQueries()
 {
 	for (FMassCollectorLODConfig& LODConfig : LODConfigs)
@@ -36,7 +24,7 @@ void UMassLODCollectorProcessor::ConfigureQueries()
 			BaseQuery.AddTagRequirement(*LODConfig.TagFilter.GetScriptStruct(), EMassFragmentPresence::All);
 		}
 		BaseQuery.AddRequirement<FDataFragment_Transform>(EMassFragmentAccess::ReadOnly);
-		BaseQuery.AddRequirement<FMassLODInfoFragment>(EMassFragmentAccess::ReadWrite);
+		BaseQuery.AddRequirement<FMassViewerInfoFragment>(EMassFragmentAccess::ReadWrite);
 
 		LODConfig.OnLODEntityQuery = BaseQuery;
 		LODConfig.OnLODEntityQuery.AddTagRequirement<FMassOffLODTag>(EMassFragmentPresence::None);
@@ -83,7 +71,7 @@ void CollectLOD(UMassEntitySubsystem& EntitySubsystem, FMassExecutionContext& Co
 	auto CollectLODInfo = [&Collector](FMassExecutionContext& Context)
 	{
 		TConstArrayView<FDataFragment_Transform> LocationList = Context.GetFragmentView<FDataFragment_Transform>();
-		TArrayView<FMassLODInfoFragment> ViewerInfoList = Context.GetMutableFragmentView<FMassLODInfoFragment>();
+		TArrayView<FMassViewerInfoFragment> ViewerInfoList = Context.GetMutableFragmentView<FMassViewerInfoFragment>();
 
 		Collector.CollectLODInfo(Context, LocationList, ViewerInfoList);
 	};
diff --git a/Engine/Plugins/Runtime/MassGameplay/Source/MassLOD/Private/MassLODTrait.cpp b/Engine/Plugins/Runtime/MassGameplay/Source/MassLOD/Private/MassLODTrait.cpp
--- a/Engine/Plugins/Runtime/MassGameplay/Source/MassLOD/Private/MassLODTrait.cpp
+++ b/Engine/Plugins/Runtime/MassGameplay/Source/MassLOD/Private/MassLODTrait.cpp
@@ -9,7 +9,7 @@
 void UMassSimulationLODTrait::BuildTemplate(FMassEntityTemplateBuildContext& BuildContext, UWorld& World) const
 {
 	BuildContext.AddFragmentWithDefaultInitializer<FDataFragment_Transform>();
-	BuildContext.AddFragmentWithDefaultInitializer<FMassLODInfoFragment>();
+	BuildContext.AddFragmentWithDefaultInitializer<FMassViewerInfoFragment>();
 	BuildContext.AddFragmentWithDefaultInitializer<FMassSimulationLODFragment>();
 	BuildContext.AddFragmentWithDefaultInitializer<FDataFragment_MassSimulationLODInfo>();
 	BuildContext.AddTag<FMassOffLODTag>();
diff --git a/Engine/Plugins/Runtime/MassGameplay/Source/MassLOD/Private/MassRepresentationLODCollectorProcessor.cpp b/Engine/Plugins/Runtime/MassGameplay/Source/MassLOD/Private/MassRepresentationLODCollectorProcessor.cpp
--- a/Engine/Plugins/Runtime/MassGameplay/Source/MassLOD/Private/MassRepresentationLODCollectorProcessor.cpp
+++ b/Engine/Plugins/Runtime/MassGameplay/Source/MassLOD/Private/MassRepresentationLODCollectorProcessor.cpp
@@ -10,16 +10,6 @@ UMassRepresentationLODCollectorProcessor::UMassRepresentationLODCollectorProcess
 	ExecutionFlags = (int32)EProcessorExecutionFlags::All;
 }
 
-void UMassRepresentationLODCollectorProcessor::Initialize(UObject& InOwner)
-{
-	Super::Initialize(InOwner);
-
-	for (FMassRepresentationLODCollectorConfig& LODConfig : LODConfigs)
-	{
-		LODConfig.RepresentationLODCollector.Initialize(LODConfig.FOVAnglesToDriveVisibility, LODConfig.BufferHysteresisOnFOVPercentage / 100.f);
-	}
-}
-
 void UMassRepresentationLODCollectorProcessor::ConfigureQueries()
 {
 	for (FMassRepresentationLODCollectorConfig& LODConfig : LODConfigs)
@@ -30,7 +20,7 @@ void UMassRepresentationLODCollectorProcessor::ConfigureQueries()
 			BaseQuery.AddTagRequirement(*LODConfig.TagFilter.GetScriptStruct(), EMassFragmentPresence::All);
 		}
 		BaseQuery.AddRequirement<FDataFragment_Transform>(EMassFragmentAccess::ReadOnly);
-		BaseQuery.AddRequirement<FMassLODInfoFragment>(EMassFragmentAccess::ReadWrite);
+		BaseQuery.AddRequirement<FMassViewerInfoFragment>(EMassFragmentAccess::ReadWrite);
 		BaseQuery.AddChunkRequirement<FMassVisualizationChunkFragment>(EMassFragmentAccess::ReadOnly);
 
 		LODConfig.CloseEntityQuery = BaseQuery;
@@ -49,7 +39,7 @@ void CollectLODInfo(UMassEntitySubsystem& EntitySubsystem, FMassExecutionContext
 	auto InternalCollectLODInfo = [&Collector](FMassExecutionContext& Context)
 	{
 		TConstArrayView<FDataFragment_Transform> LocationList = Context.GetFragmentView<FDataFragment_Transform>();
-		TArrayView<FMassLODInfoFragment> ViewerInfoList = Context.GetMutableFragmentView<FMassLODInfoFragment>();
+		TArrayView<FMassViewerInfoFragment> ViewerInfoList = Context.GetMutableFragmentView<FMassViewerInfoFragment>();
 
 		Collector.CollectLODInfo(Context, LocationList, ViewerInfoList);
 	};
diff --git a/Engine/Plugins/Runtime/MassGameplay/Source/MassLOD/Private/MassSimulationLOD.cpp b/Engine/Plugins/Runtime/MassGameplay/Source/MassLOD/Private/MassSimulationLOD.cpp
--- a/Engine/Plugins/Runtime/MassGameplay/Source/MassLOD/Private/MassSimulationLOD.cpp
+++ b/Engine/Plugins/Runtime/MassGameplay/Source/MassLOD/Private/MassSimulationLOD.cpp
@@ -17,13 +17,6 @@ UMassProcessor_MassSimulationLODViewersInfo::UMassProcessor_MassSimulationLODVie
 	bAutoRegisterWithProcessingPhases = false;
 }
 
-void UMassProcessor_MassSimulationLODViewersInfo::Initialize(UObject& Owner)
-{
-	Super::Initialize(Owner);
-
-	LODCollector.Initialize();
-}
-
 void UMassProcessor_MassSimulationLODViewersInfo::ConfigureQueries()
 {
 	EntityQuery.AddRequirement<FDataFragment_Transform>(EMassFragmentAccess::ReadOnly);
@@ -94,7 +87,7 @@ UMassSimulationLODProcessor::UMassSimulationLODProcessor()
 void UMassSimulationLODProcessor::ConfigureQueries()
 {
 	EntityQuery.AddRequirement<FDataFragment_Transform>(EMassFragmentAccess::ReadOnly);
-	EntityQuery.AddRequirement<FMassLODInfoFragment>(EMassFragmentAccess::ReadOnly);
+	EntityQuery.AddRequirement<FMassViewerInfoFragment>(EMassFragmentAccess::ReadOnly);
 	EntityQuery.AddRequirement<FMassSimulationLODFragment>(EMassFragmentAccess::ReadWrite);
 	EntityQuery.AddChunkRequirement<FMassSimulationVariableTickChunkFragment>(EMassFragmentAccess::ReadOnly);
 	EntityQuery.AddConstSharedRequirement<FMassSimulationLODConfig>();
@@ -140,7 +133,7 @@ void UMassSimulationLODProcessor::Execute(UMassEntitySubsystem& EntitySubsystem,
 		EntityQueryCalculateLOD.ForEachEntityChunk(EntitySubsystem, Context, [](FMassExecutionContext& Context)
 		{
 			FMassSimulationLODSharedFragment& LODSharedFragment = Context.GetMutableSharedFragment<FMassSimulationLODSharedFragment>();
-			TConstArrayView<FMassLODInfoFragment> ViewersInfoList = Context.GetFragmentView<FMassLODInfoFragment>();
+			TConstArrayView<FMassViewerInfoFragment> ViewersInfoList = Context.GetFragmentView<FMassViewerInfoFragment>();
 			TArrayView<FMassSimulationLODFragment> SimulationLODFragments = Context.GetMutableFragmentView<FMassSimulationLODFragment>();
 			LODSharedFragment.LODCalculator.CalculateLOD(Context, SimulationLODFragments, ViewersInfoList);
 		});
@@ -156,7 +149,7 @@ void UMassSimulationLODProcessor::Execute(UMassEntitySubsystem& EntitySubsystem,
 		EntityQueryAdjustDistances.ForEachEntityChunk(EntitySubsystem, Context, [](FMassExecutionContext& Context)
 		{
 			FMassSimulationLODSharedFragment& LODSharedFragment = Context.GetMutableSharedFragment<FMassSimulationLODSharedFragment>();
-			TConstArrayView<FMassLODInfoFragment> ViewersInfoList = Context.GetFragmentView<FMassLODInfoFragment>();
+			TConstArrayView<FMassViewerInfoFragment> ViewersInfoList = Context.GetFragmentView<FMassViewerInfoFragment>();
 			TArrayView<FMassSimulationLODFragment> SimulationLODFragments = Context.GetMutableFragmentView<FMassSimulationLODFragment>();
 			LODSharedFragment.LODCalculator.AdjustLODFromCount(Context, SimulationLODFragments, ViewersInfoList);
 		});
diff --git a/Engine/Plugins/Runtime/MassGameplay/Source/MassLOD/Public/MassLODCalculator.h b/Engine/Plugins/Runtime/MassGameplay/Source/MassLOD/Public/MassLODCalculator.h
--- a/Engine/Plugins/Runtime/MassGameplay/Source/MassLOD/Public/MassLODCalculator.h
+++ b/Engine/Plugins/Runtime/MassGameplay/Source/MassLOD/Public/MassLODCalculator.h
@@ -8,8 +8,8 @@
 
 /**
  * Helper struct to calculate LOD for each agent and maximize count per LOD
- *   Requires FMassLODSourceInfo fragment collected by the TMassLODCollector.
- *   Stores information in FMassLODResultInfo fragment.
+ *   Requires TViewerInfoFragment fragment collected by the TMassLODCollector.
+ *   Stores information in TLODFragment fragment.
 */
 template <typename FLODLogic = FLODDefaultLogic >
 struct TMassLODCalculator : public FMassLODBaseLogic
@@ -22,12 +22,16 @@ public:
 	 * @Param InBufferHysteresisOnFOVRatio distance hysteresis used to calculate LOD
 	 * @Param InLODMaxCount the maximum count for each LOD
 	 * @Param InLODMaxCountPerViewer the maximum count for each LOD per viewer (Only when FLODLogic::bMaximizeCountPerViewer is enabled)
+	 * @Param InVisibleDistanceToFrustum is the distance from the frustum to start considering this entity is visible (Only when FLODLogic::bDoVisibilityLogic is enabled)
+	 * @Param InVisibleDistanceToFrustumHysteresis once visible, what extra distance the entity need to be before considered not visible anymore (Only when FLODLogic::bDoVisibilityLogic is enabled)
 	 * @Param InVisibleLODDistance the maximum count for each LOD per viewer (Only when FLODLogic::bDoVisibilityLogic is enabled)
 	 */
 	void Initialize(const float InBaseLODDistance[EMassLOD::Max], 
 					const float InBufferHysteresisOnDistanceRatio, 
 					const int32 InLODMaxCount[EMassLOD::Max], 
 					const int32 InLODMaxCountPerViewer[EMassLOD::Max] = nullptr,
+					const float InVisibleDistanceToFrustum = 0.0f,
+					const float InVisibleDistanceToFrustumHysteresis = 0.0f,
 					const float InVisibleLODDistance[EMassLOD::Max] = nullptr);
 
 	/**
@@ -42,8 +46,8 @@ public:
 	 * @Param LODList is the fragment where calculation are stored
 	 * @Param ViewersInfoList is the source information fragment for LOD calculation
 	 */
-	template< typename FMassLODResultInfo, typename FMassLODSourceInfo >
-	void CalculateLOD(FMassExecutionContext& Context, TArrayView<FMassLODResultInfo> LODList, TConstArrayView<FMassLODSourceInfo> ViewersInfoList);
+	template <typename TLODFragment, typename TViewerInfoFragment>
+	void CalculateLOD(FMassExecutionContext& Context, TArrayView<TLODFragment> LODList, TConstArrayView<TViewerInfoFragment> ViewersInfoList);
 
 	/**
 	 * Adjust LOD distances by clamping them to respect the maximum LOD count
@@ -57,16 +61,16 @@ public:
 	 * @Param LODList is the fragment where calculation are stored
 	 * @Param ViewersInfoList is the source information fragment for LOD calculation
 	 */
-	template< typename FMassLODResultInfo, typename FMassLODSourceInfo >
-	void AdjustLODFromCount(FMassExecutionContext& Context, TArrayView<FMassLODResultInfo> LODList, TConstArrayView<FMassLODSourceInfo> ViewersInfoList);
+	template <typename TLODFragment, typename TViewerInfoFragment>
+	void AdjustLODFromCount(FMassExecutionContext& Context, TArrayView<TLODFragment> LODList, TConstArrayView<TViewerInfoFragment> ViewersInfoList);
 
 	/**
 	 * Turn Off all LOD, called for each entity chunks
 	 * @Param Context of the chunk execution
 	 * @Param LODList is the fragment where calculation are stored
 	 */
-	template <typename FMassLODResultInfo>
-	void ForceOffLOD(FMassExecutionContext& Context, TArrayView<FMassLODResultInfo> LODList);
+	template <typename TLODFragment>
+	void ForceOffLOD(FMassExecutionContext& Context, TArrayView<TLODFragment> LODList);
 
 	/**
 	 * Debug draw the current state of each agent as a color coded square
@@ -75,8 +79,8 @@ public:
 	 * @Param LocationList is the fragment transforms of the entities
 	 * @Param World where the debug display should be drawn
 	 */
-	template< typename FMassLODResultInfo, typename TFragmentLocation >
-	void DebugDisplayLOD(FMassExecutionContext& Context, TConstArrayView<FMassLODResultInfo> LODList, TConstArrayView<TFragmentLocation> LocationList, UWorld* World);
+	template <typename TLODFragment, typename TTransformFragment>
+	void DebugDisplayLOD(FMassExecutionContext& Context, TConstArrayView<TLODFragment> LODList, TConstArrayView<TTransformFragment> LocationList, UWorld* World);
 
 protected:
 
@@ -127,6 +131,8 @@ protected:
 
 	EMassLOD::Type ComputeLODFromSettings(const EMassLOD::Type PrevLOD, const float DistanceToViewerSq, const bool bIsVisible, bool* bIsInAVisibleRange, const FMassLODRuntimeData& Data) const;
 
+	bool CalculateVisibility(const bool bWasVisible, const float DistanceToFrustum) const;
+
 	/** LOD distances */
 	TStaticArray<float, EMassLOD::Max> BaseLODDistance;
 	TStaticArray<float, EMassLOD::Max> VisibleLODDistance;
@@ -140,6 +146,12 @@ protected:
 	/** Ratio for Buffer Distance Hysteresis */
 	float BufferHysteresisOnDistanceRatio = 0.1f;
 
+	/** How far away from frustum does this entities are considered visible */
+	float VisibleDistanceToFrustum = 0.0f;
+
+	/** Once visible how much further than distance to frustum does the entities need to be before being consider not visible */
+	float VisibleDistanceToFrustumWithHysteresis = 0.0f;
+
 	/** The size of each subdivision per LOD (LOD Size/MaxBucketsPerLOD) */
 	TStaticArray<float, EMassLOD::Max> BaseBucketSize;
 	TStaticArray<float, EMassLOD::Max> VisibleBucketSize;
@@ -159,8 +171,13 @@ void TMassLODCalculator<FLODLogic>::Initialize(const float InBaseLODDistance[EMa
 											   const float InBufferHysteresisOnDistanceRatio,
 											   const int32 InLODMaxCount[EMassLOD::Max],
 											   const int32 InLODMaxCountPerViewer[EMassLOD::Max] /*= nullptr*/,
-											   const float InVisibleLODDistance[EMassLOD::Max] /*= nullptr*/)
+											   const float InVisibleDistanceToFrustum /*= 0.0f*/,
+											   const float InVisibleDistanceToFrustumHysteresis /*= 0.0f*/,
+											   const float InVisibleLODDistance[EMassLOD::Max] /*= nullptr*/ )
 {
+	static_assert(!FLODLogic::bCalculateLODPerViewer || FLODLogic::bStoreInfoPerViewer, "Need to enable store info per viewer to be able to calculate LOD per viewer");
+	static_assert(!FLODLogic::bMaximizeCountPerViewer || FLODLogic::bCalculateLODPerViewer, "Need to enable CalculatedLODPerviewer in order to maximize count per viewer");
+
 	checkf(FLODLogic::bMaximizeCountPerViewer == (InLODMaxCountPerViewer != nullptr), TEXT("Missmatched between expected parameter InLODMaxCountPerViewer and LOD logic trait bMaximizeCountPerViewer."));
 	checkf(FLODLogic::bDoVisibilityLogic == (InVisibleLODDistance != nullptr), TEXT("Missmatched between expected parameter InVisibleLODDistance and LOD logic trait bDoVisibilityLogic."));
 
@@ -204,6 +221,16 @@ void TMassLODCalculator<FLODLogic>::Initialize(const float InBaseLODDistance[EMa
 
 	// Assuming that off is the farthest distance, calculate the max LOD distance
 	MaxLODDistance = !FLODLogic::bDoVisibilityLogic || BaseLODDistance[EMassLOD::Off] >= VisibleLODDistance[EMassLOD::Off] ? BaseLODDistance[EMassLOD::Off] : VisibleLODDistance[EMassLOD::Off];
+
+	// Distance to frustum settings
+	VisibleDistanceToFrustum = InVisibleDistanceToFrustum;
+	VisibleDistanceToFrustumWithHysteresis = InVisibleDistanceToFrustum + InVisibleDistanceToFrustumHysteresis;
+}
+
+template <typename FLODLogic>
+bool TMassLODCalculator<FLODLogic>::CalculateVisibility(const bool bWasVisible, const float DistanceToFrustum) const
+{
+	return DistanceToFrustum < (bWasVisible ? VisibleDistanceToFrustumWithHysteresis : VisibleDistanceToFrustum);
 }
 
 template <typename FLODLogic>
@@ -265,8 +292,8 @@ float TMassLODCalculator<FLODLogic>::AccumulateCountInRuntimeData(const EMassLOD
 }
 
 template <typename FLODLogic>
-template< typename FMassLODResultInfo, typename FMassLODSourceInfo >
-void TMassLODCalculator<FLODLogic>::CalculateLOD(FMassExecutionContext& Context, TArrayView<FMassLODResultInfo> LODList, TConstArrayView<FMassLODSourceInfo> ViewersInfoList)
+template <typename TLODFragment, typename TViewerInfoFragment>
+void TMassLODCalculator<FLODLogic>::CalculateLOD(FMassExecutionContext& Context, TArrayView<TLODFragment> LODList, TConstArrayView<TViewerInfoFragment> ViewersInfoList)
 {
 #if WITH_MASSGAMEPLAY_DEBUG
 	if (UE::MassLOD::Debug::bLODCalculationsPaused)
@@ -279,63 +306,72 @@ void TMassLODCalculator<FLODLogic>::CalculateLOD(FMassExecutionContext& Context,
 	for (int EntityIdx = 0; EntityIdx < NumEntities; EntityIdx++)
 	{
 		// Calculate the LOD purely upon distances
-		const FMassLODSourceInfo& EntityViewersInfo = ViewersInfoList[EntityIdx];
-		FMassLODResultInfo& EntityLOD = LODList[EntityIdx];
-		EntityLOD.ClosestViewerDistanceSq = FLT_MAX;
-		bool bIsVisibleByAViewer = false;
+		const TViewerInfoFragment& EntityViewersInfo = ViewersInfoList[EntityIdx];
+		TLODFragment& EntityLOD = LODList[EntityIdx];
 		bool bIsInAVisibleRange = false;
+		const float ClosestDistanceToFrustum = GetClosestDistanceToFrustum<FLODLogic::bDoVisibilityLogic>(EntityViewersInfo, FLT_MAX);
+		const bool bWasVisibleByAViewer = GetbIsVisibleByAViewer<FLODLogic::bDoVisibilityLogic>(EntityLOD, false);
+		const bool bIsVisibleByAViewer = CalculateVisibility(bWasVisibleByAViewer, ClosestDistanceToFrustum);
 
-		for (int ViewerIdx = 0; ViewerIdx < Viewers.Num(); ++ViewerIdx)
+		// Set visibility
+		SetbWasVisibleByAViewer<FLODLogic::bDoVisibilityLogic>(EntityLOD, bWasVisibleByAViewer);
+		SetbIsVisibleByAViewer<FLODLogic::bDoVisibilityLogic>(EntityLOD, bIsVisibleByAViewer);
+
+		// Find new LOD
+		EntityLOD.PrevLOD = EntityLOD.LOD;
+		EntityLOD.LOD = ComputeLODFromSettings(EntityLOD.PrevLOD, EntityViewersInfo.ClosestViewerDistanceSq, bIsVisibleByAViewer, &bIsInAVisibleRange, RuntimeData);
+
+		SetbWasInVisibleRange<FLODLogic::bDoVisibilityLogic>(EntityLOD, GetbIsInVisibleRange<FLODLogic::bDoVisibilityLogic>(EntityLOD, false));
+		SetbIsInVisibleRange<FLODLogic::bDoVisibilityLogic>(EntityLOD, bIsInAVisibleRange);
+
+		// Accumulate in buckets
+		const float LODSignificance = AccumulateCountInRuntimeData<FLODLogic::bCalculateLODSignificance>(EntityLOD.LOD, EntityViewersInfo.ClosestViewerDistanceSq, bIsVisibleByAViewer, RuntimeData);
+		SetLODSignificance<FLODLogic::bCalculateLODSignificance>(EntityLOD, LODSignificance);
+
+		// Do per viewer logic if asked for
+		if (FLODLogic::bStoreInfoPerViewer)
 		{
-			const FViewerLODInfo& Viewer = Viewers[ViewerIdx];
-			if (Viewer.bClearData)
-			{
-				SetLODPerViewer<FLODLogic::bStoreLODPerViewer>(EntityLOD, ViewerIdx, EMassLOD::Max);
-				SetPrevLODPerViewer<FLODLogic::bStoreLODPerViewer>(EntityLOD, ViewerIdx, EMassLOD::Max);
-			}
-			if (Viewer.Handle.IsValid())
+			for (int ViewerIdx = 0; ViewerIdx < Viewers.Num(); ++ViewerIdx)
 			{
-				const bool bIsVisibleByViewer = GetbIsVisibleByViewer<FLODLogic::bDoVisibilityLogic>(EntityViewersInfo, ViewerIdx, false);
-				bIsVisibleByAViewer |= bIsVisibleByViewer;
-
-				if (EntityLOD.ClosestViewerDistanceSq > EntityViewersInfo.DistanceToViewerSq[ViewerIdx])
+				const FViewerLODInfo& Viewer = Viewers[ViewerIdx];
+				if (Viewer.bClearData)
 				{
-					EntityLOD.ClosestViewerDistanceSq = EntityViewersInfo.DistanceToViewerSq[ViewerIdx];
+					SetLODPerViewer<FLODLogic::bCalculateLODPerViewer>(EntityLOD, ViewerIdx, EMassLOD::Max);
+					SetPrevLODPerViewer<FLODLogic::bCalculateLODPerViewer>(EntityLOD, ViewerIdx, EMassLOD::Max);
+					SetbWasVisibleByViewer<FLODLogic::bDoVisibilityLogic&& FLODLogic::bStoreInfoPerViewer>(EntityLOD, ViewerIdx, false);
+					SetbIsVisibleByViewer<FLODLogic::bDoVisibilityLogic && FLODLogic::bStoreInfoPerViewer>(EntityLOD, ViewerIdx, false);
 				}
-
-				if (FLODLogic::bStoreLODPerViewer || FLODLogic::bMaximizeCountPerViewer)
+				if (Viewer.Handle.IsValid())
 				{
-					const EMassLOD::Type PrevLODPerViewer = GetLODPerViewer<FLODLogic::bStoreLODPerViewer>(EntityLOD, ViewerIdx, EntityLOD.LOD);
-
-					// Find new LOD
-					const EMassLOD::Type LODPerViewer = ComputeLODFromSettings(PrevLODPerViewer, EntityViewersInfo.DistanceToViewerSq[ViewerIdx], bIsVisibleByViewer, &bIsInAVisibleRange, RuntimeData);
+					const float DistanceToFrustum = GetDistanceToFrustum<FLODLogic::bDoVisibilityLogic && FLODLogic::bStoreInfoPerViewer>(EntityViewersInfo, ViewerIdx, FLT_MAX);
+					const bool bWasVisibleByViewer = GetbIsVisibleByViewer<FLODLogic::bDoVisibilityLogic && FLODLogic::bStoreInfoPerViewer>(EntityLOD, ViewerIdx, false);
+					const bool bIsVisibleByViewer = CalculateVisibility(bWasVisibleByAViewer, DistanceToFrustum);
 
-					// Set Per viewer LOD
-					SetPrevLODPerViewer<FLODLogic::bStoreLODPerViewer>(EntityLOD, ViewerIdx, PrevLODPerViewer);
-					SetLODPerViewer<FLODLogic::bStoreLODPerViewer>(EntityLOD, ViewerIdx, LODPerViewer);
+					// Set visibility
+					SetbIsVisibleByViewer<FLODLogic::bDoVisibilityLogic && FLODLogic::bStoreInfoPerViewer>(EntityLOD, ViewerIdx, bIsVisibleByViewer);
+					SetbWasVisibleByViewer<FLODLogic::bDoVisibilityLogic && FLODLogic::bStoreInfoPerViewer>(EntityLOD, ViewerIdx, bWasVisibleByViewer);
 
-					if (FLODLogic::bMaximizeCountPerViewer)
+					if (FLODLogic::bCalculateLODPerViewer)
 					{
-						// Accumulate in buckets
-						AccumulateCountInRuntimeData<false>(LODPerViewer, EntityViewersInfo.DistanceToViewerSq[ViewerIdx], bIsVisibleByViewer, RuntimeDataPerViewer[ViewerIdx]);
+						const float DistanceToViewerSq = GetDistanceToViewerSq<FLODLogic::bStoreInfoPerViewer>(EntityViewersInfo, ViewerIdx, FLT_MAX);
+						const EMassLOD::Type PrevLODPerViewer = GetLODPerViewer<FLODLogic::bCalculateLODPerViewer>(EntityLOD, ViewerIdx, EntityLOD.LOD);
+
+						// Find new LOD
+						const EMassLOD::Type LODPerViewer = ComputeLODFromSettings(PrevLODPerViewer, DistanceToViewerSq, bIsVisibleByViewer, &bIsInAVisibleRange, RuntimeData);
+
+						// Set Per viewer LOD
+						SetPrevLODPerViewer<FLODLogic::bCalculateLODPerViewer>(EntityLOD, ViewerIdx, PrevLODPerViewer);
+						SetLODPerViewer<FLODLogic::bCalculateLODPerViewer>(EntityLOD, ViewerIdx, LODPerViewer);
+
+						if (FLODLogic::bMaximizeCountPerViewer)
+						{
+							// Accumulate in buckets
+							AccumulateCountInRuntimeData<false>(LODPerViewer, DistanceToViewerSq, bIsVisibleByViewer, RuntimeDataPerViewer[ViewerIdx]);
+						}
 					}
 				}
 			}
 		}
-
-		// Find new LOD
-		EntityLOD.PrevLOD = EntityLOD.LOD;
-		EntityLOD.LOD = ComputeLODFromSettings(EntityLOD.PrevLOD, EntityLOD.ClosestViewerDistanceSq, bIsVisibleByAViewer, &bIsInAVisibleRange, RuntimeData);
-
-		// Set visibility
-		SetbWasVisibleByAViewer<FLODLogic::bDoVisibilityLogic>(EntityLOD, GetbIsVisibleByAViewer<FLODLogic::bDoVisibilityLogic>(EntityLOD, false));
-		SetbIsVisibleByAViewer<FLODLogic::bDoVisibilityLogic>(EntityLOD, bIsVisibleByAViewer);
-		SetbWasInVisibleRange<FLODLogic::bDoVisibilityLogic>(EntityLOD, GetbIsInVisibleRange<FLODLogic::bDoVisibilityLogic>(EntityLOD, false));
-		SetbIsInVisibleRange<FLODLogic::bDoVisibilityLogic>(EntityLOD, bIsInAVisibleRange);
-
-		// Accumulate in buckets
-		const float LODSignificance = AccumulateCountInRuntimeData<FLODLogic::bCalculateLODSignificance>(EntityLOD.LOD, EntityLOD.ClosestViewerDistanceSq, bIsVisibleByAViewer, RuntimeData);
-		SetLODSignificance<FLODLogic::bCalculateLODSignificance>(EntityLOD, LODSignificance);
 	}
 }
 
@@ -473,15 +509,15 @@ bool TMassLODCalculator<FLODLogic>::AdjustDistancesFromCount()
 }
 
 template <typename FLODLogic>
-template< typename FMassLODResultInfo, typename FMassLODSourceInfo >
-void TMassLODCalculator<FLODLogic>::AdjustLODFromCount(FMassExecutionContext& Context, TArrayView<FMassLODResultInfo> LODList, TConstArrayView<FMassLODSourceInfo> ViewersInfoList)
+template <typename TLODFragment, typename TViewerInfoFragment>
+void TMassLODCalculator<FLODLogic>::AdjustLODFromCount(FMassExecutionContext& Context, TArrayView<TLODFragment> LODList, TConstArrayView<TViewerInfoFragment> ViewersInfoList)
 {
 	const int32 NumEntities = Context.GetNumEntities();
 	// Adjust LOD for each viewer and remember the new highest
 	for (int EntityIdx = 0; EntityIdx < NumEntities; EntityIdx++)
 	{
-		const FMassLODSourceInfo& EntityViewersInfo = ViewersInfoList[EntityIdx];
-		FMassLODResultInfo& EntityLOD = LODList[EntityIdx];
+		const TViewerInfoFragment& EntityViewersInfo = ViewersInfoList[EntityIdx];
+		TLODFragment& EntityLOD = LODList[EntityIdx];
 		EMassLOD::Type HighestViewerLOD = EMassLOD::Off;
 		if (FLODLogic::bMaximizeCountPerViewer)
 		{
@@ -492,22 +528,23 @@ void TMassLODCalculator<FLODLogic>::AdjustLODFromCount(FMassExecutionContext& Co
 					continue;
 				}
 
-				const bool bIsVisibleByViewer = GetbIsVisibleByViewer<FLODLogic::bDoVisibilityLogic>(EntityViewersInfo, ViewerIdx, false);
-				EMassLOD::Type LODPerViewer = GetLODPerViewer<FLODLogic::bStoreLODPerViewer>(EntityLOD, ViewerIdx, EntityLOD.LOD);
+				const float DistanceToViewerSq = GetDistanceToViewerSq<FLODLogic::bStoreInfoPerViewer>(EntityViewersInfo, ViewerIdx, FLT_MAX);
+				const bool bIsVisibleByViewer = GetbIsVisibleByViewer<FLODLogic::bDoVisibilityLogic && FLODLogic::bStoreInfoPerViewer>(EntityLOD, ViewerIdx, false);
+				EMassLOD::Type LODPerViewer = GetLODPerViewer<FLODLogic::bCalculateLODPerViewer>(EntityLOD, ViewerIdx, EntityLOD.LOD);
 
-				LODPerViewer = ComputeLODFromSettings(LODPerViewer, EntityViewersInfo.DistanceToViewerSq[ViewerIdx], bIsVisibleByViewer, nullptr, RuntimeDataPerViewer[ViewerIdx]);
+				LODPerViewer = ComputeLODFromSettings(LODPerViewer, DistanceToViewerSq, bIsVisibleByViewer, nullptr, RuntimeDataPerViewer[ViewerIdx]);
 
 				if (HighestViewerLOD < LODPerViewer)
 				{
 					HighestViewerLOD = LODPerViewer;
 				}
 
-				SetLODPerViewer<FLODLogic::bStoreLODPerViewer>(EntityLOD, ViewerIdx, LODPerViewer);
+				SetLODPerViewer<FLODLogic::bCalculateLODPerViewer>(EntityLOD, ViewerIdx, LODPerViewer);
 			}
 		}
 
 		const bool bIsVisibleByAViewer = GetbIsVisibleByAViewer<FLODLogic::bDoVisibilityLogic>(EntityLOD, false);
-		EMassLOD::Type NewLOD = ComputeLODFromSettings(EntityLOD.PrevLOD, EntityLOD.ClosestViewerDistanceSq, bIsVisibleByAViewer, nullptr, RuntimeData);
+		EMassLOD::Type NewLOD = ComputeLODFromSettings(EntityLOD.PrevLOD, EntityViewersInfo.ClosestViewerDistanceSq, bIsVisibleByAViewer, nullptr, RuntimeData);
 
 		// Maybe the highest of all the viewers is now lower than the global entity LOD, make sure to update the it accordingly
 		if (FLODLogic::bMaximizeCountPerViewer && NewLOD < HighestViewerLOD)
@@ -529,7 +566,7 @@ void TMassLODCalculator<FLODLogic>::AdjustLODFromCount(FMassExecutionContext& Co
 					const TStaticArray<float, EMassLOD::Max>& AdjustedLODDistance = FLODLogic::bDoVisibilityLogic && bIsVisibleByAViewer ? RuntimeData.AdjustedVisibleLODDistance : RuntimeData.AdjustedBaseLODDistance;
 
 					// Need to clamp as the Sqrt is not precise enough and always end up with floating calculation errors
-					const float DistanceBetweenLODThresholdAndEntity = FMath::Max(FMath::Sqrt(EntityLOD.ClosestViewerDistanceSq) - AdjustedLODDistance[NewLOD], 0.f);
+					const float DistanceBetweenLODThresholdAndEntity = FMath::Max(FMath::Sqrt(EntityViewersInfo.ClosestViewerDistanceSq) - AdjustedLODDistance[NewLOD], 0.f);
 
 					// Derive significance from distance between viewer and where the agent stands between both LOD threshold
 					const float AdjustedDistanceBetweenCurrentLODAndNext = AdjustedLODDistance[NewLOD + 1] - AdjustedLODDistance[NewLOD];
@@ -544,13 +581,13 @@ void TMassLODCalculator<FLODLogic>::AdjustLODFromCount(FMassExecutionContext& Co
 }
 
 template <typename FLODLogic>
-template <typename FMassLODResultInfo>
-void TMassLODCalculator<FLODLogic>::ForceOffLOD(FMassExecutionContext& Context, TArrayView<FMassLODResultInfo> LODList)
+template <typename TLODFragment>
+void TMassLODCalculator<FLODLogic>::ForceOffLOD(FMassExecutionContext& Context, TArrayView<TLODFragment> LODList)
 {
 	const int32 NumEntities = Context.GetNumEntities();
 	for (int EntityIdx = 0; EntityIdx < NumEntities; EntityIdx++)
 	{
-		FMassLODResultInfo& EntityLOD = LODList[EntityIdx];
+		TLODFragment& EntityLOD = LODList[EntityIdx];
 
 		// Force off LOD
 		EntityLOD.PrevLOD = EntityLOD.LOD;
@@ -565,14 +602,14 @@ void TMassLODCalculator<FLODLogic>::ForceOffLOD(FMassExecutionContext& Context,
 }
 
 template <typename FLODLogic>
-template< typename FMassLODResultInfo, typename TFragmentLocation >
-void TMassLODCalculator<FLODLogic>::DebugDisplayLOD(FMassExecutionContext& Context, TConstArrayView<FMassLODResultInfo> LODList, TConstArrayView<TFragmentLocation> LocationList, UWorld* World)
+template <typename TLODFragment, typename TTransformFragment>
+void TMassLODCalculator<FLODLogic>::DebugDisplayLOD(FMassExecutionContext& Context, TConstArrayView<TLODFragment> LODList, TConstArrayView<TTransformFragment> LocationList, UWorld* World)
 {
 	const int32 NumEntities = Context.GetNumEntities();
 	for (int EntityIdx = 0; EntityIdx < NumEntities; EntityIdx++)
 	{
-		const TFragmentLocation& EntityLocation = LocationList[EntityIdx];
-		const FMassLODResultInfo& EntityLOD = LODList[EntityIdx];
+		const TTransformFragment& EntityLocation = LocationList[EntityIdx];
+		const TLODFragment& EntityLOD = LODList[EntityIdx];
 		int32 LODIdx = (int32)EntityLOD.LOD;
 		DrawDebugSolidBox(World, EntityLocation.GetTransform().GetLocation() + FVector(0.0f, 0.0f, 120.0f), FVector(25.0f), UE::MassLOD::LODColors[LODIdx]);
 	}
diff --git a/Engine/Plugins/Runtime/MassGameplay/Source/MassLOD/Public/MassLODCollector.h b/Engine/Plugins/Runtime/MassGameplay/Source/MassLOD/Public/MassLODCollector.h
--- a/Engine/Plugins/Runtime/MassGameplay/Source/MassLOD/Public/MassLODCollector.h
+++ b/Engine/Plugins/Runtime/MassGameplay/Source/MassLOD/Public/MassLODCollector.h
@@ -5,34 +5,16 @@
 #include "MassLODLogic.h"
 #include "DrawDebugHelpers.h"
 
-/**
- * This is the expected member variables for the TArrayView<FMassLODSourceInfo> when calling TMassLODCollector methods
- *
-struct FMassLODSourceInfo
-{
-	// Distance to viewer (Always needed)
-	TStaticArray<float, UE::MassLOD::MaxNumOfViewers> DistanceToViewerSq;
 
-	// Visibility information (Only when FLODLogic::bDoVisibilityLogic is enabled)
-	TStaticArray<bool, UE::MassLOD::MaxNumOfViewers> bIsVisibleByViewer;
-};
-*/
 
 /**
  * Helper struct to collect needed information on each agent that will be needed later for LOD calculation
- *   Requires FMassTransform fragment.
- *   Stores information in FMassLODSourceInfo fragment.
+ *   Requires TTransformFragment fragment.
+ *   Stores information in TViewerInfoFragment fragment.
  */
 template <typename FLODLogic = FLODDefaultLogic >
 struct TMassLODCollector : public FMassLODBaseLogic
 {
-	/**
-	 * Initializes the LOD collector, needed to be called once at initialization time
-	 * @Param InDistanceToFrustum the extra distance to frustrum to use to know if the entity is visible (Only when FLODLogic::bDoVisibilityLogic is enabled)
-	 * @Param InDistanceToFrustumHysteresis the hyteresis on the distance to frustrum to use when entity was previously visible to become not visible (Only when FLODLogic::bDoVisibilityLogic is enabled)
-	 */
-	void Initialize(const float InDistanceToFrustum = 0.0f, const float InDistanceToFrustumHysteresis = 0.0f, bool bTemp = false);
-
 	/**
 	 * Prepares execution for the current frame, needed to be called before every execution
 	 * @Param Viewers is the array of all the known viewers
@@ -45,22 +27,11 @@ struct TMassLODCollector : public FMassLODBaseLogic
 	 * @Param TranformList is the fragment transforms of the entities
 	 * @Param ViewersInfoList is the fragment where to store source information for LOD calculation
 	 */
-	template< typename FMassTransform, typename FMassLODSourceInfo >
-	void CollectLODInfo(FMassExecutionContext& Context, TConstArrayView<FMassTransform> TranformList, TArrayView<FMassLODSourceInfo> ViewersInfoList);
+	template <typename TTransformFragment, typename TViewerInfoFragment>
+	void CollectLODInfo(FMassExecutionContext& Context, TConstArrayView<TTransformFragment> TranformList, TArrayView<TViewerInfoFragment> ViewersInfoList);
 
-protected:
-
-	float DistanceToFrustum = 0.0f;
-	float DistanceToFrustumHysteresis = 0.0f;
 };
 
-template <typename FLODLogic>
-void TMassLODCollector<FLODLogic>::Initialize(const float InDistanceToFrustum /*= 0.0f*/, const float InDistanceToFrustumHysteresis /*= 0.0f*/, bool bTemp /*= false*/)
-{
-	DistanceToFrustum = InDistanceToFrustum;
-	DistanceToFrustumHysteresis = InDistanceToFrustumHysteresis;
-}
-
 template <typename FLODLogic>
 void TMassLODCollector<FLODLogic>::PrepareExecution(TConstArrayView<FViewerInfo> ViewersInfo)
 {
@@ -68,45 +39,48 @@ void TMassLODCollector<FLODLogic>::PrepareExecution(TConstArrayView<FViewerInfo>
 }
 
 template <typename FLODLogic>
-template< typename FMassTransform, typename FMassLODSourceInfo >
-void TMassLODCollector<FLODLogic>::CollectLODInfo(FMassExecutionContext& Context, TConstArrayView<FMassTransform> TranformList, TArrayView<FMassLODSourceInfo> ViewersInfoList)
+template <typename TTransformFragment, typename TViewerInfoFragment>
+void TMassLODCollector<FLODLogic>::CollectLODInfo(FMassExecutionContext& Context, TConstArrayView<TTransformFragment> TranformList, TArrayView<TViewerInfoFragment> ViewersInfoList)
 {
 	const int32 NumEntities = Context.GetNumEntities();
 	for (int EntityIdx = 0; EntityIdx < NumEntities; EntityIdx++)
 	{
-		const FMassTransform& EntityTransform = TranformList[EntityIdx];
-		FMassLODSourceInfo& EntityViewerInfo = ViewersInfoList[EntityIdx];
+		float ClosestViewerDistanceSq = FLT_MAX;
+		float ClosestDistanceToFrustum = FLT_MAX;
+		const TTransformFragment& EntityTransform = TranformList[EntityIdx];
+		TViewerInfoFragment& EntityViewerInfo = ViewersInfoList[EntityIdx];
 		for (int ViewerIdx = 0; ViewerIdx < Viewers.Num(); ++ViewerIdx)
 		{
 			const FViewerLODInfo& Viewer = Viewers[ViewerIdx];
 			if (Viewer.bClearData)
 			{
-				EntityViewerInfo.DistanceToViewerSq[ViewerIdx] = FLT_MAX;
-				SetbIsVisibleByViewer<FLODLogic::bDoVisibilityLogic>(EntityViewerInfo, ViewerIdx, false);
+				SetDistanceToViewerSq<FLODLogic::bStoreInfoPerViewer>(EntityViewerInfo, ViewerIdx, FLT_MAX);
+				SetDistanceToFrustum<FLODLogic::bDoVisibilityLogic && FLODLogic::bStoreInfoPerViewer>(EntityViewerInfo, ViewerIdx, FLT_MAX);
 			}
 			if (Viewer.Handle.IsValid())
 			{
 				const FVector& EntityLocation = EntityTransform.GetTransform().GetLocation();
 				const FVector ViewerToEntity = EntityLocation - Viewer.Location;
-				EntityViewerInfo.DistanceToViewerSq[ViewerIdx] = ViewerToEntity.SizeSquared();
+				const float DistanceToViewerSq = ViewerToEntity.SizeSquared();
+				if (ClosestViewerDistanceSq > DistanceToViewerSq)
+				{
+					ClosestViewerDistanceSq = DistanceToViewerSq;
+				}
+				SetDistanceToViewerSq<FLODLogic::bStoreInfoPerViewer>(EntityViewerInfo, ViewerIdx, DistanceToViewerSq);
 
 				if (FLODLogic::bDoVisibilityLogic)
 				{
-					const bool bWasVisibleByViewer = GetbIsVisibleByViewer<FLODLogic::bDoVisibilityLogic>(EntityViewerInfo, ViewerIdx, false);
-					bool bIsVisibleByViewer = false;
-
-					if (EntityViewerInfo.DistanceToViewerSq[ViewerIdx] <= KINDA_SMALL_NUMBER)
-					{
-						bIsVisibleByViewer = true;
-					}
-					else
+					const float DistanceToFrustum = Viewer.Frustum.DistanceTo(EntityLocation);
+					SetDistanceToFrustum<FLODLogic::bDoVisibilityLogic && FLODLogic::bStoreInfoPerViewer>(EntityViewerInfo, ViewerIdx, DistanceToFrustum);
+					if (ClosestDistanceToFrustum > DistanceToFrustum)
 					{
-						bIsVisibleByViewer = Viewer.Frustum.IntersectSphere(EntityLocation, bWasVisibleByViewer ? DistanceToFrustum + DistanceToFrustumHysteresis : DistanceToFrustum );
+						ClosestDistanceToFrustum = DistanceToFrustum;
 					}
 
-					SetbIsVisibleByViewer<FLODLogic::bDoVisibilityLogic>(EntityViewerInfo, ViewerIdx, bIsVisibleByViewer);
 				}
 			}
 		}
+		EntityViewerInfo.ClosestViewerDistanceSq = ClosestViewerDistanceSq;
+		SetClosestDistanceToFrustum<FLODLogic::bDoVisibilityLogic>(EntityViewerInfo, ClosestDistanceToFrustum);
 	}
 }
\ No newline at end of file
diff --git a/Engine/Plugins/Runtime/MassGameplay/Source/MassLOD/Public/MassLODCollectorProcessor.h b/Engine/Plugins/Runtime/MassGameplay/Source/MassLOD/Public/MassLODCollectorProcessor.h
--- a/Engine/Plugins/Runtime/MassGameplay/Source/MassLOD/Public/MassLODCollectorProcessor.h
+++ b/Engine/Plugins/Runtime/MassGameplay/Source/MassLOD/Public/MassLODCollectorProcessor.h
@@ -18,13 +18,6 @@ struct FMassCollectorLODConfig
 	UPROPERTY(EditAnywhere, Category = "Mass|LOD", config, meta = (BaseStruct = "MassTag"))
 	FInstancedStruct TagFilter;
 
-	/** How far away from frustrum does this entities are considered visible */
-	UPROPERTY(EditAnywhere, Category = "Mass|LOD", meta = (ClampMin = "0.0", UIMin = "0.0"), config)
-	float DistanceToFrustum = 0.0f;
-	// Once visible how much further than DistanceToFrustum does the entities need to be before being cull again
-	UPROPERTY(EditAnywhere, Category = "Mass|LOD", meta = (ClampMin = "0.0", UIMin = "0.0"), config)
-	float DistanceToFrustumHysteresis = 0.0f;
-
 	/** Runtime data for matching the LOD config */
 	TMassLODCollector<FMassRepresentationLODLogic> RepresentationLODCollector;
 	FMassEntityQuery CloseEntityQuery;
@@ -52,7 +45,6 @@ public:
 	UMassLODCollectorProcessor();
 
 protected:
-	virtual void Initialize(UObject& Owner) override;
 	virtual void ConfigureQueries() override;
 	virtual void Execute(UMassEntitySubsystem& EntitySubsystem, FMassExecutionContext& Context) override;
 
diff --git a/Engine/Plugins/Runtime/MassGameplay/Source/MassLOD/Public/MassLODLogic.h b/Engine/Plugins/Runtime/MassGameplay/Source/MassLOD/Public/MassLODLogic.h
--- a/Engine/Plugins/Runtime/MassGameplay/Source/MassLOD/Public/MassLODLogic.h
+++ b/Engine/Plugins/Runtime/MassGameplay/Source/MassLOD/Public/MassLODLogic.h
@@ -26,7 +26,8 @@ struct FLODDefaultLogic
 {
 	enum
 	{
-		bStoreLODPerViewer = false, // Enable to calculate and store the result LOD per viewer in the FMassLODResultInfo::LODPerViewer and FMassLODResultInfo::PrevLODPerViewer.
+		bStoreInfoPerViewer = false, // Enable to store all calculated information per viewer
+		bCalculateLODPerViewer = false, // Enable to calculate and store the result LOD per viewer in the FMassLODResultInfo::LODPerViewer and FMassLODResultInfo::PrevLODPerViewer.
 		bMaximizeCountPerViewer = false, // Enable to maximize count per viewer, requires a valid InLODMaxCountPerViewer parameter during initialization of TMassLODCalculator.
 		bDoVisibilityLogic = false, // Enable to calculate visibility and apply its own LOD distances. Requires a valid InVisibleLODDistance parameter during initialization of TMassLODCalculator.
 		bCalculateLODSignificance = false, // Enable to calculate and set the a more precise LOD floating point significance in member FMassLODResultInfo::LODSignificance.
@@ -65,13 +66,31 @@ struct FMassCombinedLODLogic : public FLODDefaultLogic
 };
 
 /**
- * This is the expected member variables for the TArrayView<FMassLODFragment> when calling the TMassLODCalculator and TMassLODtickRateController methods
- 
-struct FMassLODFragment
+ * TMassLODCollector outputs
+ *
+struct FMassViewerInfoFragment
 {
-	// Saved closest viewer distance
+	// Closest viewer distance  (Always needed)
 	float ClosestViewerDistanceSq;
 
+	// Square distances to each valid viewers (Always needed)
+	TStaticArray<float, UE::MassLOD::MaxNumOfViewers> DistanceToViewerSq;
+
+	// @Todo optimize by adding a boolean in the FLODLogic to enable the storing store that only when wanted
+	// Closest distance to frustum (Required only when FLODLogic::bDoVisibilityLogic is enabled)
+	float ClosestDistanceToFrustum;
+
+	// @Todo optimize by adding a boolean in the FLODDefaultLogic to enable the storing store that only when wanted
+	// Distances to each valid viewers frustums (Required only when FLODLogic::bDoVisibilityLogic)
+	TStaticArray<float, UE::MassLOD::MaxNumOfViewers> DistanceToFrustum;
+};
+*/
+
+/**
+ * TMassLODCalculator outputs
+ *
+ struct FMassLODFragment
+{
 	// LOD information
 	TEnumAsByte<EMassLOD::Type> LOD;
 	TEnumAsByte<EMassLOD::Type> PrevLOD;
@@ -80,16 +99,21 @@ struct FMassLODFragment
 	TStaticArray<EMassLOD::Type, UE::MassLOD::MaxNumOfViewers> LODPerViewer;
 	TStaticArray<EMassLOD::Type, UE::MassLOD::MaxNumOfViewers> PrevLODPerViewer;
 
+	// Floating point LOD value, scaling from 0 to 3, 0 highest LOD and 3 being completely off LOD 
+	// (Required only when FLODLogic::bCalculateLODSignificance is enabled)
+	float LODSignificance = 0.0f; // 
+
 	// Visibility information (Required only when FLODLogic::bDoVisibilityLogic is enabled)
 	bool bIsVisibleByAViewer;
 	bool bWasVisibleByAViewer;
 	bool bIsInVisibleRange;
 	bool bWasInVisibleRange;
 
-	// Floating point significance (Required only when FLODLogic::bCalculateLODSignificance is enabled)
-	// scaling from 0 to 3, 0 highest LOD and 3 being completely off LOD
-	float LODSignificance = 0.0f;
+	// Visibility information per viewer (Only when FLODLogic::bDoVisibilityLogic is enabled)
+	TStaticArray<bool, UE::MassLOD::MaxNumOfViewers> bIsVisibleByViewer;
+	TStaticArray<bool, UE::MassLOD::MaxNumOfViewers> bWasVisibleByViewer;
 
+	// @Todo move to its own fragment remove the FLODLogic::bDoVariableTickRate
 	// Accumulated DeltaTime (Required only when FLODLogic::bDoVariableTickRate is enalbed)
 	float DeltaTime = 0.0f;
 	float LastTickedTime = 0.0f;
@@ -120,16 +144,32 @@ struct MASSLOD_API FMassLODBaseLogic
 protected:
 	void CacheViewerInformation(TConstArrayView<FViewerInfo> ViewerInfos, const bool bLocalViewersOnly);
 
+	// Per viewer LOD information conditional fragment accessors
+	DECLARE_CONDITIONAL_MEMBER_ARRAY_ACCESSORS(Condition, EMassLOD::Type, LODPerViewer);
+	DECLARE_CONDITIONAL_MEMBER_ARRAY_ACCESSORS(Condition, EMassLOD::Type, PrevLODPerViewer);
+
+	// LOD Significance conditional fragment accessors
+	DECLARE_CONDITIONAL_MEMBER_ACCESSORS(Condition, float, LODSignificance);
+
+	// Visibility conditional fragment accessors
+	DECLARE_CONDITIONAL_MEMBER_ACCESSORS(Condition, float, ClosestDistanceToFrustum);
 	DECLARE_CONDITIONAL_MEMBER_ACCESSORS(Condition, bool, bIsVisibleByAViewer);
 	DECLARE_CONDITIONAL_MEMBER_ACCESSORS(Condition, bool, bWasVisibleByAViewer);
 	DECLARE_CONDITIONAL_MEMBER_ACCESSORS(Condition, bool, bIsInVisibleRange);
 	DECLARE_CONDITIONAL_MEMBER_ACCESSORS(Condition, bool, bWasInVisibleRange);
-	DECLARE_CONDITIONAL_MEMBER_ACCESSORS(Condition, float, LODSignificance);
+
+	// Per viewer distance conditional fragment accessors
+	DECLARE_CONDITIONAL_MEMBER_ARRAY_ACCESSORS(Condition, float, DistanceToViewerSq);
+
+	// Per viewer visibility conditional fragment accessors
+	DECLARE_CONDITIONAL_MEMBER_ARRAY_ACCESSORS(Condition, float, DistanceToFrustum);
+	DECLARE_CONDITIONAL_MEMBER_ARRAY_ACCESSORS(Condition, bool, bIsVisibleByViewer);
+	DECLARE_CONDITIONAL_MEMBER_ARRAY_ACCESSORS(Condition, bool, bWasVisibleByViewer);
+
+	// @todo, there should be a fragment dedicated to this and it would remove the need to this
+	// Tick controller conditional fragment accessors,=
 	DECLARE_CONDITIONAL_MEMBER_ACCESSORS(Condition, float, DeltaTime);
 	DECLARE_CONDITIONAL_MEMBER_ACCESSORS(Condition, float, LastTickedTime);
-	DECLARE_CONDITIONAL_MEMBER_ARRAY_ACCESSORS(Condition, bool, bIsVisibleByViewer);
-	DECLARE_CONDITIONAL_MEMBER_ARRAY_ACCESSORS(Condition, EMassLOD::Type, LODPerViewer);
-	DECLARE_CONDITIONAL_MEMBER_ARRAY_ACCESSORS(Condition, EMassLOD::Type, PrevLODPerViewer);
 
 	TArray<FViewerLODInfo> Viewers;
 };
\ No newline at end of file
diff --git a/Engine/Plugins/Runtime/MassGameplay/Source/MassLOD/Public/MassLODTickRateController.h b/Engine/Plugins/Runtime/MassGameplay/Source/MassLOD/Public/MassLODTickRateController.h
--- a/Engine/Plugins/Runtime/MassGameplay/Source/MassLOD/Public/MassLODTickRateController.h
+++ b/Engine/Plugins/Runtime/MassGameplay/Source/MassLOD/Public/MassLODTickRateController.h
@@ -10,7 +10,7 @@
  * Helper struct to control LOD tick rate for each agent, 
  * It will add a fragment tag to group the agent of the same LOD together, that way the user can do tick rate logic per chunk.
  */
-template <typename FVariableTickChunkFragment, typename FLODLogic = FLODDefaultLogic >
+template <typename TVariableTickChunkFragment, typename FLODLogic = FLODDefaultLogic >
 struct TMassLODTickRateController : public FMassLODBaseLogic
 {
 public:
@@ -43,8 +43,8 @@ public:
 	 * @param Time of the simulation to use for this update
 	 * @return bool return if the chunk should be tick this frame
 	 */
-	template <typename FMassLODResultInfo>
-	bool UpdateTickRateFromLOD(FMassExecutionContext& Context, TArrayView<FMassLODResultInfo>& LODList, const float Time);
+	template <typename TLODFragment>
+	bool UpdateTickRateFromLOD(FMassExecutionContext& Context, TArrayView<TLODFragment>& LODList, const float Time);
 
 protected:
 
@@ -55,8 +55,8 @@ protected:
 	bool bShouldSpreadFirstUpdate = false;
 };
 
-template <typename FVariableTickChunkFragment, typename FLODLogic>
-void TMassLODTickRateController<FVariableTickChunkFragment, FLODLogic>::Initialize(const float InTickRates[EMassLOD::Max], const bool bInShouldSpreadFirstUpdate/* = false*/)
+template <typename TVariableTickChunkFragment, typename FLODLogic>
+void TMassLODTickRateController<TVariableTickChunkFragment, FLODLogic>::Initialize(const float InTickRates[EMassLOD::Max], const bool bInShouldSpreadFirstUpdate/* = false*/)
 {
 	checkf(InTickRates, TEXT("You need to provide tick rate values to use this class."));
 	checkf(FLODLogic::bDoVariableTickRate, TEXT("You need to enalbe bDoVariableTickRate to use this class."));
@@ -70,25 +70,25 @@ void TMassLODTickRateController<FVariableTickChunkFragment, FLODLogic>::Initiali
 	bShouldSpreadFirstUpdate = bInShouldSpreadFirstUpdate;
 }
 
-template <typename FVariableTickChunkFragment, typename FLODLogic>
-bool TMassLODTickRateController<FVariableTickChunkFragment, FLODLogic>::ShouldCalculateLODForChunk(const FMassExecutionContext& Context) const
+template <typename TVariableTickChunkFragment, typename FLODLogic>
+bool TMassLODTickRateController<TVariableTickChunkFragment, FLODLogic>::ShouldCalculateLODForChunk(const FMassExecutionContext& Context) const
 {
 	// EMassLOD::Off does not need to handle max count, so we can use ticking rate for them if available
-	const FMassVariableTickChunkFragment& ChunkData = Context.GetChunkFragment<FVariableTickChunkFragment>();
+	const FMassVariableTickChunkFragment& ChunkData = Context.GetChunkFragment<TVariableTickChunkFragment>();
 	return ChunkData.GetLOD() != EMassLOD::Off || ChunkData.ShouldTickThisFrame();
 }
 
-template <typename FVariableTickChunkFragment, typename FLODLogic>
-bool TMassLODTickRateController<FVariableTickChunkFragment, FLODLogic>::ShouldAdjustLODFromCountForChunk(const FMassExecutionContext& Context) const
+template <typename TVariableTickChunkFragment, typename FLODLogic>
+bool TMassLODTickRateController<TVariableTickChunkFragment, FLODLogic>::ShouldAdjustLODFromCountForChunk(const FMassExecutionContext& Context) const
 {
 	// EMassLOD::Off does not need to handle max count, so we can skip it
-	const FMassVariableTickChunkFragment& ChunkData = Context.GetChunkFragment<FVariableTickChunkFragment>();
+	const FMassVariableTickChunkFragment& ChunkData = Context.GetChunkFragment<TVariableTickChunkFragment>();
 	return ChunkData.GetLOD() != EMassLOD::Off;
 }
 
-template <typename FVariableTickChunkFragment, typename FLODLogic>
-template <typename FMassLODResultInfo>
-bool TMassLODTickRateController<FVariableTickChunkFragment, FLODLogic>::UpdateTickRateFromLOD(FMassExecutionContext& Context, TArrayView<FMassLODResultInfo>& LODList, const float Time)
+template <typename TVariableTickChunkFragment, typename FLODLogic>
+template <typename TLODFragment>
+bool TMassLODTickRateController<TVariableTickChunkFragment, FLODLogic>::UpdateTickRateFromLOD(FMassExecutionContext& Context, TArrayView<TLODFragment>& LODList, const float Time)
 {
 
 	bool bShouldTickThisFrame = true;
@@ -96,7 +96,7 @@ bool TMassLODTickRateController<FVariableTickChunkFragment, FLODLogic>::UpdateTi
 	const float DeltaTime = Context.GetDeltaTimeSeconds();
 	bool bFirstUpdate = false;
 
-	FMassVariableTickChunkFragment& ChunkData = Context.GetMutableChunkFragment<FVariableTickChunkFragment>();
+	FMassVariableTickChunkFragment& ChunkData = Context.GetMutableChunkFragment<TVariableTickChunkFragment>();
 	EMassLOD::Type ChunkLOD = ChunkData.GetLOD();
 	if (ChunkLOD == EMassLOD::Max)
 	{
@@ -148,7 +148,7 @@ bool TMassLODTickRateController<FVariableTickChunkFragment, FLODLogic>::UpdateTi
 		const int32 NumEntities = Context.GetNumEntities();
 		for (int32 Index = 0; Index < NumEntities; ++Index)
 		{
-			FMassLODResultInfo& EntityLOD = LODList[Index];
+			TLODFragment& EntityLOD = LODList[Index];
 			const float LastTickedTime = GetLastTickedTime<FLODLogic::bDoVariableTickRate>(EntityLOD, 0.0f);
 			const float NewDeltaTime = LastTickedTime != 0.0f ? Time - LastTickedTime : DeltaTime;
 			SetDeltaTime<FLODLogic::bDoVariableTickRate>(EntityLOD, NewDeltaTime);
diff --git a/Engine/Plugins/Runtime/MassGameplay/Source/MassLOD/Public/MassLODTypes.h b/Engine/Plugins/Runtime/MassGameplay/Source/MassLOD/Public/MassLODTypes.h
--- a/Engine/Plugins/Runtime/MassGameplay/Source/MassLOD/Public/MassLODTypes.h
+++ b/Engine/Plugins/Runtime/MassGameplay/Source/MassLOD/Public/MassLODTypes.h
@@ -76,15 +76,15 @@ struct MASSLOD_API FMassOffLODTag : public FMassTag
  * Data fragment to store the calculated distances to viewers
  */
 USTRUCT()
-struct MASSLOD_API FMassLODInfoFragment : public FMassFragment
+struct MASSLOD_API FMassViewerInfoFragment : public FMassFragment
 {
 	GENERATED_BODY()
 
-	// Distance between each viewer and entity
-	TStaticArray<float, UE::MassLOD::MaxNumOfViewers> DistanceToViewerSq;
+	// Closest viewer distance
+	float ClosestViewerDistanceSq;
 
-	// Whether the entity is visible to each viewer
-	TStaticArray<bool, UE::MassLOD::MaxNumOfViewers> bIsVisibleByViewer;
+	// Closest distance to frustum
+	float ClosestDistanceToFrustum;
 };
 
 USTRUCT()
diff --git a/Engine/Plugins/Runtime/MassGameplay/Source/MassLOD/Public/MassRepresentationLODCollectorProcessor.h b/Engine/Plugins/Runtime/MassGameplay/Source/MassLOD/Public/MassRepresentationLODCollectorProcessor.h
--- a/Engine/Plugins/Runtime/MassGameplay/Source/MassLOD/Public/MassRepresentationLODCollectorProcessor.h
+++ b/Engine/Plugins/Runtime/MassGameplay/Source/MassLOD/Public/MassRepresentationLODCollectorProcessor.h
@@ -19,12 +19,6 @@ struct FMassRepresentationLODCollectorConfig
 	UPROPERTY(EditAnywhere, Category = "Mass|LOD", config, meta = (BaseStruct = "MassTag"))
 	FInstancedStruct TagFilter;
 
-	/** Set FOV Angles to define the limits between using Visible or Base LOD Distances */
-	UPROPERTY(EditAnywhere, Category = "Mass|LOD", meta = (ClampMin = "0.0", ClampMax = "180.0", UIMin = "0.0", UIMax = "180.0"), config)
-	float FOVAnglesToDriveVisibility = 45.0f;
-	UPROPERTY(EditAnywhere, Category = "Mass|LOD", meta = (ClampMin = "0.0", UIMin = "0.0"), config)
-	float BufferHysteresisOnFOVPercentage = 10.0f;
-
 	/** Runtime data for matching the LOD config */
 	TMassLODCollector<FMassRepresentationLODLogic> RepresentationLODCollector;
 	FMassEntityQuery CloseEntityQuery;
@@ -39,11 +33,10 @@ class MASSLOD_API UMassRepresentationLODCollectorProcessor : public UMassProcess
 {
 	GENERATED_BODY()
 
-	public:
+public:
 	UMassRepresentationLODCollectorProcessor();
 
-	protected:
-	virtual void Initialize(UObject& Owner) override;
+protected:
 	virtual void ConfigureQueries() override;
 	virtual void Execute(UMassEntitySubsystem& EntitySubsystem, FMassExecutionContext& Context) override;
 
diff --git a/Engine/Plugins/Runtime/MassGameplay/Source/MassLOD/Public/MassSimulationLOD.h b/Engine/Plugins/Runtime/MassGameplay/Source/MassLOD/Public/MassSimulationLOD.h
--- a/Engine/Plugins/Runtime/MassGameplay/Source/MassLOD/Public/MassSimulationLOD.h
+++ b/Engine/Plugins/Runtime/MassGameplay/Source/MassLOD/Public/MassSimulationLOD.h
@@ -20,8 +20,8 @@ struct MASSLOD_API FDataFragment_MassSimulationLODInfo : public FMassFragment
 {
 	GENERATED_BODY()
 
-	// Distance between each viewer and entity
-	TStaticArray<float, UE::MassLOD::MaxNumOfViewers> DistanceToViewerSq;
+	// Closest viewer distance
+	float ClosestViewerDistanceSq;
 };
 
 struct FSimulationLODLogic : public FLODDefaultLogic
@@ -41,7 +41,6 @@ public:
 	UMassProcessor_MassSimulationLODViewersInfo();
 
 protected:
-	virtual void Initialize(UObject& Owner) override;
 	virtual void ConfigureQueries() override;
 	virtual void Execute(UMassEntitySubsystem& EntitySubsystem, FMassExecutionContext& Context) override;
 
diff --git a/Engine/Plugins/Runtime/MassGameplay/Source/MassReplication/Private/MassReplicationProcessorBase.cpp b/Engine/Plugins/Runtime/MassGameplay/Source/MassReplication/Private/MassReplicationProcessorBase.cpp
--- a/Engine/Plugins/Runtime/MassGameplay/Source/MassReplication/Private/MassReplicationProcessorBase.cpp
+++ b/Engine/Plugins/Runtime/MassGameplay/Source/MassReplication/Private/MassReplicationProcessorBase.cpp
@@ -35,10 +35,10 @@ UMassReplicationProcessorBase::UMassReplicationProcessorBase()
 
 void UMassReplicationProcessorBase::ConfigureQueries()
 {
-	EntityQuery.AddRequirement<FDataFragment_MassSimulationLODInfo>(EMassFragmentAccess::ReadOnly);
+	EntityQuery.AddRequirement<FMassReplicationLODInfoFragment>(EMassFragmentAccess::ReadOnly);
 	EntityQuery.AddRequirement<FMassNetworkIDFragment>(EMassFragmentAccess::ReadOnly);
 	EntityQuery.AddRequirement<FDataFragment_ReplicationTemplateID>(EMassFragmentAccess::ReadOnly);
-	EntityQuery.AddRequirement<FMassReplicationViewerLODFragment>(EMassFragmentAccess::ReadWrite);
+	EntityQuery.AddRequirement<FMassReplicationLODFragment>(EMassFragmentAccess::ReadWrite);
 	EntityQuery.AddRequirement<FMassReplicatedAgentFragment>(EMassFragmentAccess::ReadWrite);
 }
 
@@ -138,8 +138,8 @@ void UMassReplicationProcessorBase::Execute(UMassEntitySubsystem& EntitySubsyste
 
 	EntityQuery.ForEachEntityChunk(EntitySubsystem, Context, [this](FMassExecutionContext& Context)
 		{
-			const TConstArrayView<FDataFragment_MassSimulationLODInfo> ViewersInfoList = Context.GetFragmentView<FDataFragment_MassSimulationLODInfo>();
-			const TArrayView<FMassReplicationViewerLODFragment> ViewerLODList = Context.GetMutableFragmentView<FMassReplicationViewerLODFragment>();
+			const TConstArrayView<FMassReplicationLODInfoFragment> ViewersInfoList = Context.GetFragmentView<FMassReplicationLODInfoFragment>();
+			const TArrayView<FMassReplicationLODFragment> ViewerLODList = Context.GetMutableFragmentView<FMassReplicationLODFragment>();
 			LODCalculator.CalculateLOD(Context, ViewerLODList, ViewersInfoList);
 		});
 
@@ -147,8 +147,8 @@ void UMassReplicationProcessorBase::Execute(UMassEntitySubsystem& EntitySubsyste
 	{
 		EntityQuery.ForEachEntityChunk(EntitySubsystem, Context, [this](FMassExecutionContext& Context)
 			{
-				const TConstArrayView<FDataFragment_MassSimulationLODInfo> ViewersInfoList = Context.GetFragmentView<FDataFragment_MassSimulationLODInfo>();
-				const TArrayView<FMassReplicationViewerLODFragment> ViewerLODList = Context.GetMutableFragmentView<FMassReplicationViewerLODFragment>();
+				const TConstArrayView<FMassReplicationLODInfoFragment> ViewersInfoList = Context.GetFragmentView<FMassReplicationLODInfoFragment>();
+				const TArrayView<FMassReplicationLODFragment> ViewerLODList = Context.GetMutableFragmentView<FMassReplicationLODFragment>();
 				LODCalculator.AdjustLODFromCount(Context, ViewerLODList, ViewersInfoList);
 			});
 	}
diff --git a/Engine/Plugins/Runtime/MassGameplay/Source/MassReplication/Private/MassReplicationTrait.cpp b/Engine/Plugins/Runtime/MassGameplay/Source/MassReplication/Private/MassReplicationTrait.cpp
--- a/Engine/Plugins/Runtime/MassGameplay/Source/MassReplication/Private/MassReplicationTrait.cpp
+++ b/Engine/Plugins/Runtime/MassGameplay/Source/MassReplication/Private/MassReplicationTrait.cpp
@@ -20,6 +20,6 @@ void UMassReplicationTrait::BuildTemplate(FMassEntityTemplateBuildContext& Build
 
 	BuildContext.AddFragmentWithDefaultInitializer<FMassNetworkIDFragment>();
 	BuildContext.AddFragment<FMassReplicatedAgentFragment>();
-	BuildContext.AddFragment<FMassReplicationViewerLODFragment>();
-	BuildContext.AddFragment<FDataFragment_MassSimulationLODInfo>();
+	BuildContext.AddFragment<FMassReplicationLODInfoFragment>();
+	BuildContext.AddFragment<FMassReplicationLODFragment>();
 }
\ No newline at end of file
diff --git a/Engine/Plugins/Runtime/MassGameplay/Source/MassReplication/Public/MassReplicationProcessorBase.h b/Engine/Plugins/Runtime/MassGameplay/Source/MassReplication/Public/MassReplicationProcessorBase.h
--- a/Engine/Plugins/Runtime/MassGameplay/Source/MassReplication/Public/MassReplicationProcessorBase.h
+++ b/Engine/Plugins/Runtime/MassGameplay/Source/MassReplication/Public/MassReplicationProcessorBase.h
@@ -98,7 +98,7 @@ void UMassReplicationProcessorBase::CalculateClientReplication(FMassExecutionCon
 	const int32 NumEntities = Context.GetNumEntities();
 
 	TConstArrayView<FMassNetworkIDFragment> NetworkIDList = Context.GetFragmentView<FMassNetworkIDFragment>();
-	TArrayView<FMassReplicationViewerLODFragment> ViewerLODList = Context.GetMutableFragmentView<FMassReplicationViewerLODFragment>();
+	TArrayView<FMassReplicationLODFragment> ViewerLODList = Context.GetMutableFragmentView<FMassReplicationLODFragment>();
 	TArrayView<FMassReplicatedAgentFragment> ReplicatedAgentList = Context.GetMutableFragmentView<FMassReplicatedAgentFragment>();
 	TConstArrayView<FDataFragment_ReplicationTemplateID> TemplateIDList = Context.GetFragmentView<FDataFragment_ReplicationTemplateID>();
 
diff --git a/Engine/Plugins/Runtime/MassGameplay/Source/MassReplication/Public/MassReplicationTypes.h b/Engine/Plugins/Runtime/MassGameplay/Source/MassReplication/Public/MassReplicationTypes.h
--- a/Engine/Plugins/Runtime/MassGameplay/Source/MassReplication/Public/MassReplicationTypes.h
+++ b/Engine/Plugins/Runtime/MassGameplay/Source/MassReplication/Public/MassReplicationTypes.h
@@ -189,26 +189,43 @@ struct FReplicationLODLogic : public FLODDefaultLogic
 {
 	enum
 	{
-		bStoreLODPerViewer = true,
+		bStoreInfoPerViewer = true,
+		bCalculateLODPerViewer = true,
 		bMaximizeCountPerViewer = true,
 	};
 };
 
+/*
+ * Data fragment to store the calculated distances to viewers
+ */
 USTRUCT()
-struct MASSREPLICATION_API FMassReplicationViewerLODFragment : public FMassFragment
+struct MASSREPLICATION_API FMassReplicationLODInfoFragment : public FMassFragment
 {
 	GENERATED_BODY()
 
-	/** Saved closest ViewerDistance */
-	float ClosestViewerDistanceSq = FLT_MAX;
+	/** Closest viewer distance */
+	float ClosestViewerDistanceSq;
+
+	/** Distance between each viewer and entity */
+	TStaticArray<float, UE::MassLOD::MaxNumOfViewers> DistanceToViewerSq;
+};
+
+USTRUCT()
+struct MASSREPLICATION_API FMassReplicationLODFragment : public FMassFragment
+{
+	GENERATED_BODY()
 
 	/**LOD information */
-	TStaticArray<TEnumAsByte<EMassLOD::Type>, UE::MassLOD::MaxNumOfViewers> LODPerViewer;
 	TEnumAsByte<EMassLOD::Type> LOD = EMassLOD::Max;
 
 	/** Previous LOD information*/
-	TStaticArray<TEnumAsByte<EMassLOD::Type>, UE::MassLOD::MaxNumOfViewers> PrevLODPerViewer;
 	TEnumAsByte<EMassLOD::Type> PrevLOD = EMassLOD::Max;
+
+	/** Per viewer LOD information */
+	TStaticArray<TEnumAsByte<EMassLOD::Type>, UE::MassLOD::MaxNumOfViewers> PrevLODPerViewer;
+
+	/** Per viewer previous LOD information */
+	TStaticArray<TEnumAsByte<EMassLOD::Type>, UE::MassLOD::MaxNumOfViewers> LODPerViewer;
 };
 
 UCLASS()
diff --git a/Engine/Plugins/Runtime/MassGameplay/Source/MassRepresentation/Private/MassVisualizationLODProcessor.cpp b/Engine/Plugins/Runtime/MassGameplay/Source/MassRepresentation/Private/MassVisualizationLODProcessor.cpp
--- a/Engine/Plugins/Runtime/MassGameplay/Source/MassRepresentation/Private/MassVisualizationLODProcessor.cpp
+++ b/Engine/Plugins/Runtime/MassGameplay/Source/MassRepresentation/Private/MassVisualizationLODProcessor.cpp
@@ -33,7 +33,7 @@ UMassVisualizationLODProcessor::UMassVisualizationLODProcessor()
 
 void UMassVisualizationLODProcessor::ConfigureQueries()
 {
-	CloseEntityQuery.AddRequirement<FMassLODInfoFragment>(EMassFragmentAccess::ReadOnly);
+	CloseEntityQuery.AddRequirement<FMassViewerInfoFragment>(EMassFragmentAccess::ReadOnly);
 	CloseEntityQuery.AddRequirement<FMassRepresentationLODFragment>(EMassFragmentAccess::ReadWrite);
 	CloseEntityQuery.AddRequirement<FDataFragment_Transform>(EMassFragmentAccess::ReadOnly);
 
@@ -49,7 +49,7 @@ void UMassVisualizationLODProcessor::Initialize(UObject& Owner)
 {
 	Super::Initialize(Owner);
 
-	LODCalculator.Initialize(BaseLODDistance, BufferHysteresisOnDistancePercentage / 100.f, LODMaxCount, nullptr, VisibleLODDistance);
+	LODCalculator.Initialize(BaseLODDistance, BufferHysteresisOnDistancePercentage / 100.f, LODMaxCount, nullptr, DistanceToFrustum, DistanceToFrustumHysteresis, VisibleLODDistance);
 }
 
 void UMassVisualizationLODProcessor::PrepareExecution()
diff --git a/Engine/Plugins/Runtime/MassGameplay/Source/MassRepresentation/Private/MassVisualizationTrait.cpp b/Engine/Plugins/Runtime/MassGameplay/Source/MassRepresentation/Private/MassVisualizationTrait.cpp
--- a/Engine/Plugins/Runtime/MassGameplay/Source/MassRepresentation/Private/MassVisualizationTrait.cpp
+++ b/Engine/Plugins/Runtime/MassGameplay/Source/MassRepresentation/Private/MassVisualizationTrait.cpp
@@ -36,7 +36,7 @@ void UMassVisualizationTrait::BuildTemplate(FMassEntityTemplateBuildContext& Bui
 
 		BuildContext.AddDeinitializer(*(RepresentationFragmentDeinitializerClass->GetDefaultObject<UMassProcessor>()));
 
-		BuildContext.AddFragment<FMassLODInfoFragment>();
+		BuildContext.AddFragment<FMassViewerInfoFragment>();
 		BuildContext.AddFragment<FMassRepresentationLODFragment>();
 		BuildContext.AddTag<FMassVisibilityCulledByDistanceTag>();
 
diff --git a/Engine/Plugins/Runtime/MassGameplay/Source/MassRepresentation/Public/MassRepresentationFragments.h b/Engine/Plugins/Runtime/MassGameplay/Source/MassRepresentation/Public/MassRepresentationFragments.h
--- a/Engine/Plugins/Runtime/MassGameplay/Source/MassRepresentation/Public/MassRepresentationFragments.h
+++ b/Engine/Plugins/Runtime/MassGameplay/Source/MassRepresentation/Public/MassRepresentationFragments.h
@@ -15,9 +15,6 @@ struct MASSREPRESENTATION_API FMassRepresentationLODFragment : public FMassFragm
 {
 	GENERATED_BODY()
 
-	/** Saved closest ViewerDistance */
-	float ClosestViewerDistanceSq = FLT_MAX;
-
 	/** LOD information */
 	TEnumAsByte<EMassLOD::Type> LOD = EMassLOD::Max;
 	TEnumAsByte<EMassLOD::Type> PrevLOD = EMassLOD::Max;
diff --git a/Engine/Plugins/Runtime/MassGameplay/Source/MassRepresentation/Public/MassVisualizationLODProcessor.h b/Engine/Plugins/Runtime/MassGameplay/Source/MassRepresentation/Public/MassVisualizationLODProcessor.h
--- a/Engine/Plugins/Runtime/MassGameplay/Source/MassRepresentation/Public/MassVisualizationLODProcessor.h
+++ b/Engine/Plugins/Runtime/MassGameplay/Source/MassRepresentation/Public/MassVisualizationLODProcessor.h
@@ -36,7 +36,7 @@ protected:
 
 	void PrepareExecution();
 
-	template <typename TMassViewerLODInfoFragment = FMassLODInfoFragment>
+	template <typename TMassViewerLODInfoFragment = FMassViewerInfoFragment>
 	void ExecuteInternal(UMassEntitySubsystem& EntitySubsystem, FMassExecutionContext& Context);
 
 	
@@ -59,6 +59,13 @@ protected:
 	UPROPERTY(EditAnywhere, Category = "Mass|LOD", config)
 	int32 LODMaxCount[EMassLOD::Max];
 
+	/** How far away from frustum does this entities are considered visible */
+	UPROPERTY(EditAnywhere, Category = "Mass|LOD", meta = (ClampMin = "0.0", UIMin = "0.0"), config)
+	float DistanceToFrustum = 0.0f;
+	/** Once visible how much further than DistanceToFrustum does the entities need to be before being cull again */
+	UPROPERTY(EditAnywhere, Category = "Mass|LOD", meta = (ClampMin = "0.0", UIMin = "0.0"), config)
+	float DistanceToFrustumHysteresis = 0.0f;
+
 	TMassLODCalculator<FMassRepresentationLODLogic> LODCalculator;
 
 	FMassEntityQuery CloseEntityQuery;
