Redoing changes from CL-17977443 that got missed.

diff --git a/Engine/Plugins/Runtime/MassGameplay/Source/MassRepresentation/Private/MassVisualizationComponent.cpp b/Engine/Plugins/Runtime/MassGameplay/Source/MassRepresentation/Private/MassVisualizationComponent.cpp
--- a/Engine/Plugins/Runtime/MassGameplay/Source/MassRepresentation/Private/MassVisualizationComponent.cpp
+++ b/Engine/Plugins/Runtime/MassGameplay/Source/MassRepresentation/Private/MassVisualizationComponent.cpp
@@ -16,7 +16,7 @@
 
 namespace UE::MassRepresentation
 {
-	int32 GCallUpdateInstances = 0;
+	int32 GCallUpdateInstances = 1;
 	FAutoConsoleVariableRef  CVarCallUpdateInstances(TEXT("Mass.CallUpdateInstances"), GCallUpdateInstances, TEXT("Toggle between UpdateInstances and BatchUpdateTransform."));
 }  // UE::MassRepresentation
 
