Reduce max supported viewers in mass to 3, saves ~5meg and helps performances

#rnx

diff --git a/Engine/Plugins/Runtime/MassGameplay/Source/MassLOD/Public/MassLODTypes.h b/Engine/Plugins/Runtime/MassGameplay/Source/MassLOD/Public/MassLODTypes.h
--- a/Engine/Plugins/Runtime/MassGameplay/Source/MassLOD/Public/MassLODTypes.h
+++ b/Engine/Plugins/Runtime/MassGameplay/Source/MassLOD/Public/MassLODTypes.h
@@ -19,7 +19,9 @@ namespace UE::MassLOD
 	constexpr int32 DebugNumberViewersPerController = 50;
 	constexpr int32 MaxNumOfViewers = 50;
 #else
-	constexpr int32 MaxNumOfViewers = 16;
+	// @TODO Setup it up to 3 for as an optimization, but need a better way later
+	//constexpr int32 MaxNumOfViewers = 16;
+	constexpr int32 MaxNumOfViewers = 3;
 #endif // UE_ALLOW_DEBUG_REPLICATION_DUPLICATE_VIEWERS_PER_CONTROLLER
 
 	constexpr int32 MaxBucketsPerLOD = 250;
