Bulk-merge of Mass-related Frosty-side changes over from //UE5/Main

diff --git a/Engine/Plugins/Runtime/MassGameplay/Source/MassRepresentation/Private/MassRepresentationProcessor.cpp b/Engine/Plugins/Runtime/MassGameplay/Source/MassRepresentation/Private/MassRepresentationProcessor.cpp
--- a/Engine/Plugins/Runtime/MassGameplay/Source/MassRepresentation/Private/MassRepresentationProcessor.cpp
+++ b/Engine/Plugins/Runtime/MassGameplay/Source/MassRepresentation/Private/MassRepresentationProcessor.cpp
@@ -45,7 +45,7 @@ void UMassRepresentationProcessor::Initialize(UObject& Owner)
 	Super::Initialize(Owner);
 
 	World = Owner.GetWorld();
-	PipeEntitySubsystem = UWorld::GetSubsystem<UMassEntitySubsystem>(World);
+	CachedEntitySubsystem = UWorld::GetSubsystem<UMassEntitySubsystem>(World);
 	RepresentationSubsystem = UWorld::GetSubsystem<UMassRepresentationSubsystem>(World);
 
 	// Calculate the default representation when actor isn't spawned yet.
@@ -232,12 +232,12 @@ void UMassRepresentationProcessor::UpdateRepresentation(FMassExecutionContext& C
 	}
 }
 
-void UMassRepresentationProcessor::Execute(UMassEntitySubsystem& EntitySubsystem, FMassExecutionContext& Context)
+void UMassRepresentationProcessor::Execute(UMassEntitySubsystem& InEntitySubsystem, FMassExecutionContext& Context)
 {
 	check(RepresentationSubsystem);
 
 	// Visualize entities
-	EntityQuery.ForEachEntityChunk(EntitySubsystem, Context, [this](FMassExecutionContext& Context)
+	EntityQuery.ForEachEntityChunk(InEntitySubsystem, Context, [this](FMassExecutionContext& Context)
 	{
 		UpdateRepresentation(Context);
 	});
@@ -425,11 +425,11 @@ void UMassRepresentationProcessor::UpdateEntityVisibility(const FMassEntityHandl
 
 void UMassRepresentationProcessor::OnActorPreSpawn(const FMassActorSpawnRequestHandle& SpawnRequestHandle, const FStructView& SpawnRequest)
 {
-	check(PipeEntitySubsystem);
+	check(CachedEntitySubsystem);
 	check(RepresentationSubsystem);
 
 	const FMassActorSpawnRequest& MassActorSpawnRequest = SpawnRequest.Get<FMassActorSpawnRequest>();
-	const FMassEntityView View(*PipeEntitySubsystem, MassActorSpawnRequest.MassAgent);
+	const FMassEntityView View(*CachedEntitySubsystem, MassActorSpawnRequest.MassAgent);
 
 	if (FDataFragment_Actor* ActorInfo = View.GetFragmentDataPtr<FDataFragment_Actor>())
 	{
@@ -459,13 +459,13 @@ void UMassRepresentationProcessor::OnActorPreSpawn(const FMassActorSpawnRequestH
 
 EMassActorSpawnRequestAction UMassRepresentationProcessor::OnActorPostSpawn(const FMassActorSpawnRequestHandle& SpawnRequestHandle, const FStructView& SpawnRequest)
 {
-	check(PipeEntitySubsystem);
+	check(CachedEntitySubsystem);
 
 	const FMassActorSpawnRequest& MassActorSpawnRequest = SpawnRequest.Get<FMassActorSpawnRequest>();
 	checkf(MassActorSpawnRequest.SpawnedActor, TEXT("Expecting valid spawned actor"));
 
 	// Might be already done if the actor has a MassAgentComponent via the callback OnMassAgentComponentEntityAssociated on the MassRepresentationSubsystem
-	FDataFragment_Actor& ActorInfo = PipeEntitySubsystem->GetFragmentDataChecked<FDataFragment_Actor>(MassActorSpawnRequest.MassAgent);
+	FDataFragment_Actor& ActorInfo = CachedEntitySubsystem->GetFragmentDataChecked<FDataFragment_Actor>(MassActorSpawnRequest.MassAgent);
 	if (ActorInfo.IsValid())
 	{
 		// If already set, make sure it is pointing to the same actor.
diff --git a/Engine/Plugins/Runtime/MassGameplay/Source/MassRepresentation/Public/MassRepresentationProcessor.h b/Engine/Plugins/Runtime/MassGameplay/Source/MassRepresentation/Public/MassRepresentationProcessor.h
--- a/Engine/Plugins/Runtime/MassGameplay/Source/MassRepresentation/Public/MassRepresentationProcessor.h
+++ b/Engine/Plugins/Runtime/MassGameplay/Source/MassRepresentation/Public/MassRepresentationProcessor.h
@@ -191,7 +191,7 @@ protected:
 	UWorld* World;
 
 	UPROPERTY(Transient)
-	UMassEntitySubsystem* PipeEntitySubsystem;
+	UMassEntitySubsystem* CachedEntitySubsystem;
 
 	/** A cache pointer to the representation subsystem */
 	UPROPERTY(Transient)
