Split the tick rate fragment out of the LODFragment

#rnx

diff --git a/Engine/Plugins/AI/MassAI/Source/MassAIMovement/Private/MassZoneGraphMovementProcessors.cpp b/Engine/Plugins/AI/MassAI/Source/MassAIMovement/Private/MassZoneGraphMovementProcessors.cpp
--- a/Engine/Plugins/AI/MassAI/Source/MassAIMovement/Private/MassZoneGraphMovementProcessors.cpp
+++ b/Engine/Plugins/AI/MassAI/Source/MassAIMovement/Private/MassZoneGraphMovementProcessors.cpp
@@ -191,6 +191,7 @@ void UMassZoneGraphPathFollowProcessor::ConfigureQueries()
 	EntityQuery_Conditional.AddRequirement<FMassZoneGraphLaneLocationFragment>(EMassFragmentAccess::ReadWrite);
 	EntityQuery_Conditional.AddRequirement<FMassMoveTargetFragment>(EMassFragmentAccess::ReadWrite);
 	EntityQuery_Conditional.AddRequirement<FMassSimulationLODFragment>(EMassFragmentAccess::ReadOnly, EMassFragmentPresence::Optional);
+	EntityQuery_Conditional.AddRequirement<FMassSimulationVariableTickFragment>(EMassFragmentAccess::ReadOnly, EMassFragmentPresence::Optional);
 
 	EntityQuery_Conditional.AddChunkRequirement<FMassSimulationVariableTickChunkFragment>(EMassFragmentAccess::ReadOnly, EMassFragmentPresence::Optional);
 	EntityQuery_Conditional.SetChunkFilter(&FMassSimulationVariableTickChunkFragment::ShouldTickChunkThisFrame);
@@ -207,13 +208,15 @@ void UMassZoneGraphPathFollowProcessor::Execute(UMassEntitySubsystem& EntitySubs
 	TArray<FMassEntityHandle> EntitiesToSignalLaneChanged;
 
 	EntityQuery_Conditional.ForEachEntityChunk(EntitySubsystem, Context, [this, &EntitiesToSignalPathDone, &EntitiesToSignalLaneChanged](FMassExecutionContext& Context)
-		{
+	{
 		const int32 NumEntities = Context.GetNumEntities();
 		const TArrayView<FMassZoneGraphShortPathFragment> ShortPathList = Context.GetMutableFragmentView<FMassZoneGraphShortPathFragment>();
 		const TArrayView<FMassZoneGraphLaneLocationFragment> LaneLocationList = Context.GetMutableFragmentView<FMassZoneGraphLaneLocationFragment>();
 		const TArrayView<FMassMoveTargetFragment> MoveTargetList = Context.GetMutableFragmentView<FMassMoveTargetFragment>();
 		const TConstArrayView<FMassSimulationLODFragment> SimLODList = Context.GetFragmentView<FMassSimulationLODFragment>();
 		const bool bHasLOD = (SimLODList.Num() > 0);
+		const TConstArrayView<FMassSimulationVariableTickFragment> SimVariableTickList = Context.GetFragmentView<FMassSimulationVariableTickFragment>();
+		const bool bHasVariableTick = (SimVariableTickList.Num() > 0);
 		const float WorldDeltaTime = Context.GetDeltaTimeSeconds();
 
 		for (int32 EntityIndex = 0; EntityIndex < NumEntities; ++EntityIndex)
@@ -222,7 +225,7 @@ void UMassZoneGraphPathFollowProcessor::Execute(UMassEntitySubsystem& EntitySubs
 			FMassZoneGraphLaneLocationFragment& LaneLocation = LaneLocationList[EntityIndex];
 			FMassMoveTargetFragment& MoveTarget = MoveTargetList[EntityIndex];
 			const FMassEntityHandle Entity = Context.GetEntity(EntityIndex);
-			const float DeltaTime = bHasLOD ? SimLODList[EntityIndex].DeltaTime : WorldDeltaTime;
+			const float DeltaTime = bHasVariableTick ? SimVariableTickList[EntityIndex].DeltaTime : WorldDeltaTime;
 
 			bool bDisplayDebug = false;
 #if WITH_MASSGAMEPLAY_DEBUG && UNSAFE_FOR_MT // this will result in bDisplayDebug == false and disabling of all the vlogs below
diff --git a/Engine/Plugins/Runtime/MassGameplay/Source/MassLOD/Private/MassLODCollectorProcessor.cpp b/Engine/Plugins/Runtime/MassGameplay/Source/MassLOD/Private/MassLODCollectorProcessor.cpp
--- a/Engine/Plugins/Runtime/MassGameplay/Source/MassLOD/Private/MassLODCollectorProcessor.cpp
+++ b/Engine/Plugins/Runtime/MassGameplay/Source/MassLOD/Private/MassLODCollectorProcessor.cpp
@@ -30,7 +30,7 @@ void UMassLODCollectorProcessor::ConfigureQueries()
 		LODConfig.OnLODEntityQuery.AddTagRequirement<FMassOffLODTag>(EMassFragmentPresence::None);
 		LODConfig.OffLODEntityQuery_Conditional = BaseQuery;
 		LODConfig.OffLODEntityQuery_Conditional.AddTagRequirement<FMassOffLODTag>(EMassFragmentPresence::All);
-		LODConfig.OffLODEntityQuery_Conditional.AddChunkRequirement<FMassSimulationVariableTickChunkFragment>(EMassFragmentAccess::ReadOnly);
+		LODConfig.OffLODEntityQuery_Conditional.AddChunkRequirement<FMassSimulationVariableTickChunkFragment>(EMassFragmentAccess::ReadOnly, EMassFragmentPresence::Optional);
 		LODConfig.OffLODEntityQuery_Conditional.SetChunkFilter(&FMassSimulationVariableTickChunkFragment::ShouldTickChunkThisFrame);
 
 		LODConfig.CloseEntityQuery = BaseQuery;
@@ -53,12 +53,12 @@ void UMassLODCollectorProcessor::ConfigureQueries()
 		LODConfig.FarAndOffLODEntityQuery_Conditional = BaseQuery;
 		LODConfig.FarAndOffLODEntityQuery_Conditional.AddTagRequirement<FMassVisibilityCulledByDistanceTag>(EMassFragmentPresence::All);
 		LODConfig.FarAndOffLODEntityQuery_Conditional.AddTagRequirement<FMassOffLODTag>(EMassFragmentPresence::All);
-		LODConfig.FarAndOffLODEntityQuery_Conditional.AddChunkRequirement<FMassSimulationVariableTickChunkFragment>(EMassFragmentAccess::ReadOnly);
+		LODConfig.FarAndOffLODEntityQuery_Conditional.AddChunkRequirement<FMassSimulationVariableTickChunkFragment>(EMassFragmentAccess::ReadOnly, EMassFragmentPresence::Optional);
 		LODConfig.FarAndOffLODEntityQuery_Conditional.AddChunkRequirement<FMassVisualizationChunkFragment>(EMassFragmentAccess::ReadOnly);
 		LODConfig.FarAndOffLODEntityQuery_Conditional.SetChunkFilter([](const FMassExecutionContext& Context)
 			{
 				return Context.GetChunkFragment<FMassVisualizationChunkFragment>().ShouldUpdateVisualization()
-					|| Context.GetChunkFragment<FMassSimulationVariableTickChunkFragment>().ShouldTickThisFrame();
+					|| FMassSimulationVariableTickChunkFragment::ShouldTickChunkThisFrame(Context);
 			});
 	}
 }
diff --git a/Engine/Plugins/Runtime/MassGameplay/Source/MassLOD/Private/MassLODTrait.cpp b/Engine/Plugins/Runtime/MassGameplay/Source/MassLOD/Private/MassLODTrait.cpp
--- a/Engine/Plugins/Runtime/MassGameplay/Source/MassLOD/Private/MassLODTrait.cpp
+++ b/Engine/Plugins/Runtime/MassGameplay/Source/MassLOD/Private/MassLODTrait.cpp
@@ -10,10 +10,14 @@ void UMassSimulationLODTrait::BuildTemplate(FMassEntityTemplateBuildContext& Bui
 {
 	BuildContext.AddFragmentWithDefaultInitializer<FDataFragment_Transform>();
 	BuildContext.AddFragmentWithDefaultInitializer<FMassViewerInfoFragment>();
-	BuildContext.AddFragmentWithDefaultInitializer<FMassSimulationLODFragment>();
-	BuildContext.AddFragmentWithDefaultInitializer<FDataFragment_MassSimulationLODInfo>();
-	BuildContext.AddTag<FMassOffLODTag>();
-	BuildContext.AddChunkFragment<FMassSimulationVariableTickChunkFragment>();
+	FMassSimulationLODFragment& LODFragment = BuildContext.AddFragment_GetRef<FMassSimulationLODFragment>();
+
+	// Start all simulation LOD in the Off 
+	if(Config.bSetLODTags || bEnableVariableTicking)
+	{
+		LODFragment.LOD = EMassLOD::Off;
+		BuildContext.AddTag<FMassOffLODTag>();
+	}
 
 	UMassEntitySubsystem* EntitySubsystem = UWorld::GetSubsystem<UMassEntitySubsystem>(&World);
 
@@ -22,4 +26,17 @@ void UMassSimulationLODTrait::BuildTemplate(FMassEntityTemplateBuildContext& Bui
 	BuildContext.AddConstSharedFragment(ConfigFragment);
 	FSharedStruct SharedFragment = EntitySubsystem->GetOrCreateSharedFragment<FMassSimulationLODSharedFragment>(ConfigHash, Config);
 	BuildContext.AddSharedFragment(SharedFragment);
+
+	// Variable ticking from simulation LOD
+	if(bEnableVariableTicking)
+	{
+		BuildContext.AddFragmentWithDefaultInitializer<FMassSimulationVariableTickFragment>();
+		BuildContext.AddChunkFragment<FMassSimulationVariableTickChunkFragment>();
+
+		uint32 VariableTickConfigHash = UE::StructUtils::GetStructCrc32(FConstStructView::Make(VariableTickConfig));
+		FConstSharedStruct VariableTickConfigFragment = EntitySubsystem->GetOrCreateConstSharedFragment(VariableTickConfigHash, VariableTickConfig);
+		BuildContext.AddConstSharedFragment(VariableTickConfigFragment);
+		FSharedStruct VariableTickSharedFragment = EntitySubsystem->GetOrCreateSharedFragment<FMassSimulationVariableTickSharedFragment>(VariableTickConfigHash, VariableTickConfig);
+		BuildContext.AddSharedFragment(VariableTickSharedFragment);
+	}
 }
diff --git a/Engine/Plugins/Runtime/MassGameplay/Source/MassLOD/Private/MassSimulationLOD.cpp b/Engine/Plugins/Runtime/MassGameplay/Source/MassLOD/Private/MassSimulationLOD.cpp
--- a/Engine/Plugins/Runtime/MassGameplay/Source/MassLOD/Private/MassSimulationLOD.cpp
+++ b/Engine/Plugins/Runtime/MassGameplay/Source/MassLOD/Private/MassSimulationLOD.cpp
@@ -51,7 +51,13 @@ FMassSimulationLODConfig::FMassSimulationLODConfig()
 	LODMaxCount[EMassLOD::Medium] = 200;
 	LODMaxCount[EMassLOD::Low] = 300;
 	LODMaxCount[EMassLOD::Off] = INT_MAX;
+}
 
+//-----------------------------------------------------------------------------
+// FMassSimulationVariableTickConfig
+//-----------------------------------------------------------------------------
+FMassSimulationVariableTickConfig::FMassSimulationVariableTickConfig()
+{
 	TickRates[EMassLOD::High] = 0.0f;
 	TickRates[EMassLOD::Medium] = 0.5f;
 	TickRates[EMassLOD::Low] = 1.0f;
@@ -64,7 +70,14 @@ FMassSimulationLODConfig::FMassSimulationLODConfig()
 FMassSimulationLODSharedFragment::FMassSimulationLODSharedFragment(const FMassSimulationLODConfig& LODConfig)
 {
 	LODCalculator.Initialize(LODConfig.LODDistance, LODConfig.BufferHysteresisOnDistancePercentage / 100.0f, LODConfig.LODMaxCount);
-	LODTickRateController.Initialize(LODConfig.TickRates, LODConfig.bSpreadFirstSimulationUpdate);
+}
+
+//-----------------------------------------------------------------------------
+// FMassSimulationLODSharedFragment
+//-----------------------------------------------------------------------------
+FMassSimulationVariableTickSharedFragment::FMassSimulationVariableTickSharedFragment(const FMassSimulationVariableTickConfig& TickRateConfig)
+{
+	LODTickRateController.Initialize(TickRateConfig.TickRates, TickRateConfig.bSpreadFirstSimulationUpdate);
 }
 
 //-----------------------------------------------------------------------------
@@ -89,16 +102,13 @@ void UMassSimulationLODProcessor::ConfigureQueries()
 	EntityQuery.AddRequirement<FDataFragment_Transform>(EMassFragmentAccess::ReadOnly);
 	EntityQuery.AddRequirement<FMassViewerInfoFragment>(EMassFragmentAccess::ReadOnly);
 	EntityQuery.AddRequirement<FMassSimulationLODFragment>(EMassFragmentAccess::ReadWrite);
-	EntityQuery.AddChunkRequirement<FMassSimulationVariableTickChunkFragment>(EMassFragmentAccess::ReadOnly);
 	EntityQuery.AddConstSharedRequirement<FMassSimulationLODConfig>();
 	EntityQuery.AddSharedRequirement<FMassSimulationLODSharedFragment>(EMassFragmentAccess::ReadWrite);
+	EntityQuery.AddChunkRequirement<FMassSimulationVariableTickChunkFragment>(EMassFragmentAccess::ReadOnly, EMassFragmentPresence::Optional);
+	EntityQuery.AddSharedRequirement<FMassSimulationVariableTickSharedFragment>(EMassFragmentAccess::ReadOnly, EMassFragmentPresence::Optional);
 
 	EntityQueryCalculateLOD = EntityQuery;
-	EntityQueryCalculateLOD.SetChunkFilter([](const FMassExecutionContext& Context)
-	{
-		const FMassSimulationLODSharedFragment& LODSharedFragment = Context.GetSharedFragment<FMassSimulationLODSharedFragment>();
-		return LODSharedFragment.LODTickRateController.ShouldCalculateLODForChunk(Context);
-	});
+	EntityQueryCalculateLOD.SetChunkFilter(FMassSimulationVariableTickSharedFragment::ShouldCalculateLODForChunk);
 
 	EntityQueryAdjustDistances = EntityQuery;
 	EntityQueryAdjustDistances.SetArchetypeFilter([](const FMassExecutionContext& Context)
@@ -106,10 +116,22 @@ void UMassSimulationLODProcessor::ConfigureQueries()
 		const FMassSimulationLODSharedFragment& LODSharedFragment = Context.GetSharedFragment<FMassSimulationLODSharedFragment>();
 		return LODSharedFragment.bHasAdjustedDistancesFromCount;
 	});
-	EntityQueryCalculateLOD.SetChunkFilter([](const FMassExecutionContext& Context)
+	EntityQueryCalculateLOD.SetChunkFilter(&FMassSimulationVariableTickSharedFragment::ShouldAdjustLODFromCountForChunk);
+
+	EntityQueryVariableTick.AddRequirement<FMassSimulationLODFragment>(EMassFragmentAccess::ReadOnly);
+	EntityQueryVariableTick.AddRequirement<FMassSimulationVariableTickFragment>(EMassFragmentAccess::ReadWrite);
+	EntityQueryVariableTick.AddConstSharedRequirement<FMassSimulationVariableTickConfig>();
+	EntityQueryVariableTick.AddChunkRequirement<FMassSimulationVariableTickChunkFragment>(EMassFragmentAccess::ReadWrite);
+	EntityQueryVariableTick.AddSharedRequirement<FMassSimulationVariableTickSharedFragment>(EMassFragmentAccess::ReadOnly);
+
+	// In case where the variableTick isn't enabled, we might need to set LOD tags as if the users still wants them
+	EntityQuerySetLODTag.AddRequirement<FMassSimulationLODFragment>(EMassFragmentAccess::ReadOnly);
+	EntityQuerySetLODTag.AddRequirement<FMassSimulationVariableTickFragment>(EMassFragmentAccess::ReadWrite, EMassFragmentPresence::None);
+	EntityQuerySetLODTag.AddConstSharedRequirement<FMassSimulationLODConfig>();
+	EntityQuerySetLODTag.SetArchetypeFilter([](const FMassExecutionContext& Context)
 	{
-		const FMassSimulationLODSharedFragment& LODSharedFragment = Context.GetSharedFragment<FMassSimulationLODSharedFragment>();
-		return LODSharedFragment.LODTickRateController.ShouldAdjustLODFromCountForChunk(Context);
+		const FMassSimulationLODConfig& LODConfig = Context.GetConstSharedFragment<FMassSimulationLODConfig>();
+		return LODConfig.bSetLODTags;
 	});
 }
 
@@ -159,12 +181,32 @@ void UMassSimulationLODProcessor::Execute(UMassEntitySubsystem& EntitySubsystem,
 		TRACE_CPUPROFILER_EVENT_SCOPE(TEXT("VariableTickRates"))
 		check(World);
 		const float Time = World->GetTimeSeconds();
-		EntityQuery.ForEachEntityChunk(EntitySubsystem, Context, [Time](FMassExecutionContext& Context)
+		EntityQueryVariableTick.ForEachEntityChunk(EntitySubsystem, Context, [Time](FMassExecutionContext& Context)
 		{
-			FMassSimulationLODSharedFragment& LODSharedFragment = Context.GetMutableSharedFragment<FMassSimulationLODSharedFragment>();
-			TArrayView<FMassSimulationLODFragment> SimulationLODFragments = Context.GetMutableFragmentView<FMassSimulationLODFragment>();
+			FMassSimulationVariableTickSharedFragment& TickRateSharedFragment = Context.GetMutableSharedFragment<FMassSimulationVariableTickSharedFragment>();
+			TConstArrayView<FMassSimulationLODFragment> SimulationLODFragments = Context.GetMutableFragmentView<FMassSimulationLODFragment>();
+			TArrayView<FMassSimulationVariableTickFragment> SimulationVariableTickFragments = Context.GetMutableFragmentView<FMassSimulationVariableTickFragment>();
 
-			LODSharedFragment.LODTickRateController.UpdateTickRateFromLOD(Context, SimulationLODFragments, Time);
+			TickRateSharedFragment.LODTickRateController.UpdateTickRateFromLOD(Context, SimulationLODFragments, SimulationVariableTickFragments, Time);
+		});
+	}
+
+	{
+		TRACE_CPUPROFILER_EVENT_SCOPE(TEXT("SetLODTags"))
+		check(World);
+		EntityQuerySetLODTag.ForEachEntityChunk(EntitySubsystem, Context, [](FMassExecutionContext& Context)
+		{
+			TConstArrayView<FMassSimulationLODFragment> SimulationLODFragments = Context.GetMutableFragmentView<FMassSimulationLODFragment>();
+			const int32 NumEntities = Context.GetNumEntities();
+			for (int32 Index = 0; Index < NumEntities; ++Index)
+			{
+				const FMassSimulationLODFragment& EntityLOD = SimulationLODFragments[Index];
+				if (EntityLOD.PrevLOD != EntityLOD.LOD)
+				{
+					FMassEntityHandle Entity = Context.GetEntity(Index);
+					Context.Defer().PushCommand(FCommandSwapTags(Entity, UE::MassLOD::GetLODTagFromLOD(EntityLOD.PrevLOD), UE::MassLOD::GetLODTagFromLOD(EntityLOD.LOD)));
+				}
+			}
 		});
 	}
 
diff --git a/Engine/Plugins/Runtime/MassGameplay/Source/MassLOD/Public/MassLODLogic.h b/Engine/Plugins/Runtime/MassGameplay/Source/MassLOD/Public/MassLODLogic.h
--- a/Engine/Plugins/Runtime/MassGameplay/Source/MassLOD/Public/MassLODLogic.h
+++ b/Engine/Plugins/Runtime/MassGameplay/Source/MassLOD/Public/MassLODLogic.h
@@ -32,16 +32,11 @@ struct FLODDefaultLogic
 		bDoVisibilityLogic = false, // Enable to calculate visibility and apply its own LOD distances. Requires a valid InVisibleLODDistance parameter during initialization of TMassLODCalculator.
 		bCalculateLODSignificance = false, // Enable to calculate and set the a more precise LOD floating point significance in member FMassLODResultInfo::LODSignificance.
 		bLocalViewersOnly = false, // Enable to calculate LOD from LocalViewersOnly, otherwise will be done on all viewers.
-		bDoVariableTickRate = false, // Enable to update entity variable tick rate calculation
 	};
 };
 
 struct FMassSimulationLODLogic : public FLODDefaultLogic
 {
-	enum
-	{
-		bDoVariableTickRate = true,
-	};
 };
 
 struct FMassRepresentationLODLogic : public FLODDefaultLogic
@@ -58,7 +53,6 @@ struct FMassCombinedLODLogic : public FLODDefaultLogic
 {
 	enum
 	{
-		bDoVariableTickRate = true,
 		bDoVisibilityLogic = true,
 		bCalculateLODSignificance = true,
 		bLocalViewersOnly = true,
@@ -112,9 +106,15 @@ struct FMassViewerInfoFragment
 	// Visibility information per viewer (Only when FLODLogic::bDoVisibilityLogic is enabled)
 	TStaticArray<bool, UE::MassLOD::MaxNumOfViewers> bIsVisibleByViewer;
 	TStaticArray<bool, UE::MassLOD::MaxNumOfViewers> bWasVisibleByViewer;
+}
+*/
 
-	// @Todo move to its own fragment remove the FLODLogic::bDoVariableTickRate
-	// Accumulated DeltaTime (Required only when FLODLogic::bDoVariableTickRate is enalbed)
+/**
+ * TMassLODTickRateController outputs
+ *
+ struct FMassVariableTickFragment
+{
+	// Accumulated DeltaTime
 	float DeltaTime = 0.0f;
 	float LastTickedTime = 0.0f;
 };
@@ -166,10 +166,5 @@ protected:
 	DECLARE_CONDITIONAL_MEMBER_ARRAY_ACCESSORS(Condition, bool, bIsVisibleByViewer);
 	DECLARE_CONDITIONAL_MEMBER_ARRAY_ACCESSORS(Condition, bool, bWasVisibleByViewer);
 
-	// @todo, there should be a fragment dedicated to this and it would remove the need to this
-	// Tick controller conditional fragment accessors,=
-	DECLARE_CONDITIONAL_MEMBER_ACCESSORS(Condition, float, DeltaTime);
-	DECLARE_CONDITIONAL_MEMBER_ACCESSORS(Condition, float, LastTickedTime);
-
 	TArray<FViewerLODInfo> Viewers;
 };
\ No newline at end of file
diff --git a/Engine/Plugins/Runtime/MassGameplay/Source/MassLOD/Public/MassLODTickRateController.h b/Engine/Plugins/Runtime/MassGameplay/Source/MassLOD/Public/MassLODTickRateController.h
--- a/Engine/Plugins/Runtime/MassGameplay/Source/MassLOD/Public/MassLODTickRateController.h
+++ b/Engine/Plugins/Runtime/MassGameplay/Source/MassLOD/Public/MassLODTickRateController.h
@@ -43,8 +43,8 @@ public:
 	 * @param Time of the simulation to use for this update
 	 * @return bool return if the chunk should be tick this frame
 	 */
-	template <typename TLODFragment>
-	bool UpdateTickRateFromLOD(FMassExecutionContext& Context, TArrayView<TLODFragment>& LODList, const float Time);
+	template <typename TLODFragment, typename TVariableTickRateFragment>
+	bool UpdateTickRateFromLOD(FMassExecutionContext& Context, TConstArrayView<TLODFragment> LODList, TArrayView<TVariableTickRateFragment> TickRateList, const float Time);
 
 protected:
 
@@ -59,7 +59,6 @@ template <typename TVariableTickChunkFragment, typename FLODLogic>
 void TMassLODTickRateController<TVariableTickChunkFragment, FLODLogic>::Initialize(const float InTickRates[EMassLOD::Max], const bool bInShouldSpreadFirstUpdate/* = false*/)
 {
 	checkf(InTickRates, TEXT("You need to provide tick rate values to use this class."));
-	checkf(FLODLogic::bDoVariableTickRate, TEXT("You need to enalbe bDoVariableTickRate to use this class."));
 
 	// Make a copy of all the settings
 	for (int x = 0; x < EMassLOD::Max; x++)
@@ -87,10 +86,9 @@ bool TMassLODTickRateController<TVariableTickChunkFragment, FLODLogic>::ShouldAd
 }
 
 template <typename TVariableTickChunkFragment, typename FLODLogic>
-template <typename TLODFragment>
-bool TMassLODTickRateController<TVariableTickChunkFragment, FLODLogic>::UpdateTickRateFromLOD(FMassExecutionContext& Context, TArrayView<TLODFragment>& LODList, const float Time)
+template <typename TLODFragment, typename TVariableTickRateFragment>
+bool TMassLODTickRateController<TVariableTickChunkFragment, FLODLogic>::UpdateTickRateFromLOD(FMassExecutionContext& Context, TConstArrayView<TLODFragment> LODList, TArrayView<TVariableTickRateFragment> TickRateList, const float Time)
 {
-
 	bool bShouldTickThisFrame = true;
 	bool bWasChunkTicked = true;
 	const float DeltaTime = Context.GetDeltaTimeSeconds();
@@ -148,11 +146,10 @@ bool TMassLODTickRateController<TVariableTickChunkFragment, FLODLogic>::UpdateTi
 		const int32 NumEntities = Context.GetNumEntities();
 		for (int32 Index = 0; Index < NumEntities; ++Index)
 		{
-			TLODFragment& EntityLOD = LODList[Index];
-			const float LastTickedTime = GetLastTickedTime<FLODLogic::bDoVariableTickRate>(EntityLOD, 0.0f);
-			const float NewDeltaTime = LastTickedTime != 0.0f ? Time - LastTickedTime : DeltaTime;
-			SetDeltaTime<FLODLogic::bDoVariableTickRate>(EntityLOD, NewDeltaTime);
-			SetLastTickedTime<FLODLogic::bDoVariableTickRate>(EntityLOD, Time);
+			const TLODFragment& EntityLOD = LODList[Index];
+			TVariableTickRateFragment& TickRate = TickRateList[Index];
+			TickRate.DeltaTime = TickRate.LastTickedTime != 0.0f ? Time - TickRate.LastTickedTime : DeltaTime;
+			TickRate.LastTickedTime = Time;
 			if (EntityLOD.LOD != ChunkLOD)
 			{
 				FMassEntityHandle Entity = Context.GetEntity(Index);
diff --git a/Engine/Plugins/Runtime/MassGameplay/Source/MassLOD/Public/MassLODTrait.h b/Engine/Plugins/Runtime/MassGameplay/Source/MassLOD/Public/MassLODTrait.h
--- a/Engine/Plugins/Runtime/MassGameplay/Source/MassLOD/Public/MassLODTrait.h
+++ b/Engine/Plugins/Runtime/MassGameplay/Source/MassLOD/Public/MassLODTrait.h
@@ -15,6 +15,12 @@ class MASSLOD_API UMassSimulationLODTrait : public UMassEntityTraitBase
 	UPROPERTY(Category = "Config", EditAnywhere)
 	FMassSimulationLODConfig Config;
 
+	UPROPERTY(Category = "Config", EditAnywhere)
+	bool bEnableVariableTicking = false;
+
+	UPROPERTY(Category = "Config", EditAnywhere, meta = (EditCondition = "bEnableVariableTicking"))
+	FMassSimulationVariableTickConfig VariableTickConfig;
+
 protected:
 	virtual void BuildTemplate(FMassEntityTemplateBuildContext& BuildContext, UWorld& World) const override;
 };
diff --git a/Engine/Plugins/Runtime/MassGameplay/Source/MassLOD/Public/MassSimulationLOD.h b/Engine/Plugins/Runtime/MassGameplay/Source/MassLOD/Public/MassSimulationLOD.h
--- a/Engine/Plugins/Runtime/MassGameplay/Source/MassLOD/Public/MassSimulationLOD.h
+++ b/Engine/Plugins/Runtime/MassGameplay/Source/MassLOD/Public/MassSimulationLOD.h
@@ -61,6 +61,12 @@ struct MASSLOD_API FMassSimulationLODFragment : public FMassFragment
 
 	/** Previous LOD information*/
 	TEnumAsByte<EMassLOD::Type> PrevLOD = EMassLOD::Max;
+};
+
+USTRUCT()
+struct MASSLOD_API FMassSimulationVariableTickFragment : public FMassFragment
+{
+	GENERATED_BODY()
 
 	/** Accumulated delta time to use upon next tick */
 	float LastTickedTime = 0.0f;
@@ -98,27 +104,38 @@ struct FMassSimulationLODConfig : public FMassSharedFragment
 	FMassSimulationLODConfig();
 
 	/** Distance where each LOD becomes relevant */
-	UPROPERTY(EditAnywhere, Category = "Mass|LOD", config)
+	UPROPERTY(EditAnywhere, Category = "LOD", config)
 	float LODDistance[EMassLOD::Max];
 
 	/** Hysteresis percentage on delta between the LOD distances */
-	UPROPERTY(EditAnywhere, Category = "Mass|LOD", meta = (ClampMin = "0.0", UIMin = "0.0"), config)
+	UPROPERTY(EditAnywhere, Category = "LOD", meta = (ClampMin = "0.0", UIMin = "0.0"), config)
 	float BufferHysteresisOnDistancePercentage = 10.0f;
 
 	/** Maximum limit of entity per LOD */
-	UPROPERTY(EditAnywhere, Category = "Mass|LOD", config)
+	UPROPERTY(EditAnywhere, Category = "LOD", config)
 	int32 LODMaxCount[EMassLOD::Max];
 
+	/** If true, will set the associated LOD tag on the entity */
+	UPROPERTY(EditAnywhere, Category = "LOD", config)
+	bool bSetLODTags = false;
+};
+
+USTRUCT()
+struct FMassSimulationVariableTickConfig : public FMassSharedFragment
+{
+	GENERATED_BODY()
+
+	FMassSimulationVariableTickConfig();
+
 	/** Rate in seconds at which entities should update when in this LOD */
-	UPROPERTY(EditAnywhere, Category = "Mass|LOD", config)
+	UPROPERTY(EditAnywhere, Category = "VariableTick", config)
 	float TickRates[EMassLOD::Max];
 
 	/** If true, will spread the first simulation update over TickRate period */
-	UPROPERTY(EditAnywhere, Category = "Mass|Representation", config)
+	UPROPERTY(EditAnywhere, Category = "VariableTick", config)
 	bool bSpreadFirstSimulationUpdate = false;
 };
 
-
 USTRUCT()
 struct FMassSimulationLODSharedFragment : public FMassSharedFragment
 {
@@ -127,10 +144,39 @@ struct FMassSimulationLODSharedFragment : public FMassSharedFragment
 	FMassSimulationLODSharedFragment() = default;
 	FMassSimulationLODSharedFragment(const FMassSimulationLODConfig& Config);
 
-	/** Runtime data for matching the LOD config */
-	bool bHasAdjustedDistancesFromCount = false;
+	/** Runtime data for matching the simulation LOD config */
 	TMassLODCalculator<FMassSimulationLODLogic> LODCalculator;
+	bool bHasAdjustedDistancesFromCount = false;
+};
+
+USTRUCT()
+struct FMassSimulationVariableTickSharedFragment : public FMassSharedFragment
+{
+	GENERATED_BODY()
+
+	FMassSimulationVariableTickSharedFragment() = default;
+	FMassSimulationVariableTickSharedFragment(const FMassSimulationVariableTickConfig& Config);
+
+	/** Runtime data for matching the simulation tick rate config */
 	TMassLODTickRateController<FMassSimulationVariableTickChunkFragment, FMassSimulationLODLogic> LODTickRateController;
+
+	static bool ShouldCalculateLODForChunk(const FMassExecutionContext& Context)
+	{
+		if (const FMassSimulationVariableTickSharedFragment* TickRateSharedFragment = Context.GetSharedFragmentPtr<FMassSimulationVariableTickSharedFragment>())
+		{
+			return TickRateSharedFragment->LODTickRateController.ShouldCalculateLODForChunk(Context);
+		}
+		return true;
+	}
+
+	static bool ShouldAdjustLODFromCountForChunk(const FMassExecutionContext& Context)
+	{
+		if (const FMassSimulationVariableTickSharedFragment* TickRateSharedFragment = Context.GetSharedFragmentPtr<FMassSimulationVariableTickSharedFragment>())
+		{
+			return TickRateSharedFragment->LODTickRateController.ShouldAdjustLODFromCountForChunk(Context);
+		}
+		return true;
+	}
 };
 
 UCLASS(meta = (DisplayName = "Simulation LOD"))
@@ -151,4 +197,6 @@ protected:
 	FMassEntityQuery EntityQuery;
 	FMassEntityQuery EntityQueryCalculateLOD;
 	FMassEntityQuery EntityQueryAdjustDistances;
+	FMassEntityQuery EntityQueryVariableTick;
+	FMassEntityQuery EntityQuerySetLODTag;
 };
diff --git a/Engine/Plugins/Runtime/MassGameplay/Source/MassMovement/Private/MassSimpleMovementTrait.cpp b/Engine/Plugins/Runtime/MassGameplay/Source/MassMovement/Private/MassSimpleMovementTrait.cpp
--- a/Engine/Plugins/Runtime/MassGameplay/Source/MassMovement/Private/MassSimpleMovementTrait.cpp
+++ b/Engine/Plugins/Runtime/MassGameplay/Source/MassMovement/Private/MassSimpleMovementTrait.cpp
@@ -35,7 +35,7 @@ void UMassSimpleMovementProcessor::ConfigureQueries()
 	EntityQuery.AddRequirement<FDataFragment_Transform>(EMassFragmentAccess::ReadWrite);
 	EntityQuery.AddTagRequirement<FMassSimpleMovementTag>(EMassFragmentPresence::All);
 
-	EntityQuery.AddRequirement<FMassSimulationLODFragment>(EMassFragmentAccess::ReadOnly, EMassFragmentPresence::Optional);
+	EntityQuery.AddRequirement<FMassSimulationVariableTickFragment>(EMassFragmentAccess::ReadOnly, EMassFragmentPresence::Optional);
 	EntityQuery.AddChunkRequirement<FMassSimulationVariableTickChunkFragment>(EMassFragmentAccess::ReadOnly, EMassFragmentPresence::Optional);
 	EntityQuery.SetChunkFilter(&FMassSimulationVariableTickChunkFragment::ShouldTickChunkThisFrame);
 }
@@ -46,15 +46,15 @@ void UMassSimpleMovementProcessor::Execute(UMassEntitySubsystem& EntitySubsystem
 		{
 			const TConstArrayView<FMassVelocityFragment> VelocitiesList = Context.GetFragmentView<FMassVelocityFragment>();
 			const TArrayView<FDataFragment_Transform> TransformsList = Context.GetMutableFragmentView<FDataFragment_Transform>();
-			const TConstArrayView<FMassSimulationLODFragment> SimLODList = Context.GetFragmentView<FMassSimulationLODFragment>();
-			const bool bHasLOD = (SimLODList.Num() > 0);
+			const TConstArrayView<FMassSimulationVariableTickFragment> SimVariableTickList = Context.GetFragmentView<FMassSimulationVariableTickFragment>();
+			const bool bHasVariableTick = (SimVariableTickList.Num() > 0);
 			const float WorldDeltaTime = Context.GetDeltaTimeSeconds();
 		
 			for (int32 EntityIndex = 0; EntityIndex < Context.GetNumEntities(); ++EntityIndex)
 			{
 				const FMassVelocityFragment& Velocity = VelocitiesList[EntityIndex];
 				FTransform& Transform = TransformsList[EntityIndex].GetMutableTransform();
-				const float DeltaTime = bHasLOD ? SimLODList[EntityIndex].DeltaTime : WorldDeltaTime; 
+				const float DeltaTime = bHasVariableTick ? SimVariableTickList[EntityIndex].DeltaTime : WorldDeltaTime;
 				Transform.SetTranslation(Transform.GetTranslation() + Velocity.Value * DeltaTime);
 			}
 		}));
