Starting to support more than MaxNumOfViewers
Fixed replication as the LOD was broken in previous check-ins

#rnx

diff --git a/Engine/Plugins/AI/MassCrowd/Source/MassCrowd/Private/MassCrowdReplicationProcessor.cpp b/Engine/Plugins/AI/MassCrowd/Source/MassCrowd/Private/MassCrowdReplicationProcessor.cpp
--- a/Engine/Plugins/AI/MassCrowd/Source/MassCrowd/Private/MassCrowdReplicationProcessor.cpp
+++ b/Engine/Plugins/AI/MassCrowd/Source/MassCrowd/Private/MassCrowdReplicationProcessor.cpp
@@ -36,6 +36,8 @@ void UMassCrowdReplicationProcessor::ConfigureQueries()
 	FMassReplicationProcessorPositionYawHandler::AddRequirements(EntityQuery);
 	FMassReplicationProcessorPathHandler::AddRequirements(EntityQuery);
 
+	CollectViewerInfoQuery.AddTagRequirement<FTagFragment_MassCrowd>(EMassFragmentPresence::All);
+	CalculateLODQuery.AddTagRequirement<FTagFragment_MassCrowd>(EMassFragmentPresence::All);
 	EntityQuery.AddTagRequirement<FTagFragment_MassCrowd>(EMassFragmentPresence::All);
 }
 
diff --git a/Engine/Plugins/Runtime/MassGameplay/Source/MassLOD/Public/MassLODCalculator.h b/Engine/Plugins/Runtime/MassGameplay/Source/MassLOD/Public/MassLODCalculator.h
--- a/Engine/Plugins/Runtime/MassGameplay/Source/MassLOD/Public/MassLODCalculator.h
+++ b/Engine/Plugins/Runtime/MassGameplay/Source/MassLOD/Public/MassLODCalculator.h
@@ -163,7 +163,7 @@ protected:
 	FMassLODRuntimeData RuntimeData;
 
 	/** Runtime data for each viewer specific LOD calculation, used only when bMaximizeCountPerViewer is true */
-	TStaticArray<FMassLODRuntimeData, UE::MassLOD::MaxNumOfViewers> RuntimeDataPerViewer;
+	TArray<FMassLODRuntimeData> RuntimeDataPerViewer;
 };
 
 template <typename FLODLogic>
@@ -240,6 +240,7 @@ void TMassLODCalculator<FLODLogic>::PrepareExecution(TConstArrayView<FViewerInfo
 
 	if (FLODLogic::bMaximizeCountPerViewer)
 	{
+		RuntimeDataPerViewer.SetNum(Viewers.Num());
 		for (int ViewerIdx = 0; ViewerIdx < Viewers.Num(); ++ViewerIdx)
 		{
 			// Reset viewer data
@@ -331,6 +332,11 @@ void TMassLODCalculator<FLODLogic>::CalculateLOD(FMassExecutionContext& Context,
 		// Do per viewer logic if asked for
 		if (FLODLogic::bStoreInfoPerViewer)
 		{
+			SetLODPerViewerNum<FLODLogic::bCalculateLODPerViewer>(EntityLOD, Viewers.Num());
+			SetPrevLODPerViewerNum<FLODLogic::bCalculateLODPerViewer>(EntityLOD, Viewers.Num());
+			SetbIsVisibleByViewerNum<FLODLogic::bDoVisibilityLogic&& FLODLogic::bStoreInfoPerViewer>(EntityLOD, Viewers.Num());
+			SetbWasVisibleByViewerNum<FLODLogic::bDoVisibilityLogic&& FLODLogic::bStoreInfoPerViewer>(EntityLOD, Viewers.Num());
+
 			for (int ViewerIdx = 0; ViewerIdx < Viewers.Num(); ++ViewerIdx)
 			{
 				const FViewerLODInfo& Viewer = Viewers[ViewerIdx];
diff --git a/Engine/Plugins/Runtime/MassGameplay/Source/MassLOD/Public/MassLODCollector.h b/Engine/Plugins/Runtime/MassGameplay/Source/MassLOD/Public/MassLODCollector.h
--- a/Engine/Plugins/Runtime/MassGameplay/Source/MassLOD/Public/MassLODCollector.h
+++ b/Engine/Plugins/Runtime/MassGameplay/Source/MassLOD/Public/MassLODCollector.h
@@ -49,6 +49,9 @@ void TMassLODCollector<FLODLogic>::CollectLODInfo(FMassExecutionContext& Context
 		float ClosestDistanceToFrustum = FLT_MAX;
 		const TTransformFragment& EntityTransform = TranformList[EntityIdx];
 		TViewerInfoFragment& EntityViewerInfo = ViewersInfoList[EntityIdx];
+
+		SetDistanceToViewerSqNum<FLODLogic::bStoreInfoPerViewer>(EntityViewerInfo, Viewers.Num());
+		SetDistanceToFrustumNum<FLODLogic::bDoVisibilityLogic && FLODLogic::bStoreInfoPerViewer>(EntityViewerInfo, Viewers.Num());
 		for (int ViewerIdx = 0; ViewerIdx < Viewers.Num(); ++ViewerIdx)
 		{
 			const FViewerLODInfo& Viewer = Viewers[ViewerIdx];
diff --git a/Engine/Plugins/Runtime/MassGameplay/Source/MassLOD/Public/MassLODLogic.h b/Engine/Plugins/Runtime/MassGameplay/Source/MassLOD/Public/MassLODLogic.h
--- a/Engine/Plugins/Runtime/MassGameplay/Source/MassLOD/Public/MassLODLogic.h
+++ b/Engine/Plugins/Runtime/MassGameplay/Source/MassLOD/Public/MassLODLogic.h
@@ -14,6 +14,8 @@ template <bool Condition, typename TemplateClass> static FORCEINLINE void Set##M
 template <bool Condition, typename TemplateClass> static FORCEINLINE void Set##MemberName(TemplateClass& Obj, typename TEnableIf<!Condition, MemberType>::Type Value) {}
 
 #define DECLARE_CONDITIONAL_MEMBER_ARRAY_ACCESSORS( Condition, MemberType, MemberName ) \
+template <bool Condition, typename TemplateClass> static FORCEINLINE void Set##MemberName##Num(TemplateClass& Obj, typename TEnableIf< Condition, int32>::Type Num) { Obj.MemberName.SetNum(Num); } \
+template <bool Condition, typename TemplateClass> static FORCEINLINE void Set##MemberName##Num(TemplateClass& Obj, typename TEnableIf<!Condition, int32>::Type Num) {} \
 template <bool Condition, typename TemplateClass> static FORCEINLINE typename TEnableIf< Condition, MemberType>::Type Get##MemberName(TemplateClass& Obj, int32 Index, MemberType DefaultValue) { return Obj.MemberName[Index]; } \
 template <bool Condition, typename TemplateClass> static FORCEINLINE typename TEnableIf<!Condition, MemberType>::Type Get##MemberName(TemplateClass& Obj, int32 Index, MemberType DefaultValue) { return DefaultValue; } \
 template <bool Condition, typename TemplateClass> static FORCEINLINE void Set##MemberName(TemplateClass& Obj, int32 Index, typename TEnableIf< Condition, MemberType>::Type Value) { Obj.MemberName[Index] = Value; } \
@@ -27,7 +29,7 @@ struct FLODDefaultLogic
 	enum
 	{
 		bStoreInfoPerViewer = false, // Enable to store all calculated information per viewer
-		bCalculateLODPerViewer = false, // Enable to calculate and store the result LOD per viewer in the FMassLODResultInfo::LODPerViewer and FMassLODResultInfo::PrevLODPerViewer.
+		bCalculateLODPerViewer = false, // Enable to calculate and store the result LOD per viewer in the FMassLODResultInfo::LODPerViewer and FMassLODResultInfo::PrevLODPerViewer, requires bStoreInfoPerViewer to be true as well.
 		bMaximizeCountPerViewer = false, // Enable to maximize count per viewer, requires a valid InLODMaxCountPerViewer parameter during initialization of TMassLODCalculator.
 		bDoVisibilityLogic = false, // Enable to calculate visibility and apply its own LOD distances. Requires a valid InVisibleLODDistance parameter during initialization of TMassLODCalculator.
 		bCalculateLODSignificance = false, // Enable to calculate and set the a more precise LOD floating point significance in member FMassLODResultInfo::LODSignificance.
@@ -67,16 +69,14 @@ struct FMassViewerInfoFragment
 	// Closest viewer distance  (Always needed)
 	float ClosestViewerDistanceSq;
 
-	// Square distances to each valid viewers (Always needed)
-	TStaticArray<float, UE::MassLOD::MaxNumOfViewers> DistanceToViewerSq;
+	// Square distances to each valid viewers (Required when FLODLogic::bStoreInfoPerViewer is enabled)
+	TArray<float> DistanceToViewerSq;
 
-	// @Todo optimize by adding a boolean in the FLODLogic to enable the storing store that only when wanted
-	// Closest distance to frustum (Required only when FLODLogic::bDoVisibilityLogic is enabled)
+	// Closest distance to frustum (Required when FLODLogic::bDoVisibilityLogic is enabled)
 	float ClosestDistanceToFrustum;
 
-	// @Todo optimize by adding a boolean in the FLODDefaultLogic to enable the storing store that only when wanted
-	// Distances to each valid viewers frustums (Required only when FLODLogic::bDoVisibilityLogic)
-	TStaticArray<float, UE::MassLOD::MaxNumOfViewers> DistanceToFrustum;
+	// Distances to each valid viewers frustums (Required when FLODLogic::bDoVisibilityLogic and FLODLogic::bStoreInfoPerViewer are enabled)
+	TArray<float> DistanceToFrustum;
 };
 */
 
@@ -89,23 +89,23 @@ struct FMassViewerInfoFragment
 	TEnumAsByte<EMassLOD::Type> LOD;
 	TEnumAsByte<EMassLOD::Type> PrevLOD;
 
-	// Per viewer LOD information (Required only when FLODLogic::bStoreLODPerViewer is enabled)
-	TStaticArray<EMassLOD::Type, UE::MassLOD::MaxNumOfViewers> LODPerViewer;
-	TStaticArray<EMassLOD::Type, UE::MassLOD::MaxNumOfViewers> PrevLODPerViewer;
+	// Per viewer LOD information (Required when FLODLogic::bCalculateLODPerViewer is enabled)
+	TArray<EMassLOD::Type> LODPerViewer;
+	TArray<EMassLOD::Type> PrevLODPerViewer;
 
 	// Floating point LOD value, scaling from 0 to 3, 0 highest LOD and 3 being completely off LOD 
 	// (Required only when FLODLogic::bCalculateLODSignificance is enabled)
 	float LODSignificance = 0.0f; // 
 
-	// Visibility information (Required only when FLODLogic::bDoVisibilityLogic is enabled)
+	// Visibility information (Required when FLODLogic::bDoVisibilityLogic is enabled)
 	bool bIsVisibleByAViewer;
 	bool bWasVisibleByAViewer;
 	bool bIsInVisibleRange;
 	bool bWasInVisibleRange;
 
-	// Visibility information per viewer (Only when FLODLogic::bDoVisibilityLogic is enabled)
-	TStaticArray<bool, UE::MassLOD::MaxNumOfViewers> bIsVisibleByViewer;
-	TStaticArray<bool, UE::MassLOD::MaxNumOfViewers> bWasVisibleByViewer;
+	// Visibility information per viewer (Required when FLODLogic::bDoVisibilityLogic and FLODLogicbStoreInfoPerViewer are enabled)
+	TArray<bool> bIsVisibleByViewer;
+	TArray<bool> bWasVisibleByViewer;
 }
 */
 
diff --git a/Engine/Plugins/Runtime/MassGameplay/Source/MassReplication/Private/MassReplicationProcessorBase.cpp b/Engine/Plugins/Runtime/MassGameplay/Source/MassReplication/Private/MassReplicationProcessorBase.cpp
--- a/Engine/Plugins/Runtime/MassGameplay/Source/MassReplication/Private/MassReplicationProcessorBase.cpp
+++ b/Engine/Plugins/Runtime/MassGameplay/Source/MassReplication/Private/MassReplicationProcessorBase.cpp
@@ -3,6 +3,7 @@
 #include "MassReplicationProcessorBase.h"
 #include "MassClientBubbleHandler.h"
 #include "MassLODManager.h"
+#include "MassCommonFragments.h"
 
 //----------------------------------------------------------------------//
 //  UMassProcessor_Replication
@@ -35,7 +36,12 @@ UMassReplicationProcessorBase::UMassReplicationProcessorBase()
 
 void UMassReplicationProcessorBase::ConfigureQueries()
 {
-	EntityQuery.AddRequirement<FMassReplicationLODInfoFragment>(EMassFragmentAccess::ReadOnly);
+	CollectViewerInfoQuery.AddRequirement<FDataFragment_Transform>(EMassFragmentAccess::ReadOnly);
+	CollectViewerInfoQuery.AddRequirement<FMassReplicationViewerInfoFragment>(EMassFragmentAccess::ReadOnly);
+
+	CalculateLODQuery.AddRequirement<FMassReplicationViewerInfoFragment>(EMassFragmentAccess::ReadOnly);
+	CalculateLODQuery.AddRequirement<FMassReplicationLODFragment>(EMassFragmentAccess::ReadWrite);
+
 	EntityQuery.AddRequirement<FMassNetworkIDFragment>(EMassFragmentAccess::ReadOnly);
 	EntityQuery.AddRequirement<FDataFragment_ReplicationTemplateID>(EMassFragmentAccess::ReadOnly);
 	EntityQuery.AddRequirement<FMassReplicationLODFragment>(EMassFragmentAccess::ReadWrite);
@@ -69,6 +75,7 @@ void UMassReplicationProcessorBase::PrepareExecution()
 	//then call base class
 	check(LODManager);
 	const TArray<FViewerInfo>& Viewers = LODManager->GetViewers();
+	LODCollector.PrepareExecution(Viewers);
 	LODCalculator.PrepareExecution(Viewers);
 
 	const TArray<FMassClientHandle>& CurrentClientHandles = ReplicationManager->GetClientReplicationHandles();
@@ -136,18 +143,25 @@ void UMassReplicationProcessorBase::Execute(UMassEntitySubsystem& EntitySubsyste
 
 	PrepareExecution();
 
-	EntityQuery.ForEachEntityChunk(EntitySubsystem, Context, [this](FMassExecutionContext& Context)
+	CollectViewerInfoQuery.ForEachEntityChunk(EntitySubsystem, Context, [this](FMassExecutionContext& Context)
+	{
+		const TConstArrayView<FDataFragment_Transform> LocationList = Context.GetFragmentView<FDataFragment_Transform>();
+		const TArrayView<FMassReplicationViewerInfoFragment> ViewersInfoList = Context.GetMutableFragmentView<FMassReplicationViewerInfoFragment>();
+		LODCollector.CollectLODInfo(Context, LocationList, ViewersInfoList);
+	});
+
+	CalculateLODQuery.ForEachEntityChunk(EntitySubsystem, Context, [this](FMassExecutionContext& Context)
 		{
-			const TConstArrayView<FMassReplicationLODInfoFragment> ViewersInfoList = Context.GetFragmentView<FMassReplicationLODInfoFragment>();
+			const TConstArrayView<FMassReplicationViewerInfoFragment> ViewersInfoList = Context.GetFragmentView<FMassReplicationViewerInfoFragment>();
 			const TArrayView<FMassReplicationLODFragment> ViewerLODList = Context.GetMutableFragmentView<FMassReplicationLODFragment>();
 			LODCalculator.CalculateLOD(Context, ViewerLODList, ViewersInfoList);
 		});
 
 	if (LODCalculator.AdjustDistancesFromCount())
 	{
-		EntityQuery.ForEachEntityChunk(EntitySubsystem, Context, [this](FMassExecutionContext& Context)
+		CalculateLODQuery.ForEachEntityChunk(EntitySubsystem, Context, [this](FMassExecutionContext& Context)
 			{
-				const TConstArrayView<FMassReplicationLODInfoFragment> ViewersInfoList = Context.GetFragmentView<FMassReplicationLODInfoFragment>();
+				const TConstArrayView<FMassReplicationViewerInfoFragment> ViewersInfoList = Context.GetFragmentView<FMassReplicationViewerInfoFragment>();
 				const TArrayView<FMassReplicationLODFragment> ViewerLODList = Context.GetMutableFragmentView<FMassReplicationLODFragment>();
 				LODCalculator.AdjustLODFromCount(Context, ViewerLODList, ViewersInfoList);
 			});
diff --git a/Engine/Plugins/Runtime/MassGameplay/Source/MassReplication/Private/MassReplicationTrait.cpp b/Engine/Plugins/Runtime/MassGameplay/Source/MassReplication/Private/MassReplicationTrait.cpp
--- a/Engine/Plugins/Runtime/MassGameplay/Source/MassReplication/Private/MassReplicationTrait.cpp
+++ b/Engine/Plugins/Runtime/MassGameplay/Source/MassReplication/Private/MassReplicationTrait.cpp
@@ -20,6 +20,6 @@ void UMassReplicationTrait::BuildTemplate(FMassEntityTemplateBuildContext& Build
 
 	BuildContext.AddFragmentWithDefaultInitializer<FMassNetworkIDFragment>();
 	BuildContext.AddFragment<FMassReplicatedAgentFragment>();
-	BuildContext.AddFragment<FMassReplicationLODInfoFragment>();
+	BuildContext.AddFragment<FMassReplicationViewerInfoFragment>();
 	BuildContext.AddFragment<FMassReplicationLODFragment>();
 }
\ No newline at end of file
diff --git a/Engine/Plugins/Runtime/MassGameplay/Source/MassReplication/Public/MassReplicationProcessorBase.h b/Engine/Plugins/Runtime/MassGameplay/Source/MassReplication/Public/MassReplicationProcessorBase.h
--- a/Engine/Plugins/Runtime/MassGameplay/Source/MassReplication/Public/MassReplicationProcessorBase.h
+++ b/Engine/Plugins/Runtime/MassGameplay/Source/MassReplication/Public/MassReplicationProcessorBase.h
@@ -81,8 +81,11 @@ protected:
 	UPROPERTY()
 	TArray<AMassClientBubbleInfoBase*> BubbleInfos;
 
+	TMassLODCollector<FReplicationLODLogic> LODCollector;
 	TMassLODCalculator<FReplicationLODLogic> LODCalculator;
 
+	FMassEntityQuery CollectViewerInfoQuery;
+	FMassEntityQuery CalculateLODQuery;
 	FMassEntityQuery EntityQuery;
 };
 
diff --git a/Engine/Plugins/Runtime/MassGameplay/Source/MassReplication/Public/MassReplicationTypes.h b/Engine/Plugins/Runtime/MassGameplay/Source/MassReplication/Public/MassReplicationTypes.h
--- a/Engine/Plugins/Runtime/MassGameplay/Source/MassReplication/Public/MassReplicationTypes.h
+++ b/Engine/Plugins/Runtime/MassGameplay/Source/MassReplication/Public/MassReplicationTypes.h
@@ -199,7 +199,7 @@ struct FReplicationLODLogic : public FLODDefaultLogic
  * Data fragment to store the calculated distances to viewers
  */
 USTRUCT()
-struct MASSREPLICATION_API FMassReplicationLODInfoFragment : public FMassFragment
+struct MASSREPLICATION_API FMassReplicationViewerInfoFragment : public FMassFragment
 {
 	GENERATED_BODY()
 
@@ -207,7 +207,7 @@ struct MASSREPLICATION_API FMassReplicationLODInfoFragment : public FMassFragmen
 	float ClosestViewerDistanceSq;
 
 	/** Distance between each viewer and entity */
-	TStaticArray<float, UE::MassLOD::MaxNumOfViewers> DistanceToViewerSq;
+	TArray<float> DistanceToViewerSq;
 };
 
 USTRUCT()
@@ -222,10 +222,10 @@ struct MASSREPLICATION_API FMassReplicationLODFragment : public FMassFragment
 	TEnumAsByte<EMassLOD::Type> PrevLOD = EMassLOD::Max;
 
 	/** Per viewer LOD information */
-	TStaticArray<TEnumAsByte<EMassLOD::Type>, UE::MassLOD::MaxNumOfViewers> PrevLODPerViewer;
+	TArray<TEnumAsByte<EMassLOD::Type>> LODPerViewer;
 
 	/** Per viewer previous LOD information */
-	TStaticArray<TEnumAsByte<EMassLOD::Type>, UE::MassLOD::MaxNumOfViewers> LODPerViewer;
+	TArray<TEnumAsByte<EMassLOD::Type>> PrevLODPerViewer;
 };
 
 UCLASS()
