Hiding variable config when edit condition is not enabled.

#rnx

diff --git a/Engine/Plugins/Runtime/MassGameplay/Source/MassLOD/Public/MassLODTrait.h b/Engine/Plugins/Runtime/MassGameplay/Source/MassLOD/Public/MassLODTrait.h
--- a/Engine/Plugins/Runtime/MassGameplay/Source/MassLOD/Public/MassLODTrait.h
+++ b/Engine/Plugins/Runtime/MassGameplay/Source/MassLOD/Public/MassLODTrait.h
@@ -18,7 +18,7 @@ class MASSLOD_API UMassSimulationLODTrait : public UMassEntityTraitBase
 	UPROPERTY(Category = "Config", EditAnywhere)
 	bool bEnableVariableTicking = false;
 
-	UPROPERTY(Category = "Config", EditAnywhere, meta = (EditCondition = "bEnableVariableTicking"))
+	UPROPERTY(Category = "Config", EditAnywhere, meta = (EditCondition = "bEnableVariableTicking", EditConditionHides))
 	FMassSimulationVariableTickConfig VariableTickConfig;
 
 protected:
