Changed the visibility properties from a boolean into an single enum, saves memory and they were exclusive anyway

#rnx

diff --git a/Engine/Plugins/Runtime/MassGameplay/Source/MassLOD/Public/MassLODCalculator.h b/Engine/Plugins/Runtime/MassGameplay/Source/MassLOD/Public/MassLODCalculator.h
--- a/Engine/Plugins/Runtime/MassGameplay/Source/MassLOD/Public/MassLODCalculator.h
+++ b/Engine/Plugins/Runtime/MassGameplay/Source/MassLOD/Public/MassLODCalculator.h
@@ -309,21 +309,18 @@ void TMassLODCalculator<FLODLogic>::CalculateLOD(FMassExecutionContext& Context,
 		// Calculate the LOD purely upon distances
 		const TViewerInfoFragment& EntityViewersInfo = ViewersInfoList[EntityIdx];
 		TLODFragment& EntityLOD = LODList[EntityIdx];
-		bool bIsInAVisibleRange = false;
 		const float ClosestDistanceToFrustum = GetClosestDistanceToFrustum<FLODLogic::bDoVisibilityLogic>(EntityViewersInfo, FLT_MAX);
-		const bool bWasVisibleByAViewer = GetbIsVisibleByAViewer<FLODLogic::bDoVisibilityLogic>(EntityLOD, false);
-		const bool bIsVisibleByAViewer = CalculateVisibility(bWasVisibleByAViewer, ClosestDistanceToFrustum);
-
-		// Set visibility
-		SetbWasVisibleByAViewer<FLODLogic::bDoVisibilityLogic>(EntityLOD, bWasVisibleByAViewer);
-		SetbIsVisibleByAViewer<FLODLogic::bDoVisibilityLogic>(EntityLOD, bIsVisibleByAViewer);
+		const EMassVisibility PrevVisibility = GetVisibility<FLODLogic::bDoVisibilityLogic>(EntityLOD, EMassVisibility::Max);
+		const bool bIsVisibleByAViewer = CalculateVisibility(PrevVisibility == EMassVisibility::CanBeSeen, ClosestDistanceToFrustum);
+		bool bIsInAVisibleRange = false;
 
 		// Find new LOD
 		EntityLOD.PrevLOD = EntityLOD.LOD;
 		EntityLOD.LOD = ComputeLODFromSettings(EntityLOD.PrevLOD, EntityViewersInfo.ClosestViewerDistanceSq, bIsVisibleByAViewer, &bIsInAVisibleRange, RuntimeData);
 
-		SetbWasInVisibleRange<FLODLogic::bDoVisibilityLogic>(EntityLOD, GetbIsInVisibleRange<FLODLogic::bDoVisibilityLogic>(EntityLOD, false));
-		SetbIsInVisibleRange<FLODLogic::bDoVisibilityLogic>(EntityLOD, bIsInAVisibleRange);
+		// Set visibility
+		SetPrevVisibility<FLODLogic::bDoVisibilityLogic>(EntityLOD, PrevVisibility);
+		SetVisibility<FLODLogic::bDoVisibilityLogic>(EntityLOD, bIsVisibleByAViewer ? EMassVisibility::CanBeSeen : (bIsInAVisibleRange ? EMassVisibility::CulledByFrustum : EMassVisibility::CulledByDistance));
 
 		// Accumulate in buckets
 		const float LODSignificance = AccumulateCountInRuntimeData<FLODLogic::bCalculateLODSignificance>(EntityLOD.LOD, EntityViewersInfo.ClosestViewerDistanceSq, bIsVisibleByAViewer, RuntimeData);
@@ -334,8 +331,8 @@ void TMassLODCalculator<FLODLogic>::CalculateLOD(FMassExecutionContext& Context,
 		{
 			SetLODPerViewerNum<FLODLogic::bCalculateLODPerViewer>(EntityLOD, Viewers.Num());
 			SetPrevLODPerViewerNum<FLODLogic::bCalculateLODPerViewer>(EntityLOD, Viewers.Num());
-			SetbIsVisibleByViewerNum<FLODLogic::bDoVisibilityLogic&& FLODLogic::bStoreInfoPerViewer>(EntityLOD, Viewers.Num());
-			SetbWasVisibleByViewerNum<FLODLogic::bDoVisibilityLogic&& FLODLogic::bStoreInfoPerViewer>(EntityLOD, Viewers.Num());
+			SetVisibilityPerViewerNum<FLODLogic::bDoVisibilityLogic&& FLODLogic::bStoreInfoPerViewer>(EntityLOD, Viewers.Num());
+			SetPrevVisibilityPerViewerNum<FLODLogic::bDoVisibilityLogic&& FLODLogic::bStoreInfoPerViewer>(EntityLOD, Viewers.Num());
 
 			for (int ViewerIdx = 0; ViewerIdx < Viewers.Num(); ++ViewerIdx)
 			{
@@ -344,18 +341,15 @@ void TMassLODCalculator<FLODLogic>::CalculateLOD(FMassExecutionContext& Context,
 				{
 					SetLODPerViewer<FLODLogic::bCalculateLODPerViewer>(EntityLOD, ViewerIdx, EMassLOD::Max);
 					SetPrevLODPerViewer<FLODLogic::bCalculateLODPerViewer>(EntityLOD, ViewerIdx, EMassLOD::Max);
-					SetbWasVisibleByViewer<FLODLogic::bDoVisibilityLogic&& FLODLogic::bStoreInfoPerViewer>(EntityLOD, ViewerIdx, false);
-					SetbIsVisibleByViewer<FLODLogic::bDoVisibilityLogic && FLODLogic::bStoreInfoPerViewer>(EntityLOD, ViewerIdx, false);
+					SetPrevVisibilityPerViewer<FLODLogic::bDoVisibilityLogic&& FLODLogic::bStoreInfoPerViewer>(EntityLOD, ViewerIdx, EMassVisibility::Max);
+					SetVisibilityPerViewer<FLODLogic::bDoVisibilityLogic && FLODLogic::bStoreInfoPerViewer>(EntityLOD, ViewerIdx, EMassVisibility::Max);
 				}
 				if (Viewer.Handle.IsValid())
 				{
 					const float DistanceToFrustum = GetDistanceToFrustum<FLODLogic::bDoVisibilityLogic && FLODLogic::bStoreInfoPerViewer>(EntityViewersInfo, ViewerIdx, FLT_MAX);
-					const bool bWasVisibleByViewer = GetbIsVisibleByViewer<FLODLogic::bDoVisibilityLogic && FLODLogic::bStoreInfoPerViewer>(EntityLOD, ViewerIdx, false);
-					const bool bIsVisibleByViewer = CalculateVisibility(bWasVisibleByAViewer, DistanceToFrustum);
-
-					// Set visibility
-					SetbIsVisibleByViewer<FLODLogic::bDoVisibilityLogic && FLODLogic::bStoreInfoPerViewer>(EntityLOD, ViewerIdx, bIsVisibleByViewer);
-					SetbWasVisibleByViewer<FLODLogic::bDoVisibilityLogic && FLODLogic::bStoreInfoPerViewer>(EntityLOD, ViewerIdx, bWasVisibleByViewer);
+					const EMassVisibility PrevVisibilityPerViewer = GetVisibilityPerViewer<FLODLogic::bDoVisibilityLogic && FLODLogic::bStoreInfoPerViewer>(EntityLOD, ViewerIdx, EMassVisibility::Max);
+					const bool bIsVisibleByViewer = CalculateVisibility(PrevVisibilityPerViewer == EMassVisibility::CanBeSeen, DistanceToFrustum);
+					bool bIsInVisibleRange = false;
 
 					if (FLODLogic::bCalculateLODPerViewer)
 					{
@@ -363,7 +357,7 @@ void TMassLODCalculator<FLODLogic>::CalculateLOD(FMassExecutionContext& Context,
 						const EMassLOD::Type PrevLODPerViewer = GetLODPerViewer<FLODLogic::bCalculateLODPerViewer>(EntityLOD, ViewerIdx, EntityLOD.LOD);
 
 						// Find new LOD
-						const EMassLOD::Type LODPerViewer = ComputeLODFromSettings(PrevLODPerViewer, DistanceToViewerSq, bIsVisibleByViewer, &bIsInAVisibleRange, RuntimeData);
+						const EMassLOD::Type LODPerViewer = ComputeLODFromSettings(PrevLODPerViewer, DistanceToViewerSq, bIsInVisibleRange, &bIsInAVisibleRange, RuntimeData);
 
 						// Set Per viewer LOD
 						SetPrevLODPerViewer<FLODLogic::bCalculateLODPerViewer>(EntityLOD, ViewerIdx, PrevLODPerViewer);
@@ -372,9 +366,13 @@ void TMassLODCalculator<FLODLogic>::CalculateLOD(FMassExecutionContext& Context,
 						if (FLODLogic::bMaximizeCountPerViewer)
 						{
 							// Accumulate in buckets
-							AccumulateCountInRuntimeData<false>(LODPerViewer, DistanceToViewerSq, bIsVisibleByViewer, RuntimeDataPerViewer[ViewerIdx]);
+							AccumulateCountInRuntimeData<false>(LODPerViewer, DistanceToViewerSq, bIsInVisibleRange, RuntimeDataPerViewer[ViewerIdx]);
 						}
 					}
+
+					// Set visibility
+					SetPrevVisibilityPerViewer<FLODLogic::bDoVisibilityLogic&& FLODLogic::bStoreInfoPerViewer>(EntityLOD, ViewerIdx, bIsVisibleByViewer ? EMassVisibility::CanBeSeen : (bIsInVisibleRange ? EMassVisibility::CulledByFrustum : EMassVisibility::CulledByDistance));
+					SetVisibilityPerViewer<FLODLogic::bDoVisibilityLogic&& FLODLogic::bStoreInfoPerViewer>(EntityLOD, ViewerIdx, PrevVisibilityPerViewer);
 				}
 			}
 		}
@@ -535,7 +533,8 @@ void TMassLODCalculator<FLODLogic>::AdjustLODFromCount(FMassExecutionContext& Co
 				}
 
 				const float DistanceToViewerSq = GetDistanceToViewerSq<FLODLogic::bStoreInfoPerViewer>(EntityViewersInfo, ViewerIdx, FLT_MAX);
-				const bool bIsVisibleByViewer = GetbIsVisibleByViewer<FLODLogic::bDoVisibilityLogic && FLODLogic::bStoreInfoPerViewer>(EntityLOD, ViewerIdx, false);
+				// Using the prev visibility here as it was already updated for this frame in the CalculateLOD method and we want the previous one
+				const bool bIsVisibleByViewer = GetPrevVisibilityPerViewer<FLODLogic::bDoVisibilityLogic && FLODLogic::bStoreInfoPerViewer>(EntityLOD, ViewerIdx, EMassVisibility::Max) == EMassVisibility::CanBeSeen;
 				EMassLOD::Type LODPerViewer = GetLODPerViewer<FLODLogic::bCalculateLODPerViewer>(EntityLOD, ViewerIdx, EntityLOD.LOD);
 
 				LODPerViewer = ComputeLODFromSettings(LODPerViewer, DistanceToViewerSq, bIsVisibleByViewer, nullptr, RuntimeDataPerViewer[ViewerIdx]);
@@ -549,7 +548,8 @@ void TMassLODCalculator<FLODLogic>::AdjustLODFromCount(FMassExecutionContext& Co
 			}
 		}
 
-		const bool bIsVisibleByAViewer = GetbIsVisibleByAViewer<FLODLogic::bDoVisibilityLogic>(EntityLOD, false);
+		// Using the prev visibility here as it was already updated for this frame in the CalculateLOD method and we want the previous one
+		const bool bIsVisibleByAViewer = GetPrevVisibility<FLODLogic::bDoVisibilityLogic>(EntityLOD, EMassVisibility::Max) == EMassVisibility::CanBeSeen;
 		EMassLOD::Type NewLOD = ComputeLODFromSettings(EntityLOD.PrevLOD, EntityViewersInfo.ClosestViewerDistanceSq, bIsVisibleByAViewer, nullptr, RuntimeData);
 
 		// Maybe the highest of all the viewers is now lower than the global entity LOD, make sure to update the it accordingly
@@ -599,10 +599,29 @@ void TMassLODCalculator<FLODLogic>::ForceOffLOD(FMassExecutionContext& Context,
 		EntityLOD.PrevLOD = EntityLOD.LOD;
 		EntityLOD.LOD = EMassLOD::Off;
 
-		SetbWasVisibleByAViewer<FLODLogic::bDoVisibilityLogic>(EntityLOD, GetbIsVisibleByAViewer<FLODLogic::bDoVisibilityLogic>(EntityLOD, false));
-		SetbIsVisibleByAViewer<FLODLogic::bDoVisibilityLogic>(EntityLOD, false);
-		SetbWasInVisibleRange<FLODLogic::bDoVisibilityLogic>(EntityLOD, GetbIsInVisibleRange<FLODLogic::bDoVisibilityLogic>(EntityLOD, false));
-		SetbIsInVisibleRange<FLODLogic::bDoVisibilityLogic>(EntityLOD, false);
+		// Set visibility as not calculated
+		SetPrevVisibility<FLODLogic::bDoVisibilityLogic>(EntityLOD, EMassVisibility::Max);
+		SetVisibility<FLODLogic::bDoVisibilityLogic>(EntityLOD, EMassVisibility::Max);
+
+		if (FLODLogic::bStoreInfoPerViewer)
+		{
+			for (int ViewerIdx = 0; ViewerIdx < Viewers.Num(); ++ViewerIdx)
+			{
+				if (!Viewers[ViewerIdx].Handle.IsValid())
+				{
+					continue;
+				}
+
+				if(FLODLogic::bCalculateLODPerViewer)
+				{
+					SetLODPerViewer<FLODLogic::bCalculateLODPerViewer>(EntityLOD, ViewerIdx, EMassLOD::Off);
+					SetPrevLODPerViewer<FLODLogic::bCalculateLODPerViewer>(EntityLOD, ViewerIdx, EMassLOD::Off);
+				}
+				SetPrevVisibilityPerViewer<FLODLogic::bDoVisibilityLogic&& FLODLogic::bStoreInfoPerViewer>(EntityLOD, ViewerIdx, EMassVisibility::Max);
+				SetVisibilityPerViewer<FLODLogic::bDoVisibilityLogic&& FLODLogic::bStoreInfoPerViewer>(EntityLOD, ViewerIdx, EMassVisibility::Max);
+			}
+		}
+
 		SetLODSignificance<FLODLogic::bCalculateLODSignificance>(EntityLOD, float(EMassLOD::Off));
 	}
 }
diff --git a/Engine/Plugins/Runtime/MassGameplay/Source/MassLOD/Public/MassLODLogic.h b/Engine/Plugins/Runtime/MassGameplay/Source/MassLOD/Public/MassLODLogic.h
--- a/Engine/Plugins/Runtime/MassGameplay/Source/MassLOD/Public/MassLODLogic.h
+++ b/Engine/Plugins/Runtime/MassGameplay/Source/MassLOD/Public/MassLODLogic.h
@@ -98,14 +98,12 @@ struct FMassViewerInfoFragment
 	float LODSignificance = 0.0f; // 
 
 	// Visibility information (Required when FLODLogic::bDoVisibilityLogic is enabled)
-	bool bIsVisibleByAViewer;
-	bool bWasVisibleByAViewer;
-	bool bIsInVisibleRange;
-	bool bWasInVisibleRange;
+	EMassVisibility Visibility;
+	EMassVisibility PrevVisibility
 
 	// Visibility information per viewer (Required when FLODLogic::bDoVisibilityLogic and FLODLogicbStoreInfoPerViewer are enabled)
-	TArray<bool> bIsVisibleByViewer;
-	TArray<bool> bWasVisibleByViewer;
+	TArray<EMassVisibility> VisibilityPerViewer;
+	TArray<EMassVisibility> PrevVisibilityPerViewer;
 }
 */
 
@@ -153,18 +151,16 @@ protected:
 
 	// Visibility conditional fragment accessors
 	DECLARE_CONDITIONAL_MEMBER_ACCESSORS(Condition, float, ClosestDistanceToFrustum);
-	DECLARE_CONDITIONAL_MEMBER_ACCESSORS(Condition, bool, bIsVisibleByAViewer);
-	DECLARE_CONDITIONAL_MEMBER_ACCESSORS(Condition, bool, bWasVisibleByAViewer);
-	DECLARE_CONDITIONAL_MEMBER_ACCESSORS(Condition, bool, bIsInVisibleRange);
-	DECLARE_CONDITIONAL_MEMBER_ACCESSORS(Condition, bool, bWasInVisibleRange);
+	DECLARE_CONDITIONAL_MEMBER_ACCESSORS(Condition, EMassVisibility, Visibility);
+	DECLARE_CONDITIONAL_MEMBER_ACCESSORS(Condition, EMassVisibility, PrevVisibility);
 
 	// Per viewer distance conditional fragment accessors
 	DECLARE_CONDITIONAL_MEMBER_ARRAY_ACCESSORS(Condition, float, DistanceToViewerSq);
 
 	// Per viewer visibility conditional fragment accessors
 	DECLARE_CONDITIONAL_MEMBER_ARRAY_ACCESSORS(Condition, float, DistanceToFrustum);
-	DECLARE_CONDITIONAL_MEMBER_ARRAY_ACCESSORS(Condition, bool, bIsVisibleByViewer);
-	DECLARE_CONDITIONAL_MEMBER_ARRAY_ACCESSORS(Condition, bool, bWasVisibleByViewer);
+	DECLARE_CONDITIONAL_MEMBER_ARRAY_ACCESSORS(Condition, EMassVisibility, VisibilityPerViewer);
+	DECLARE_CONDITIONAL_MEMBER_ARRAY_ACCESSORS(Condition, EMassVisibility, PrevVisibilityPerViewer);
 
 	TArray<FViewerLODInfo> Viewers;
 };
\ No newline at end of file
diff --git a/Engine/Plugins/Runtime/MassGameplay/Source/MassRepresentation/Private/MassRepresentationProcessor.cpp b/Engine/Plugins/Runtime/MassGameplay/Source/MassRepresentation/Private/MassRepresentationProcessor.cpp
--- a/Engine/Plugins/Runtime/MassGameplay/Source/MassRepresentation/Private/MassRepresentationProcessor.cpp
+++ b/Engine/Plugins/Runtime/MassGameplay/Source/MassRepresentation/Private/MassRepresentationProcessor.cpp
@@ -413,8 +413,8 @@ FMassVisualizationChunkFragment& UMassRepresentationProcessor::UpdateChunkVisibi
 void UMassRepresentationProcessor::UpdateEntityVisibility(const FMassEntityHandle Entity, const FMassRepresentationFragment& Representation, const FMassRepresentationLODFragment& RepresentationLOD, FMassVisualizationChunkFragment& ChunkData, FMassExecutionContext& Context)
 {
 	// Move the visible entities together into same chunks so we can skip entire chunk when not visible as an optimization
-	const EMassVisibility Visibility = Representation.CurrentRepresentation != ERepresentationType::None ? EMassVisibility::CanBeSeen : 
-		(RepresentationLOD.bIsInVisibleRange ? EMassVisibility::CulledByFrustum : EMassVisibility::CulledByDistance);
+	const EMassVisibility Visibility = Representation.CurrentRepresentation != ERepresentationType::None ? 
+		EMassVisibility::CanBeSeen : RepresentationLOD.Visibility;
 	const EMassVisibility ChunkVisibility = ChunkData.GetVisibility();
 	if (ChunkVisibility != Visibility)
 	{
diff --git a/Engine/Plugins/Runtime/MassGameplay/Source/MassRepresentation/Public/MassRepresentationFragments.h b/Engine/Plugins/Runtime/MassGameplay/Source/MassRepresentation/Public/MassRepresentationFragments.h
--- a/Engine/Plugins/Runtime/MassGameplay/Source/MassRepresentation/Public/MassRepresentationFragments.h
+++ b/Engine/Plugins/Runtime/MassGameplay/Source/MassRepresentation/Public/MassRepresentationFragments.h
@@ -20,10 +20,8 @@ struct MASSREPRESENTATION_API FMassRepresentationLODFragment : public FMassFragm
 	TEnumAsByte<EMassLOD::Type> PrevLOD = EMassLOD::Max;
 
 	/** Visibility Info */
-	bool bIsVisibleByAViewer = false;
-	bool bWasVisibleByAViewer = false;
-	bool bIsInVisibleRange = false;
-	bool bWasInVisibleRange = false;
+	EMassVisibility Visibility = EMassVisibility::Max;
+	EMassVisibility PrevVisibility = EMassVisibility::Max;
 
 	/** Value scaling from 0 to 3, 0 highest LOD we support and 3 being completely off LOD */
 	float LODSignificance = 0.0f;
diff --git a/Engine/Plugins/Runtime/MassGameplay/Source/MassRepresentation/Public/MassRepresentationProcessor.h b/Engine/Plugins/Runtime/MassGameplay/Source/MassRepresentation/Public/MassRepresentationProcessor.h
--- a/Engine/Plugins/Runtime/MassGameplay/Source/MassRepresentation/Public/MassRepresentationProcessor.h
+++ b/Engine/Plugins/Runtime/MassGameplay/Source/MassRepresentation/Public/MassRepresentationProcessor.h
@@ -151,7 +151,8 @@ protected:
 
 	virtual float GetSpawnPriority(const FMassRepresentationLODFragment& Representation)
 	{
-		return Representation.LODSignificance;
+		// Bump up the spawning priority on the visible entities
+		return Representation.LODSignificance - (Representation.Visibility == EMassVisibility::CanBeSeen ? 1.0f : 0.0f);
 	}
 
 	/*
