Changed how spawn-time locations are being assigned to entities:
* Added MassSpawnerSettings that are used to pick the default spawn location initializer class.
* The spawning functions can now indicate a processor class that will be used to consume spawn location data
* Replaced SpawnCollection with a simplified SpawnFromConfig
* Removed FMassSpawnConfigBase and replaces it's uses with just "num to spawn" which was the only thing we were using for it -- removed a bunch of code due to that

This CL is a prerequisite to removal of initializers and deinitializers from entity template and the implementation of new de/initializer approach.

diff --git a/Engine/Plugins/Runtime/MassEntity/Source/MassEntity/Public/MassEntitySubsystem.h b/Engine/Plugins/Runtime/MassEntity/Source/MassEntity/Public/MassEntitySubsystem.h
--- a/Engine/Plugins/Runtime/MassEntity/Source/MassEntity/Public/MassEntitySubsystem.h
+++ b/Engine/Plugins/Runtime/MassEntity/Source/MassEntity/Public/MassEntitySubsystem.h
@@ -403,6 +403,7 @@ private:
 	 *  @todo this data should live somewhere else, preferably be just a parameter to Query.ForEachEntityChunk function */
 	FArchetypeChunkCollection ChunkCollection;
 	
+	/** @todo rename to "payload" */
 	FInstancedStruct AuxData;
 	float DeltaTimeSeconds = 0.0f;
 	int32 ChunkSerialModificationNumber = -1;
@@ -594,6 +595,13 @@ public:
 
 	const FInstancedStruct& GetAuxData() const { return AuxData; }
 	FInstancedStruct& GetMutableAuxData() { return AuxData; }
+	
+	template<typename TFragment>
+	bool ValidateAuxDataType() const
+	{
+		const UScriptStruct* FragmentType = GetAuxData().GetScriptStruct();
+		return FragmentType != nullptr && FragmentType == TFragment::StaticStruct();
+	}
 
 	void FlushDeferred(UMassEntitySubsystem& EntitySystem) const;
 
diff --git a/Engine/Plugins/Runtime/MassGameplay/Source/MassActors/Private/Translators/MassCharacterMovementTranslators.cpp b/Engine/Plugins/Runtime/MassGameplay/Source/MassActors/Private/Translators/MassCharacterMovementTranslators.cpp
--- a/Engine/Plugins/Runtime/MassGameplay/Source/MassActors/Private/Translators/MassCharacterMovementTranslators.cpp
+++ b/Engine/Plugins/Runtime/MassGameplay/Source/MassActors/Private/Translators/MassCharacterMovementTranslators.cpp
@@ -193,77 +193,3 @@ void UMassFragmentInitializer_NavLocation::Execute(UMassEntitySubsystem& EntityS
 			}
 		});
 }
-
-//----------------------------------------------------------------------//
-// UMassFragmentInitializer_Transform 
-//----------------------------------------------------------------------//
-UMassFragmentInitializer_Transform::UMassFragmentInitializer_Transform()
-{
-	ExecutionFlags = (int32)EProcessorExecutionFlags::All;
-	FragmentType = FDataFragment_Transform::StaticStruct();
-}
-
-void UMassFragmentInitializer_Transform::ConfigureQueries()
-{
-	EntityQuery.AddRequirement<FDataFragment_Transform>(EMassFragmentAccess::ReadWrite);
-}
-
-void UMassFragmentInitializer_Transform::Execute(UMassEntitySubsystem& EntitySubsystem, FMassExecutionContext& Context)
-{
-	const UWorld* World = EntitySubsystem.GetWorld();
-
-	check(World);
-
-	const ENetMode NetMode = World->GetNetMode();
-
-	if (NetMode != NM_Client)
-	{
-		const FInstancedStruct& AuxInput = Context.GetMutableAuxData();
-		if (!(AuxInput.IsValid() && AuxInput.GetScriptStruct()->IsChildOf(FMassSpawnAuxData::StaticStruct())))
-		{
-			UE_VLOG_UELOG(this, LogMass, Log, TEXT("Execution context has invalid AuxData or it's not FMassSpawnAuxData. Entity transforms won't be initialized."));
-			return;
-		}
-
-		TArray<FTransform>& Transforms = AuxInput.GetMutable<FMassSpawnAuxData>().Transforms;
-
-		const int32 NumSpawnTransforms = Transforms.Num();
-		if (NumSpawnTransforms == 0)
-		{
-			UE_VLOG_UELOG(this, LogMass, Error, TEXT("No spawn transforms provided. Entity transforms won't be initialized."));
-			return;
-		}
-
-		int32 NumRequiredSpawnTransforms = 0;
-		EntityQuery.ForEachEntityChunk(EntitySubsystem, Context, [&NumRequiredSpawnTransforms](const FMassExecutionContext& Context)
-			{
-				NumRequiredSpawnTransforms += Context.GetNumEntities();
-			});
-
-		const int32 NumToAdd = NumRequiredSpawnTransforms - NumSpawnTransforms;
-		if (NumToAdd > 0)
-		{
-			UE_VLOG_UELOG(this, LogMass, Warning,
-				TEXT("Not enough spawn locations provided (%d) for all entities (%d). Existing locations will be reused randomly to fill the %d missing positions."),
-				NumSpawnTransforms, NumRequiredSpawnTransforms, NumToAdd);
-
-			Transforms.AddUninitialized(NumToAdd);
-			for (int i = 0; i < NumToAdd; ++i)
-			{
-				Transforms[NumSpawnTransforms + i] = Transforms[FMath::RandRange(0, NumSpawnTransforms - 1)];
-			}
-		}
-
-		EntityQuery.ForEachEntityChunk(EntitySubsystem, Context, [&Transforms, this](FMassExecutionContext& Context)
-			{
-				const TArrayView<FDataFragment_Transform> LocationList = Context.GetMutableFragmentView<FDataFragment_Transform>();
-				const int32 NumEntities = Context.GetNumEntities();
-				for (int32 i = 0; i < NumEntities; ++i)
-				{
-					const int32 AuxIndex = FMath::RandRange(0, Transforms.Num() - 1);
-					LocationList[i].GetMutableTransform() = Transforms[AuxIndex];
-					Transforms.RemoveAtSwap(AuxIndex, 1, /*bAllowShrinking=*/false);
-				}
-			});
-	}
-}
diff --git a/Engine/Plugins/Runtime/MassGameplay/Source/MassActors/Public/Translators/MassCharacterMovementTranslators.h b/Engine/Plugins/Runtime/MassGameplay/Source/MassActors/Public/Translators/MassCharacterMovementTranslators.h
--- a/Engine/Plugins/Runtime/MassGameplay/Source/MassActors/Public/Translators/MassCharacterMovementTranslators.h
+++ b/Engine/Plugins/Runtime/MassGameplay/Source/MassActors/Public/Translators/MassCharacterMovementTranslators.h
@@ -119,18 +119,3 @@ protected:
 
 	FMassEntityQuery EntityQuery;
 };
-
-UCLASS()
-class MASSACTORS_API UMassFragmentInitializer_Transform : public UMassFragmentInitializer
-{
-	GENERATED_BODY()
-
-public:
-	UMassFragmentInitializer_Transform();
-
-protected:
-	virtual void ConfigureQueries() override;
-	virtual void Execute(UMassEntitySubsystem& EntitySubsystem, FMassExecutionContext& Context) override;
-
-	FMassEntityQuery EntityQuery;
-};
\ No newline at end of file
diff --git a/Engine/Plugins/Runtime/MassGameplay/Source/MassGameplayDebug/Private/DebugEntityTemplateBuilder.cpp b/Engine/Plugins/Runtime/MassGameplay/Source/MassGameplayDebug/Private/DebugEntityTemplateBuilder.cpp
deleted file mode 100644
diff --git a/Engine/Plugins/Runtime/MassGameplay/Source/MassGameplayDebug/Private/DebugEntityTemplateBuilder.h b/Engine/Plugins/Runtime/MassGameplay/Source/MassGameplayDebug/Private/DebugEntityTemplateBuilder.h
deleted file mode 100644
diff --git a/Engine/Plugins/Runtime/MassGameplay/Source/MassGameplayDebug/Private/MassGameplayDebugModule.cpp b/Engine/Plugins/Runtime/MassGameplay/Source/MassGameplayDebug/Private/MassGameplayDebugModule.cpp
--- a/Engine/Plugins/Runtime/MassGameplay/Source/MassGameplayDebug/Private/MassGameplayDebugModule.cpp
+++ b/Engine/Plugins/Runtime/MassGameplay/Source/MassGameplayDebug/Private/MassGameplayDebugModule.cpp
@@ -5,7 +5,6 @@
 #include "IMassGameplayDebugModule.h"
 #include "MassCommonTypes.h"
 #include "MassGameplayDebugTypes.h"
-#include "DebugEntityTemplateBuilder.h"
 
 DEFINE_LOG_CATEGORY(LogMassDebug);
 
@@ -20,7 +19,6 @@ IMPLEMENT_MODULE(FMassGameplayDebug, MassGameplayDebug)
 
 void FMassGameplayDebug::StartupModule()
 {
-	UE::Mass::DebugEntityTemplateBuilder::Register();
 }
 
 void FMassGameplayDebug::ShutdownModule()
diff --git a/Engine/Plugins/Runtime/MassGameplay/Source/MassGameplayDebug/Public/MassGameplayDebugTypes.h b/Engine/Plugins/Runtime/MassGameplay/Source/MassGameplayDebug/Public/MassGameplayDebugTypes.h
--- a/Engine/Plugins/Runtime/MassGameplay/Source/MassGameplayDebug/Public/MassGameplayDebugTypes.h
+++ b/Engine/Plugins/Runtime/MassGameplay/Source/MassGameplayDebug/Public/MassGameplayDebugTypes.h
@@ -101,25 +101,3 @@ struct FAgentDebugVisualization : public FTableRowBase
 	UPROPERTY(EditAnywhere, Category = "Mass|Debug")
 	EMassEntityDebugShape WireShape = EMassEntityDebugShape::Box;
 };
-
-/** @todo comment functionality */
-USTRUCT(DisplayName="MassSpawnProps_Debug")
-struct FMassSpawnProps : public FMassSpawnConfigBase
-{
-	GENERATED_BODY()
-
-#if WITH_EDITORONLY_DATA
-	UPROPERTY(EditAnywhere, Category = "Mass|Spawn")
-	FAgentDebugVisualization DebugShape;
-#endif // WITH_EDITORONLY_DATA
-
-	friend uint32 GetTypeHash(const FMassSpawnProps& Instance)
-	{
-#if WITH_EDITORONLY_DATA
-		return HashCombine(GetTypeHash((const FMassSpawnConfigBase&)Instance)
-				, HashCombine(GetTypeHash(Instance.DebugShape.Mesh), GetTypeHash(Instance.DebugShape.MaterialOverride)));
-#else
-		return GetTypeHash((const FMassSpawnConfigBase&)Instance);
-#endif // WITH_EDITORONLY_DATA
-	}
-};
\ No newline at end of file
diff --git a/Engine/Plugins/Runtime/MassGameplay/Source/MassReplication/Public/MassClientBubbleHandler.h b/Engine/Plugins/Runtime/MassGameplay/Source/MassReplication/Public/MassClientBubbleHandler.h
--- a/Engine/Plugins/Runtime/MassGameplay/Source/MassReplication/Public/MassClientBubbleHandler.h
+++ b/Engine/Plugins/Runtime/MassGameplay/Source/MassReplication/Public/MassClientBubbleHandler.h
@@ -532,13 +532,9 @@ void TClientBubbleHandlerBase<AgentArrayItem>::PostReplicatedAddEntitiesHelper(c
 		const FMassEntityTemplateID& TemplateID = Item.Key;
 		const TArray <typename AgentArrayItem::FReplicatedAgentType*>& AgentsSpawn = Item.Value;
 
-		FMassSpawnConfigBase Config;
-		Config.MinNumber = AgentsSpawn.Num();
-		Config.MaxNumber = AgentsSpawn.Num();
-
 		TArray<FMassEntityHandle> Entities;
 
-		SpawnerSubsystem->SpawnEntities(TemplateID, Config, FStructView(), Entities);
+		SpawnerSubsystem->SpawnEntities(TemplateID, AgentsSpawn.Num(), FStructView(), TSubclassOf<UMassProcessor>(), Entities);
 
 		const FMassEntityTemplate* MassEntityTemplate = SpawnerSubsystem->GetMassEntityTemplate(TemplateID);
 		check(MassEntityTemplate);
diff --git a/Engine/Plugins/Runtime/MassGameplay/Source/MassSpawner/Private/MassSpawnLocationProcessor.cpp b/Engine/Plugins/Runtime/MassGameplay/Source/MassSpawner/Private/MassSpawnLocationProcessor.cpp
new file mode 100644
--- /dev/null
+++ b/Engine/Plugins/Runtime/MassGameplay/Source/MassSpawner/Private/MassSpawnLocationProcessor.cpp
@@ -0,0 +1,81 @@
+// Copyright Epic Games, Inc. All Rights Reserved.
+
+#include "MassSpawnLocationProcessor.h"
+#include "MassCommonFragments.h"
+#include "MassEntitySubsystem.h"
+#include "MassSpawnerTypes.h"
+#include "Engine/World.h"
+#include "VisualLogger/VisualLogger.h"
+
+//----------------------------------------------------------------------//
+// UMassSpawnLocationProcessor 
+//----------------------------------------------------------------------//
+UMassSpawnLocationProcessor::UMassSpawnLocationProcessor()
+{
+	bAutoRegisterWithProcessingPhases = false;
+}
+
+void UMassSpawnLocationProcessor::ConfigureQueries()
+{
+	EntityQuery.AddRequirement<FDataFragment_Transform>(EMassFragmentAccess::ReadWrite);
+}
+
+void UMassSpawnLocationProcessor::Execute(UMassEntitySubsystem& EntitySubsystem, FMassExecutionContext& Context)
+{
+	if (!ensure(Context.ValidateAuxDataType<FMassSpawnAuxData>()))
+	{
+		UE_VLOG_UELOG(this, LogMass, Log, TEXT("Execution context has invalid AuxData or it's not FMassSpawnAuxData. Entity transforms won't be initialized."));
+		return;
+	}
+
+	const UWorld* World = EntitySubsystem.GetWorld();
+
+	check(World);
+
+	const ENetMode NetMode = World->GetNetMode();
+
+	if (NetMode != NM_Client)
+	{
+		FMassSpawnAuxData& AuxData = Context.GetMutableAuxData().GetMutable<FMassSpawnAuxData>();
+		TArray<FTransform>& Transforms = AuxData.Transforms;
+
+		const int32 NumSpawnTransforms = Transforms.Num();
+		if (NumSpawnTransforms == 0)
+		{
+			UE_VLOG_UELOG(this, LogMass, Error, TEXT("No spawn transforms provided. Entity transforms won't be initialized."));
+			return;
+		}
+
+		int32 NumRequiredSpawnTransforms = 0;
+		EntityQuery.ForEachEntityChunk(EntitySubsystem, Context, [&NumRequiredSpawnTransforms](const FMassExecutionContext& Context)
+			{
+				NumRequiredSpawnTransforms += Context.GetNumEntities();
+			});
+
+		const int32 NumToAdd = NumRequiredSpawnTransforms - NumSpawnTransforms;
+		if (NumToAdd > 0)
+		{
+			UE_VLOG_UELOG(this, LogMass, Warning,
+				TEXT("Not enough spawn locations provided (%d) for all entities (%d). Existing locations will be reused randomly to fill the %d missing positions."),
+				NumSpawnTransforms, NumRequiredSpawnTransforms, NumToAdd);
+
+			Transforms.AddUninitialized(NumToAdd);
+			for (int i = 0; i < NumToAdd; ++i)
+			{
+				Transforms[NumSpawnTransforms + i] = Transforms[FMath::RandRange(0, NumSpawnTransforms - 1)];
+			}
+		}
+
+		EntityQuery.ForEachEntityChunk(EntitySubsystem, Context, [&Transforms, this](FMassExecutionContext& Context)
+			{
+				const TArrayView<FDataFragment_Transform> LocationList = Context.GetMutableFragmentView<FDataFragment_Transform>();
+				const int32 NumEntities = Context.GetNumEntities();
+				for (int32 i = 0; i < NumEntities; ++i)
+				{
+					const int32 AuxIndex = FMath::RandRange(0, Transforms.Num() - 1);
+					LocationList[i].GetMutableTransform() = Transforms[AuxIndex];
+					Transforms.RemoveAtSwap(AuxIndex, 1, /*bAllowShrinking=*/false);
+				}
+			});
+	}
+}
diff --git a/Engine/Plugins/Runtime/MassGameplay/Source/MassSpawner/Private/MassSpawner.cpp b/Engine/Plugins/Runtime/MassGameplay/Source/MassSpawner/Private/MassSpawner.cpp
--- a/Engine/Plugins/Runtime/MassGameplay/Source/MassSpawner/Private/MassSpawner.cpp
+++ b/Engine/Plugins/Runtime/MassGameplay/Source/MassSpawner/Private/MassSpawner.cpp
@@ -15,6 +15,8 @@
 #include "Engine/StreamableManager.h"
 #include "Engine/AssetManager.h"
 #include "MassGameplaySettings.h"
+#include "MassSpawnLocationProcessor.h"
+
 
 namespace UE::MassSpawner
 {
@@ -448,15 +450,13 @@ void AMassSpawner::SpawnAtLocations(TConstArrayView<FVector> Locations)
 			{
 				if (const UMassEntityConfigAsset* EntityConfig = EntityType.GetEntityConfig())
 				{
-					FMassSpawnConfigBase SpawnConfig;
-					SpawnConfig.MinNumber = EntityCount;
-					SpawnConfig.MaxNumber = EntityCount;
 					const FMassEntityTemplate* EntityTemplate = EntityConfig->GetConfig().GetOrCreateEntityTemplate(*this, *EntityConfig);
 					if (EntityTemplate && EntityTemplate->IsValid())
 					{
 						FSpawnedEntities& SpawnedEntities = AllSpawnedEntities.AddDefaulted_GetRef();
 						SpawnedEntities.TemplateID = EntityTemplate->GetTemplateID();
-						SpawnerSystem->SpawnEntities(EntityTemplate->GetTemplateID(), SpawnConfig, FStructView::Make(MassSpawnAuxData), SpawnedEntities.Entities);
+						SpawnerSystem->SpawnEntities(EntityTemplate->GetTemplateID(), EntityCount, FStructView::Make(MassSpawnAuxData)
+							, UMassSpawnLocationProcessor::StaticClass(), SpawnedEntities.Entities);
 					}
 				}
 			}
diff --git a/Engine/Plugins/Runtime/MassGameplay/Source/MassSpawner/Private/MassSpawnerSubsystem.cpp b/Engine/Plugins/Runtime/MassGameplay/Source/MassSpawner/Private/MassSpawnerSubsystem.cpp
--- a/Engine/Plugins/Runtime/MassGameplay/Source/MassSpawner/Private/MassSpawnerSubsystem.cpp
+++ b/Engine/Plugins/Runtime/MassGameplay/Source/MassSpawner/Private/MassSpawnerSubsystem.cpp
@@ -14,6 +14,7 @@
 #include "MassTranslatorRegistry.h"
 #include "ProfilingDebugging/CpuProfilerTrace.h"
 #include "MassSimulationSubsystem.h"
+#include "MassEntityUtils.h"
 
 
 namespace UE::Mass::SpawnerSubsystem 
@@ -102,23 +103,7 @@ void UMassSpawnerSubsystem::PostInitialize()
 	TemplateRegistryInstance = NewObject<UMassEntityTemplateRegistry>(this);
 }
 
-bool UMassSpawnerSubsystem::SpawnEntities(const AActor& ActorInstance, const uint32 NumberToSpawn, TArray<FMassEntityHandle>& OutEntities) const
-{
-	if (NumberToSpawn == 0)
-	{
-		UE_VLOG(this, LogMassSpawner, Warning, TEXT("Trying to spawn 0 entities. This would cause inefficiency. Bailing out with result FALSE."));
-		return false;
-	}
-
-	check(TemplateRegistryInstance);
-	if (const FMassEntityTemplate* EntityTemplate = TemplateRegistryInstance->FindOrBuildInstanceTemplate(ActorInstance))
-	{
-		return SpawnEntities(*EntityTemplate, NumberToSpawn, OutEntities);
-	}
-	return false;
-}
-
-bool UMassSpawnerSubsystem::SpawnEntities(const FMassEntityTemplate& EntityTemplate, const uint32 NumberToSpawn, TArray<FMassEntityHandle>& OutEntities) const
+void UMassSpawnerSubsystem::SpawnEntities(const FMassEntityTemplate& EntityTemplate, const uint32 NumberToSpawn, TArray<FMassEntityHandle>& OutEntities)
 {
 	check(EntitySystem);
 	check(EntityTemplate.IsValid());
@@ -126,33 +111,13 @@ bool UMassSpawnerSubsystem::SpawnEntities(const FMassEntityTemplate& EntityTempl
 	if (NumberToSpawn == 0)
 	{
 		UE_VLOG(this, LogMassSpawner, Warning, TEXT("Trying to spawn 0 entities. This would cause inefficiency. Bailing out with result FALSE."));
-		return false;
+		return;
 	}
 
-	//TRACE_CPUPROFILER_EVENT_SCOPE_STR("MassSpawnerSubsystem SpawnEntities");
-	
-	OutEntities.Reset();
-	if (EntitySystem->BatchCreateEntities(EntityTemplate.GetArchetype(), NumberToSpawn, OutEntities) > 0)
-	{ 
-		UE_CVLOG(OutEntities.Num() != NumberToSpawn, this, LogMassSpawner, Warning, TEXT("Tried to batch-create %d entities but created %d."), NumberToSpawn, OutEntities.Num());
-
-		TConstArrayView<FInstancedStruct> FragmentInstances = EntityTemplate.GetInitialFragmentValues();
-		const FArchetypeChunkCollection ChunkCollection(EntityTemplate.GetArchetype(), OutEntities);
-		EntitySystem->BatchSetEntityFragmentsValues(ChunkCollection, FragmentInstances);
-
-		if (EntityTemplate.GetInitializationPipeline().Processors.Num())
-		{
-			FMassEntityTemplate& MutableEntityTemplate = const_cast<FMassEntityTemplate&>(EntityTemplate);
-
-			FMassProcessingContext ProcessingContext(*EntitySystem, /*TimeDelta=*/0.0f);
-			UE::Mass::Executor::RunSparse(MutableEntityTemplate.GetMutableInitializationPipeline(), ProcessingContext, ChunkCollection);
-		}
-		return true;
-	}
-	return false;
+	DoSpawning(EntityTemplate, NumberToSpawn, FStructView(), TSubclassOf<UMassProcessor>(), OutEntities);
 }
 
-void UMassSpawnerSubsystem::SpawnEntities(FMassEntityTemplateID TemplateID, const FMassSpawnConfigBase& SpawnConfig, const FStructView& AuxData, TArray<FMassEntityHandle>& OutEntities)
+void UMassSpawnerSubsystem::SpawnEntities(FMassEntityTemplateID TemplateID, const uint32 NumberToSpawn, FConstStructView SpawnData, TSubclassOf<UMassProcessor> InitializerClass, TArray<FMassEntityHandle>& OutEntities)
 {
 	check(TemplateID.IsValid());
 
@@ -162,35 +127,25 @@ void UMassSpawnerSubsystem::SpawnEntities(FMassEntityTemplateID TemplateID, cons
 	}
 
 	const FMassEntityTemplate* EntityTemplate = TemplateRegistryInstance->FindTemplateFromTemplateID(TemplateID);
-	checkf(EntityTemplate != nullptr, TEXT("TemplateID must have been registered!"));
+	checkf(EntityTemplate && EntityTemplate->IsValid(), TEXT("SpawnEntities: TemplateID must have been registered!"));
 
-	DoSpawning(*EntityTemplate, SpawnConfig, AuxData, OutEntities);
+	DoSpawning(*EntityTemplate, NumberToSpawn, SpawnData, InitializerClass, OutEntities);
 }
 
-
-void UMassSpawnerSubsystem::SpawnCollection(TArrayView<FInstancedStruct> Collection, const int32 Count, const FStructView& AuxData /* = FStructView() */)
+void UMassSpawnerSubsystem::SpawnFromConfig(FStructView Config, const int32 NumToSpawn, FConstStructView SpawnData, TSubclassOf<UMassProcessor> InitializerClass)
 {
+	check(Config.IsValid());
+
 	if (!ensureMsgf(TemplateRegistryInstance, TEXT("UMassSpawnerSubsystem didn\'t get its OnPostWorldInit call yet!")))
 	{
 		return;
 	}
+	
+	const FMassEntityTemplate* EntityTemplate = TemplateRegistryInstance->FindOrBuildStructTemplate(Config);
+	checkf(EntityTemplate && EntityTemplate->IsValid(), TEXT("SpawnFromConfig: TemplateID must have been registered!"));
 
-	// @todo totally ignoring "Count" parameter for now. Respecting it will required building spawning model properly
-	// right now the function will just spawn Max allowed by given spawn data
-
-	for (FInstancedStruct& Entry : Collection)
-	{
-		if (Entry.IsValid())
-		{
-			const FMassEntityTemplate* EntityTemplate = TemplateRegistryInstance->FindOrBuildStructTemplate(Entry);
-			if (EntityTemplate && EntityTemplate->IsValid())
-			{
-				TArray<FMassEntityHandle> Entities;
-
-				DoSpawning(*EntityTemplate, Entry.GetMutable<FMassSpawnConfigBase>(), AuxData, Entities);
-			}
-		}
-	}
+	TArray<FMassEntityHandle> Entities;
+	DoSpawning(*EntityTemplate, NumToSpawn, SpawnData, InitializerClass, Entities);
 }
 
 void UMassSpawnerSubsystem::RegisterCollection(TArrayView<FInstancedStruct> Collection)
@@ -284,7 +239,31 @@ void UMassSpawnerSubsystem::DestroyEntities(const FMassEntityTemplateID Template
 	}
 }
 
-void UMassSpawnerSubsystem::DoSpawning(const FMassEntityTemplate& EntityTemplate, const FMassSpawnConfigBase& Data, const FStructView& AuxData, TArray<FMassEntityHandle>& OutEntities) const
+UMassProcessor* UMassSpawnerSubsystem::GetSpawnLocationInitializer(TSubclassOf<UMassProcessor> InitializerClass)
+{	
+	if (!InitializerClass)
+	{
+		return nullptr;
+	}
+
+	UMassProcessor** const Initializer = SpawnLocationInitializers.FindByPredicate([InitializerClass](const UMassProcessor* Processor)
+		{
+			return Processor && Processor->GetClass() == InitializerClass;
+		}
+	);
+
+	if (Initializer == nullptr)
+	{
+		UMassProcessor* NewInitializer = NewObject<UMassProcessor>(this, InitializerClass);
+		NewInitializer->Initialize(*this);
+		SpawnLocationInitializers.Add(NewInitializer);
+		return NewInitializer;
+	}
+
+	return *Initializer;
+}
+
+void UMassSpawnerSubsystem::DoSpawning(const FMassEntityTemplate& EntityTemplate, const int32 NumToSpawn, FConstStructView SpawnData, TSubclassOf<UMassProcessor> InitializerClass, TArray<FMassEntityHandle>& OutEntities)
 {
 	check(EntitySystem);
 	check(EntityTemplate.GetArchetype().IsValid());
@@ -295,22 +274,29 @@ void UMassSpawnerSubsystem::DoSpawning(const FMassEntityTemplate& EntityTemplate
 	// 1. Create required number of entities with EntityTemplate.Archetype
 	// 2. Copy data from FMassEntityTemplate.Fragments.
 	//		a. @todo, could be done as part of creation?
-	//		b. @todo could be done via an always-first translator in EntityTemplate.Translators
 	// 3. Run all EntityTemplate.Translators to configure remaining values (i.e. not set via FMassEntityTemplate.Fragments).
 
 	TArray<FMassEntityHandle> SpawnedEntities;
-	EntitySystem->BatchCreateEntities(EntityTemplate.GetArchetype(), Data.MaxNumber, SpawnedEntities);
+	EntitySystem->BatchCreateEntities(EntityTemplate.GetArchetype(), NumToSpawn, SpawnedEntities);
 
 	TConstArrayView<FInstancedStruct> FragmentInstances = EntityTemplate.GetInitialFragmentValues();
 	const FArchetypeChunkCollection ChunkCollection(EntityTemplate.GetArchetype(), SpawnedEntities);
 	EntitySystem->BatchSetEntityFragmentsValues(ChunkCollection, FragmentInstances);
 	
+	UMassProcessor* SpawnLocationInitializer = SpawnData.IsValid() ? GetSpawnLocationInitializer(InitializerClass) : nullptr;
+
+	if (SpawnLocationInitializer)
+	{
+		FMassProcessingContext ProcessingContext(*EntitySystem, /*TimeDelta=*/0.0f);
+		ProcessingContext.AuxData = SpawnData;
+		UE::Mass::Executor::RunProcessorsView(MakeArrayView(&SpawnLocationInitializer, 1), ProcessingContext, &ChunkCollection);
+	}
+
 	if (EntityTemplate.GetInitializationPipeline().Processors.Num())
 	{
 		FMassEntityTemplate& MutableEntityTemplate = const_cast<FMassEntityTemplate&>(EntityTemplate);
 
 		FMassProcessingContext ProcessingContext(*EntitySystem, /*TimeDelta=*/0.0f);
-		ProcessingContext.AuxData = AuxData;
 		UE::Mass::Executor::RunSparse(MutableEntityTemplate.GetMutableInitializationPipeline(), ProcessingContext, ChunkCollection);
 	}
 
diff --git a/Engine/Plugins/Runtime/MassGameplay/Source/MassSpawner/Public/MassSpawnLocationProcessor.h b/Engine/Plugins/Runtime/MassGameplay/Source/MassSpawner/Public/MassSpawnLocationProcessor.h
new file mode 100644
--- /dev/null
+++ b/Engine/Plugins/Runtime/MassGameplay/Source/MassSpawner/Public/MassSpawnLocationProcessor.h
@@ -0,0 +1,21 @@
+// Copyright Epic Games, Inc. All Rights Reserved.
+
+#pragma once 
+
+#include "MassProcessor.h"
+#include "MassSpawnLocationProcessor.generated.h"
+
+UCLASS()
+class MASSSPAWNER_API UMassSpawnLocationProcessor : public UMassProcessor
+{
+	GENERATED_BODY()
+
+public:
+	UMassSpawnLocationProcessor();
+
+protected:
+	virtual void ConfigureQueries() override;
+	virtual void Execute(UMassEntitySubsystem& EntitySubsystem, FMassExecutionContext& Context) override;
+
+	FMassEntityQuery EntityQuery;
+};
diff --git a/Engine/Plugins/Runtime/MassGameplay/Source/MassSpawner/Public/MassSpawnerSubsystem.h b/Engine/Plugins/Runtime/MassGameplay/Source/MassSpawner/Public/MassSpawnerSubsystem.h
--- a/Engine/Plugins/Runtime/MassGameplay/Source/MassSpawner/Public/MassSpawnerSubsystem.h
+++ b/Engine/Plugins/Runtime/MassGameplay/Source/MassSpawner/Public/MassSpawnerSubsystem.h
@@ -3,11 +3,11 @@
 #pragma once
 
 #include "MassEntitySubsystem.h"
+#include "InstancedStruct.h"
 #include "Subsystems/WorldSubsystem.h"
 #include "MassSpawnerSubsystem.generated.h"
 
 class UMassEntitySubsystem;
-struct FMassSpawnConfigBase;
 struct FMassEntityTemplate;
 class UMassEntityTemplateRegistry;
 struct FInstancedStruct;
@@ -19,27 +19,19 @@ UCLASS()
 class MASSSPAWNER_API UMassSpawnerSubsystem : public UWorldSubsystem
 {
 	GENERATED_BODY()
-public:
-
-	void SpawnEntities(FMassEntityTemplateID TemplateID, const FMassSpawnConfigBase& SpawnConfig, const FStructView& AuxData, TArray<FMassEntityHandle>& OutEntities);
-
-	/** Finds the appropriate entity template for ActorInstance and calls the other SpawnEntities implementation 
-	 *  @return true if spawning was successful, false otherwise. The failure might come from there being no EntityTemplate 
-	 *  @param ActorInstance instance of an actor for which we want to spawn entities for
-	 *  @param NumberToSpawn number of entities to spawn
-	 *  @param OutEntities where the IDs of created entities get added. Note that the contents of OutEntities get overridden by the function.
-	 *  for ActorInstance. See logs for more details. */
-	bool SpawnEntities(const AActor& ActorInstance, const uint32 NumberToSpawn, TArray<FMassEntityHandle>& OutEntities) const;
 
+public:
 	/** Spawns entities of the kind described by the given EntityTemplate. The spawned entities are fully initialized
 	 *  meaning the EntityTemplate.InitializationPipeline gets run for all spawned entities.
 	 *  @param EntityTemplate template to use for spawning entities
 	 *  @param NumberToSpawn number of entities to spawn
 	 *  @param OutEntities where the IDs of created entities get added. Note that the contents of OutEntities get overridden by the function.
 	 *  @return true if spawning was successful, false otherwise. In case of failure see logs for more details. */
-	bool SpawnEntities(const FMassEntityTemplate& EntityTemplate, const uint32 NumberToSpawn, TArray<FMassEntityHandle>& OutEntities) const;
+	void SpawnEntities(const FMassEntityTemplate& EntityTemplate, const uint32 NumberToSpawn, TArray<FMassEntityHandle>& OutEntities);
+
+	void SpawnEntities(FMassEntityTemplateID TemplateID, const uint32 NumberToSpawn, FConstStructView SpawnData, TSubclassOf<UMassProcessor> InitializerClass, TArray<FMassEntityHandle>& OutEntities);
 
-	void SpawnCollection(TArrayView<FInstancedStruct> Collection, const int32 Count, const FStructView& AuxData = FStructView());
+	void SpawnFromConfig(FStructView Config, const int32 NumToSpawn, FConstStructView SpawnData, TSubclassOf<UMassProcessor> InitializerClass);
 
 	void DestroyEntities(const FMassEntityTemplateID TemplateID, TConstArrayView<FMassEntityHandle> Entities);
 
@@ -53,7 +45,12 @@ protected:
 	virtual void Initialize(FSubsystemCollectionBase& Collection) override;
 	virtual void PostInitialize() override;
 
-	void DoSpawning(const FMassEntityTemplate& EntityTemplate, const FMassSpawnConfigBase& Data, const FStructView& AuxData, TArray<FMassEntityHandle>& OutEntities) const;
+	void DoSpawning(const FMassEntityTemplate& EntityTemplate, const int32 NumToSpawn, FConstStructView SpawnData, TSubclassOf<UMassProcessor> InitializerClass, TArray<FMassEntityHandle>& OutEntities);
+
+	UMassProcessor* GetSpawnLocationInitializer(TSubclassOf<UMassProcessor> InitializerClass);
+
+	UPROPERTY()
+	TArray<UMassProcessor*> SpawnLocationInitializers;
 
 	UPROPERTY()
 	UMassEntitySubsystem* EntitySystem;
diff --git a/Engine/Plugins/Runtime/MassGameplay/Source/MassSpawner/Public/MassSpawnerTypes.h b/Engine/Plugins/Runtime/MassGameplay/Source/MassSpawner/Public/MassSpawnerTypes.h
--- a/Engine/Plugins/Runtime/MassGameplay/Source/MassSpawner/Public/MassSpawnerTypes.h
+++ b/Engine/Plugins/Runtime/MassGameplay/Source/MassSpawner/Public/MassSpawnerTypes.h
@@ -26,29 +26,6 @@ struct FMassSpawnAuxData
 	TArray<FTransform> Transforms;
 };
 
-USTRUCT()
-struct FMassSpawnConfigBase
-{
-	GENERATED_BODY()
-
-	UPROPERTY(EditAnywhere, Category = "Mass|Spawn", meta = (BaseStruct = "MassFragment"))
-	TArray<FInstancedStruct> AdditionalDataFragments;
-
-	UPROPERTY(EditAnywhere, Category = "Mass|Spawn", meta = (ClampMin = 0, UIMin = 0))
-	int32 MinNumber = 0;
-
-	UPROPERTY(EditAnywhere, Category = "Mass|Spawn", meta = (ClampMin = 1, UIMin = 1))
-	int32 MaxNumber = 1;
-
-	UPROPERTY(EditAnywhere, Category = "Mass|Spawn", meta = (ClampMin = 0.0, UIMin = 0.0))
-	float Radius = 0.f;
-
-	friend uint32 GetTypeHash(const FMassSpawnConfigBase& Instance)
-	{
-		return HashCombine(HashCombine(GetTypeHash(Instance.MinNumber), GetTypeHash(Instance.MaxNumber)), GetTypeHash(Instance.Radius));
-	}
-};
-
 /**
  * Describes an entity type to spawn.
  */
