Removed unused uobject/uclass-related functionality from the entity template registry

diff --git a/Engine/Plugins/Runtime/MassGameplay/Source/MassSpawner/Private/MassEntityTemplateRegistry.cpp b/Engine/Plugins/Runtime/MassGameplay/Source/MassSpawner/Private/MassEntityTemplateRegistry.cpp
--- a/Engine/Plugins/Runtime/MassGameplay/Source/MassSpawner/Private/MassEntityTemplateRegistry.cpp
+++ b/Engine/Plugins/Runtime/MassGameplay/Source/MassSpawner/Private/MassEntityTemplateRegistry.cpp
@@ -48,7 +48,6 @@ namespace FTemplateRegistryHelpers
 // UMassEntityTemplateRegistry 
 //----------------------------------------------------------------------//
 TMap<const UScriptStruct*, UMassEntityTemplateRegistry::FStructToTemplateBuilderDelegate> UMassEntityTemplateRegistry::StructBasedBuilders;
-TMap<const UClass*, UMassEntityTemplateRegistry::FClassToTemplateBuilderDelegate> UMassEntityTemplateRegistry::ClassBasedBuilders;
 
 void UMassEntityTemplateRegistry::BeginDestroy()
 {
@@ -69,11 +68,6 @@ UMassEntityTemplateRegistry::FStructToTemplateBuilderDelegate& UMassEntityTempla
 	return StructBasedBuilders.FindOrAdd(&DataType);
 }
 
-UMassEntityTemplateRegistry::FClassToTemplateBuilderDelegate& UMassEntityTemplateRegistry::FindOrAdd(const UClass& Class)
-{
-	return ClassBasedBuilders.FindOrAdd(&Class);
-}
-
 const FMassEntityTemplate* UMassEntityTemplateRegistry::FindOrBuildStructTemplate(const FInstancedStruct& StructInstance)
 {
 	// thou shall not call this function on CDO
@@ -145,106 +139,6 @@ bool UMassEntityTemplateRegistry::BuildTemplateImpl(const FStructToTemplateBuild
 	return false;
 }
 
-const FMassEntityTemplate* UMassEntityTemplateRegistry::FindOrBuildInstanceTemplate(const UObject& Instance)
-{
-	checkf(HasAnyFlags(RF_ClassDefaultObject) == false, TEXT("Thou shall not call %s on UMassEntityTemplateRegistry's CDO"), ANSI_TO_TCHAR(__FUNCTION__));
-	check(Instance.GetClass());
-
-	const FMassEntityTemplateID* const TemplateID = LookupTemplateIDMap.Find(GetTypeHash(Instance.GetClass()->GetFName()));
-
-	if (TemplateID != nullptr)
-	{
-		if (const FMassEntityTemplate* TemplateFound = TemplateIDToTemplateMap.Find(*TemplateID))
-		{
-			return TemplateFound;
-		}
-	}
-
-	// TODO this code needs refactoring as its not accounting for instances needing a different hash to classes
-	// 2. If not: build it & store it
-	return BuildClassTemplate(*Instance.GetClass(), &Instance);
-}
-
-const FMassEntityTemplate& UMassEntityTemplateRegistry::FindInstanceTemplateChecked(const UObject& Instance) const
-{
-	checkf(HasAnyFlags(RF_ClassDefaultObject) == false,  TEXT("Thou shall not call %s on UMassEntityTemplateRegistry's CDO"), ANSI_TO_TCHAR(__FUNCTION__));
-	check(Instance.GetClass());
-
-	const FMassEntityTemplateID& TemplateID = LookupTemplateIDMap.FindChecked(GetTypeHash(Instance.GetClass()->GetFName()));
-
-	return TemplateIDToTemplateMap.FindChecked(TemplateID);
-}
-
-const FMassEntityTemplate* UMassEntityTemplateRegistry::FindOrBuildClassTemplate(const UClass& Class)
-{
-	// thou shall not call this function on UMassEntityTemplateRegistry's CDO
-	check(HasAnyFlags(RF_ClassDefaultObject) == false);
-
-	// 1. Check if we already have the template stored.
-	const FMassEntityTemplateID* const TemplateID = LookupTemplateIDMap.Find(GetTypeHash(Class.GetFName()));
-
-	if (TemplateID != nullptr)
-	{
-		if (const FMassEntityTemplate* TemplateFound = TemplateIDToTemplateMap.Find(*TemplateID))
-		{
-			return TemplateFound;
-		}
-	}
-
-	// 2. If not: build it & store it
-	return BuildClassTemplate(Class);
-}
-
-const FMassEntityTemplate* UMassEntityTemplateRegistry::BuildClassTemplate(const UClass& Class, const UObject* Instance)
-{
-	FMassEntityTemplate* EntityTemplate = nullptr;
-	// this function will be called when there's been no EntityTemplate created for Class/Instance just yet
-	// Let's see if we know how to make one	
-	FClassToTemplateBuilderDelegate* Builder = GetBuilderForClass(Class);
-	if (Builder)
-	{
-		TArray<const UMassProcessor*> Initializers;
-		
-		// TODO consider removing the need for strings here
-		// Use the class name string for the hash here, so the hash can be deterministic between client and server
-		const uint32 NameStringHash = GetTypeHash(Class.GetName());
-
-		const EMassEntityTemplateIDType TemplateType = (Instance != nullptr && Instance->HasAnyFlags(RF_ClassDefaultObject) == false) ? EMassEntityTemplateIDType::Instance : EMassEntityTemplateIDType::Class;
-		const FMassEntityTemplateID TemplateID = LookupTemplateIDMap.Add(GetTypeHash(Class.GetFName()), FMassEntityTemplateID(NameStringHash, TemplateType));
-
-		EntityTemplate = &TemplateIDToTemplateMap.Add(TemplateID);
-
-		check(EntityTemplate);
-
-		EntityTemplate->SetTemplateID(TemplateID);
-
-		FMassEntityTemplateBuildContext BuildContext(*EntityTemplate);
-		Builder->Execute(GetWorld(), Class, Instance, BuildContext);
-
-		if (ensureMsgf(!EntityTemplate->IsEmpty(), TEXT("Need at least one fragment to create an Archetype")))
-		{
-			EntityTemplate->SetUpProcessors(BuildContext.Handlers, *this);
-			InitializeEntityTemplate(*EntityTemplate);
-		}
-	}
-	UE_CVLOG_UELOG(Builder == nullptr, this, LogMassSpawner, Warning, TEXT("Attempting to build a MassAgentTemplate for uclass %s (instance %s) while template builder has not been registered for this type")
-		, *Class.GetName(), *GetNameSafe(Instance));
-
-	return EntityTemplate;
-}
-
-UMassEntityTemplateRegistry::FClassToTemplateBuilderDelegate* UMassEntityTemplateRegistry::GetBuilderForClass(const UClass& Class)
-{
-	FClassToTemplateBuilderDelegate* Builder = nullptr;
-	const UClass* TmpClass = &Class;
-	do
-	{ 
-		Builder = ClassBasedBuilders.Find(TmpClass);
-		TmpClass = TmpClass->GetSuperClass();
-	} while (Builder == nullptr && TmpClass != nullptr);
-	return Builder;
-}
-
 void UMassEntityTemplateRegistry::InitializeEntityTemplate(FMassEntityTemplate& OutTemplate) const
 {
 	// expected to be ensured by the caller
diff --git a/Engine/Plugins/Runtime/MassGameplay/Source/MassSpawner/Public/MassEntityTemplateRegistry.h b/Engine/Plugins/Runtime/MassGameplay/Source/MassSpawner/Public/MassEntityTemplateRegistry.h
--- a/Engine/Plugins/Runtime/MassGameplay/Source/MassSpawner/Public/MassEntityTemplateRegistry.h
+++ b/Engine/Plugins/Runtime/MassGameplay/Source/MassSpawner/Public/MassEntityTemplateRegistry.h
@@ -13,7 +13,6 @@
 
 class UWorld;
 class UMassEntitySubsystem;
-class UMassTranslator;
 
 enum class EFragmentInitialization : uint8
 {
@@ -216,9 +215,6 @@ protected:
 };
 
 /** @todo document 
- *  Note: the Object-Instance entity templates do not support having multiple instances of the same class that differ in
- *  composition (like with dynamically added fragments). The first registered instance will determine the composition 
- *  of the entity template.
  */
 UCLASS()
 class MASSSPAWNER_API UMassEntityTemplateRegistry : public UObject
@@ -227,19 +223,14 @@ class MASSSPAWNER_API UMassEntityTemplateRegistry : public UObject
 public:
 	// @todo consider TFunction instead
 	DECLARE_DELEGATE_ThreeParams(FStructToTemplateBuilderDelegate, const UWorld* /*World*/, const FInstancedStruct& /*InStructInstance*/, FMassEntityTemplateBuildContext& /*BuildContext*/);
-	DECLARE_DELEGATE_FourParams(FClassToTemplateBuilderDelegate, const UWorld* /*World*/, const UClass& /*InClass*/, const UObject* /*ObjectInstance*/, FMassEntityTemplateBuildContext& /*BuildContext*/);
 
 	UMassEntityTemplateRegistry() = default;
 	virtual void BeginDestroy() override;
 	virtual UWorld* GetWorld() const override;
 
 	static FStructToTemplateBuilderDelegate& FindOrAdd(const UScriptStruct& DataType);
-	static FClassToTemplateBuilderDelegate& FindOrAdd(const UClass& Class);
 
 	const FMassEntityTemplate* FindOrBuildStructTemplate(const FInstancedStruct& StructInstance);
-	const FMassEntityTemplate* FindOrBuildInstanceTemplate(const UObject& Instance);
-	const FMassEntityTemplate& FindInstanceTemplateChecked(const UObject& Instance) const;
-	const FMassEntityTemplate* FindOrBuildClassTemplate(const UClass& Class);
 
 	/** Removes all the cached template instances */
 	void DebugReset();
@@ -252,25 +243,15 @@ public:
 	void InitializeEntityTemplate(FMassEntityTemplate& OutTemplate) const;
 
 protected:
-	/** A function implementing common path for UClass and UObject template building 
-	 *	@param Class type of class for which we need to build an entity template
-	 *	@param Instance if given means we're building a template for a runtime UObject instance. It will be passed down to template builder
-	 */
-	const FMassEntityTemplate* BuildClassTemplate(const UClass& Class, const UObject* Instance = nullptr);
-
 	/** @return true if a template has been built, false otherwise */
 	bool BuildTemplateImpl(const FStructToTemplateBuilderDelegate& Builder, const FInstancedStruct& StructInstance, FMassEntityTemplate& OutTemplate);
 
-	static FClassToTemplateBuilderDelegate* GetBuilderForClass(const UClass& Class);
-
 protected:
 	static TMap<const UScriptStruct*, FStructToTemplateBuilderDelegate> StructBasedBuilders;
-	static TMap<const UClass*, FClassToTemplateBuilderDelegate> ClassBasedBuilders;
-	
+
 	/** 
 	 *  Map from a hash to a FMassEntityTemplateID
 	 *  For build from FInstancedStruct it will be the Combined hash of the UScriptStruct FName and FTemplateRegistryHelpers::CalcHash()
-	 *  For build from UClass or UObject it will be the hash of the UClass FName
 	 *  This hash will not be deterministic between server and clients
 	 */
 	TMap<uint32, FMassEntityTemplateID> LookupTemplateIDMap;
