Moved UE::Mass::SpawnerSubsystem::CreateSparseChunks to UE::Mass::Utils

diff --git a/Engine/Plugins/Runtime/MassEntity/Source/MassEntity/Private/MassEntityUtils.cpp b/Engine/Plugins/Runtime/MassEntity/Source/MassEntity/Private/MassEntityUtils.cpp
--- a/Engine/Plugins/Runtime/MassEntity/Source/MassEntity/Private/MassEntityUtils.cpp
+++ b/Engine/Plugins/Runtime/MassEntity/Source/MassEntity/Private/MassEntityUtils.cpp
@@ -1,6 +1,9 @@
 // Copyright Epic Games, Inc. All Rights Reserved.
 
 #include "MassEntityUtils.h"
+#include "MassEntityTypes.h"
+#include "MassArchetypeTypes.h"
+#include "MassEntitySubsystem.h"
 #include "Engine/EngineBaseTypes.h"
 #include "Engine/World.h"
 
@@ -31,4 +34,28 @@ EProcessorExecutionFlags GetProcessorExecutionFlagsForWold(const UWorld& World)
 	return ExecutionFlags;
 }
 
+void CreateSparseChunks(const UMassEntitySubsystem& EntitySystem, const TConstArrayView<FMassEntityHandle> Entities, TArray<FArchetypeChunkCollection>& OutChunkCollections)
+{
+	TRACE_CPUPROFILER_EVENT_SCOPE_STR("Mass_CreateSparseChunks");
+
+	TMap<const FArchetypeHandle, TArray<FMassEntityHandle>> ArchetypeToEntities;
+
+	for (const FMassEntityHandle& Entity : Entities)
+	{
+		FArchetypeHandle Archetype = EntitySystem.GetArchetypeForEntity(Entity);
+		TArray<FMassEntityHandle>& PerArchetypeEntities = ArchetypeToEntities.FindOrAdd(Archetype);
+		PerArchetypeEntities.Add(Entity);
+	}
+
+	for (auto& Pair : ArchetypeToEntities)
+	{
+		// @todo this is a temporary measure to skip entities we've destroyed without the
+		// UMassSpawnerSubsystem::DestroyEntities' caller knowledge. Should be removed once that's addressed.
+		if (Pair.Key.IsValid())
+		{
+			OutChunkCollections.Add(FArchetypeChunkCollection(Pair.Key, Pair.Value));
+		}
+	}
+}
+
 } // namespace UE::Mass::Utils
\ No newline at end of file
diff --git a/Engine/Plugins/Runtime/MassEntity/Source/MassEntity/Public/MassEntityUtils.h b/Engine/Plugins/Runtime/MassEntity/Source/MassEntity/Public/MassEntityUtils.h
--- a/Engine/Plugins/Runtime/MassEntity/Source/MassEntity/Public/MassEntityUtils.h
+++ b/Engine/Plugins/Runtime/MassEntity/Source/MassEntity/Public/MassEntityUtils.h
@@ -4,10 +4,18 @@
 
 #include "MassProcessingTypes.h"
 
+class UWorld;
+class UMassEntitySubsystem;
+struct FMassEntityHandle;
+struct FArchetypeChunkCollection;
+
 namespace UE::Mass::Utils
 {
 
 /** returns the current execution mode for the processors calculated from the world network mode */
 MASSENTITY_API extern EProcessorExecutionFlags GetProcessorExecutionFlagsForWold(const UWorld& World);
 
+
+MASSENTITY_API extern void CreateSparseChunks(const UMassEntitySubsystem& EntitySystem, const TConstArrayView<FMassEntityHandle> Entities, TArray<FArchetypeChunkCollection>& OutChunkCollections);
+
 } // namespace UE::Mass::Utils
\ No newline at end of file
diff --git a/Engine/Plugins/Runtime/MassGameplay/Source/MassSpawner/Private/MassSpawnerSubsystem.cpp b/Engine/Plugins/Runtime/MassGameplay/Source/MassSpawner/Private/MassSpawnerSubsystem.cpp
--- a/Engine/Plugins/Runtime/MassGameplay/Source/MassSpawner/Private/MassSpawnerSubsystem.cpp
+++ b/Engine/Plugins/Runtime/MassGameplay/Source/MassSpawner/Private/MassSpawnerSubsystem.cpp
@@ -56,31 +56,6 @@ namespace UE::Mass::SpawnerSubsystem
 
 		return Destructors.Num() > 0;
 	}
-	
-	void CreateSparseChunks(const UMassEntitySubsystem& EntitySystem, const TConstArrayView<FMassEntityHandle> Entities, TArray<FArchetypeChunkCollection>& OutChunkCollections)
-	{
-		TRACE_CPUPROFILER_EVENT_SCOPE_STR("SpawnerSubsystem_CreateSparseChunks");
-
-		TMap<const FArchetypeHandle, TArray<FMassEntityHandle>> ArchetypeToEntities;
-
-		for (const FMassEntityHandle& Entity : Entities)
-		{
-			FArchetypeHandle Archetype = EntitySystem.GetArchetypeForEntity(Entity);
-			TArray<FMassEntityHandle>& PerArchetypeEntities = ArchetypeToEntities.FindOrAdd(Archetype);
-			PerArchetypeEntities.Add(Entity);
-		}
-
-		for (auto& Pair : ArchetypeToEntities)
-		{
-			// @todo this is a temporary measure to skip entities we've destroyed without the
-			// UMassSpawnerSubsystem::DestroyEntities' caller knowledge. Should be removed once that's addressed.
-			if (Pair.Key.IsValid())
-			{
-				OutChunkCollections.Add(FArchetypeChunkCollection(Pair.Key, Pair.Value));
-			}
-		}
-	}
-	
 } // namespace UE::Mass::SpawnerSubsystem
 
 //----------------------------------------------------------------------//
@@ -188,7 +163,7 @@ void UMassSpawnerSubsystem::DestroyEntities(const FMassEntityTemplateID Template
 		// There's no certainty all Entities still belong to the original archetype, the MutableEntityTemplate.GetArchetype().
 		// We need to figure out the separate, per-archetype sparse chunks
 		TArray<FArchetypeChunkCollection> ChunkCollections;
-		UE::Mass::SpawnerSubsystem::CreateSparseChunks(*EntitySystem, Entities, ChunkCollections);
+		UE::Mass::Utils::CreateSparseChunks(*EntitySystem, Entities, ChunkCollections);
 
 		FMassProcessingContext ProcessingContext(*EntitySystem, /*TimeDelta=*/0.0f);
 
@@ -204,7 +179,7 @@ void UMassSpawnerSubsystem::DestroyEntities(const FMassEntityTemplateID Template
 		if (bDestructorsRun)
 		{
 			ChunkCollections.Reset();
-			UE::Mass::SpawnerSubsystem::CreateSparseChunks(*EntitySystem, Entities, ChunkCollections);
+			UE::Mass::Utils::CreateSparseChunks(*EntitySystem, Entities, ChunkCollections);
 		}
 		
 		// Run template deinitializers
@@ -225,7 +200,7 @@ void UMassSpawnerSubsystem::DestroyEntities(const FMassEntityTemplateID Template
 		if (MutableEntityTemplate.GetDeinitializationPipeline().Processors.Num())
 		{
 			ChunkCollections.Reset();
-			UE::Mass::SpawnerSubsystem::CreateSparseChunks(*EntitySystem, Entities, ChunkCollections);
+			UE::Mass::Utils::CreateSparseChunks(*EntitySystem, Entities, ChunkCollections);
 		}
 
 		for (const FArchetypeChunkCollection& Chunks : ChunkCollections)
