Cleanup old simulation LOD processors

#rnx

diff --git a/Engine/Plugins/Runtime/MassGameplay/Source/MassLOD/Private/MassSimulationLOD.cpp b/Engine/Plugins/Runtime/MassGameplay/Source/MassLOD/Private/MassSimulationLOD.cpp
--- a/Engine/Plugins/Runtime/MassGameplay/Source/MassLOD/Private/MassSimulationLOD.cpp
+++ b/Engine/Plugins/Runtime/MassGameplay/Source/MassLOD/Private/MassSimulationLOD.cpp
@@ -8,35 +8,6 @@
 #include "VisualLogger/VisualLogger.h"
 #include "MassLODUtils.h"
 
-//-----------------------------------------------------------------------------
-// UMassProcessor_MassSimulationLODViewersInfo
-//-----------------------------------------------------------------------------
-
-UMassProcessor_MassSimulationLODViewersInfo::UMassProcessor_MassSimulationLODViewersInfo()
-{
-	bAutoRegisterWithProcessingPhases = false;
-}
-
-void UMassProcessor_MassSimulationLODViewersInfo::ConfigureQueries()
-{
-	EntityQuery.AddRequirement<FDataFragment_Transform>(EMassFragmentAccess::ReadOnly);
-	EntityQuery.AddRequirement<FDataFragment_MassSimulationLODInfo>(EMassFragmentAccess::ReadWrite);
-}
-
-void UMassProcessor_MassSimulationLODViewersInfo::Execute(UMassEntitySubsystem& EntitySubsystem, FMassExecutionContext& Context)
-{
-	check(LODManager);
-	const TArray<FViewerInfo>& Viewers = LODManager->GetViewers();
-	LODCollector.PrepareExecution(Viewers);
-
-	EntityQuery.ForEachEntityChunk(EntitySubsystem, Context, [this](FMassExecutionContext& Context)
-	{
-		const TConstArrayView<FDataFragment_Transform> LocationList = Context.GetFragmentView<FDataFragment_Transform>();
-		const TArrayView<FDataFragment_MassSimulationLODInfo> ViewersInfoList = Context.GetMutableFragmentView<FDataFragment_MassSimulationLODInfo>();
-		LODCollector.CollectLODInfo(Context, LocationList, ViewersInfoList);
-	});
-}
-
 //-----------------------------------------------------------------------------
 // FMassSimulationLODConfig
 //-----------------------------------------------------------------------------
diff --git a/Engine/Plugins/Runtime/MassGameplay/Source/MassLOD/Public/MassSimulationLOD.h b/Engine/Plugins/Runtime/MassGameplay/Source/MassLOD/Public/MassSimulationLOD.h
--- a/Engine/Plugins/Runtime/MassGameplay/Source/MassLOD/Public/MassSimulationLOD.h
+++ b/Engine/Plugins/Runtime/MassGameplay/Source/MassLOD/Public/MassSimulationLOD.h
@@ -11,43 +11,6 @@
 
 #include "MassSimulationLOD.generated.h"
 
-/*
- * Data fragment to store the calculated distances to viewers
- * @todo To be removed and convert to new Simulation LOD stuff
- */
-USTRUCT()
-struct MASSLOD_API FDataFragment_MassSimulationLODInfo : public FMassFragment
-{
-	GENERATED_BODY()
-
-	// Closest viewer distance
-	float ClosestViewerDistanceSq;
-};
-
-struct FSimulationLODLogic : public FLODDefaultLogic
-{
-};
-
-/*
- * Mass processor to calculate square distances from entity to viewers.
- * This is done once per entity and will be used for all subsequent LOD calculation.
- * @todo To be removed and convert to new Simulation LOD stuff
- */
-UCLASS()
-class MASSLOD_API UMassProcessor_MassSimulationLODViewersInfo : public UMassProcessor_LODBase
-{
-	GENERATED_BODY()
-public:
-	UMassProcessor_MassSimulationLODViewersInfo();
-
-protected:
-	virtual void ConfigureQueries() override;
-	virtual void Execute(UMassEntitySubsystem& EntitySubsystem, FMassExecutionContext& Context) override;
-
-	TMassLODCollector<FSimulationLODLogic> LODCollector;
-	FMassEntityQuery EntityQuery;
-};
-
 USTRUCT()
 struct MASSLOD_API FMassSimulationLODFragment : public FMassFragment
 {
