Moved the Local Viewer check in the LOD calculation instead of during the gathering of the Viewers. This is in prevision of supporting difference type of calculation from a single LOD helper class.

#rnx

diff --git a/Engine/Plugins/Runtime/MassGameplay/Source/MassLOD/Private/MassLODLogic.cpp b/Engine/Plugins/Runtime/MassGameplay/Source/MassLOD/Private/MassLODLogic.cpp
--- a/Engine/Plugins/Runtime/MassGameplay/Source/MassLOD/Private/MassLODLogic.cpp
+++ b/Engine/Plugins/Runtime/MassGameplay/Source/MassLOD/Private/MassLODLogic.cpp
@@ -4,7 +4,7 @@
 #include "Engine/World.h"
 #include "GameFramework/PlayerController.h"
 
-void FMassLODBaseLogic::CacheViewerInformation(TConstArrayView<FViewerInfo> ViewerInfos, const bool bLocalViewersOnly)
+void FMassLODBaseLogic::CacheViewerInformation(TConstArrayView<FViewerInfo> ViewerInfos)
 {
 	if(Viewers.Num() < ViewerInfos.Num())
 	{
@@ -15,12 +15,13 @@ void FMassLODBaseLogic::CacheViewerInformation(TConstArrayView<FViewerInfo> View
 	for (int ViewerIdx = 0; ViewerIdx < Viewers.Num(); ++ViewerIdx)
 	{
 		const FViewerInfo& Viewer =  ViewerInfos.IsValidIndex(ViewerIdx) ? ViewerInfos[ViewerIdx] : FViewerInfo();
-		const FMassViewerHandle ViewerHandle =  Viewer.bEnabled && (!bLocalViewersOnly || Viewer.IsLocal()) ? Viewer.Handle : FMassViewerHandle();
+		const FMassViewerHandle ViewerHandle =  Viewer.bEnabled ? Viewer.Handle : FMassViewerHandle();
 
 		// Check if it is the same client as before
 		FViewerLODInfo& ViewerLOD = Viewers[ViewerIdx];
 		ViewerLOD.bClearData = Viewers[ViewerIdx].Handle != ViewerHandle;
 		ViewerLOD.Handle = ViewerHandle;
+		ViewerLOD.bLocal = Viewer.IsLocal();
 
 		if (ViewerHandle.IsValid())
 		{
diff --git a/Engine/Plugins/Runtime/MassGameplay/Source/MassLOD/Public/MassLODCalculator.h b/Engine/Plugins/Runtime/MassGameplay/Source/MassLOD/Public/MassLODCalculator.h
--- a/Engine/Plugins/Runtime/MassGameplay/Source/MassLOD/Public/MassLODCalculator.h
+++ b/Engine/Plugins/Runtime/MassGameplay/Source/MassLOD/Public/MassLODCalculator.h
@@ -236,7 +236,7 @@ bool TMassLODCalculator<FLODLogic>::CalculateVisibility(const bool bWasVisible,
 template <typename FLODLogic>
 void TMassLODCalculator<FLODLogic>::PrepareExecution(TConstArrayView<FViewerInfo> ViewersInfo)
 {
-	CacheViewerInformation(ViewersInfo, FLODLogic::bLocalViewersOnly);
+	CacheViewerInformation(ViewersInfo);
 
 	if (FLODLogic::bMaximizeCountPerViewer)
 	{
@@ -344,6 +344,13 @@ void TMassLODCalculator<FLODLogic>::CalculateLOD(FMassExecutionContext& Context,
 					SetPrevVisibilityPerViewer<FLODLogic::bDoVisibilityLogic&& FLODLogic::bStoreInfoPerViewer>(EntityLOD, ViewerIdx, EMassVisibility::Max);
 					SetVisibilityPerViewer<FLODLogic::bDoVisibilityLogic && FLODLogic::bStoreInfoPerViewer>(EntityLOD, ViewerIdx, EMassVisibility::Max);
 				}
+
+				// Check to see if we want only local viewer only
+				if (FLODLogic::bLocalViewersOnly && !Viewer.bLocal)
+				{
+					continue;
+				}
+
 				if (Viewer.Handle.IsValid())
 				{
 					const float DistanceToFrustum = GetDistanceToFrustum<FLODLogic::bDoVisibilityLogic && FLODLogic::bStoreInfoPerViewer>(EntityViewersInfo, ViewerIdx, FLT_MAX);
@@ -527,7 +534,14 @@ void TMassLODCalculator<FLODLogic>::AdjustLODFromCount(FMassExecutionContext& Co
 		{
 			for (int ViewerIdx = 0; ViewerIdx < Viewers.Num(); ++ViewerIdx)
 			{
-				if (!Viewers[ViewerIdx].Handle.IsValid())
+				const FViewerLODInfo& Viewer = Viewers[ViewerIdx];
+				if (!Viewer.Handle.IsValid())
+				{
+					continue;
+				}
+
+				// Check to see if we want only local viewer only
+				if (FLODLogic::bLocalViewersOnly && !Viewer.bLocal)
 				{
 					continue;
 				}
diff --git a/Engine/Plugins/Runtime/MassGameplay/Source/MassLOD/Public/MassLODCollector.h b/Engine/Plugins/Runtime/MassGameplay/Source/MassLOD/Public/MassLODCollector.h
--- a/Engine/Plugins/Runtime/MassGameplay/Source/MassLOD/Public/MassLODCollector.h
+++ b/Engine/Plugins/Runtime/MassGameplay/Source/MassLOD/Public/MassLODCollector.h
@@ -35,7 +35,7 @@ struct TMassLODCollector : public FMassLODBaseLogic
 template <typename FLODLogic>
 void TMassLODCollector<FLODLogic>::PrepareExecution(TConstArrayView<FViewerInfo> ViewersInfo)
 {
-	CacheViewerInformation(ViewersInfo, FLODLogic::bLocalViewersOnly);
+	CacheViewerInformation(ViewersInfo);
 }
 
 template <typename FLODLogic>
@@ -60,6 +60,13 @@ void TMassLODCollector<FLODLogic>::CollectLODInfo(FMassExecutionContext& Context
 				SetDistanceToViewerSq<FLODLogic::bStoreInfoPerViewer>(EntityViewerInfo, ViewerIdx, FLT_MAX);
 				SetDistanceToFrustum<FLODLogic::bDoVisibilityLogic && FLODLogic::bStoreInfoPerViewer>(EntityViewerInfo, ViewerIdx, FLT_MAX);
 			}
+
+			// Check to see if we want only local viewer only
+			if (FLODLogic::bLocalViewersOnly && !Viewer.bLocal)
+			{
+				continue;
+			}
+
 			if (Viewer.Handle.IsValid())
 			{
 				const FVector& EntityLocation = EntityTransform.GetTransform().GetLocation();
diff --git a/Engine/Plugins/Runtime/MassGameplay/Source/MassLOD/Public/MassLODLogic.h b/Engine/Plugins/Runtime/MassGameplay/Source/MassLOD/Public/MassLODLogic.h
--- a/Engine/Plugins/Runtime/MassGameplay/Source/MassLOD/Public/MassLODLogic.h
+++ b/Engine/Plugins/Runtime/MassGameplay/Source/MassLOD/Public/MassLODLogic.h
@@ -120,6 +120,9 @@ struct FMassViewerInfoFragment
 
 struct FViewerLODInfo
 {
+	/* Boolean indicating the viewer is local or not */
+	bool bLocal = false;
+
 	/* Boolean indicating the viewer data needs to be cleared */
 	bool bClearData = false;
 
@@ -140,24 +143,24 @@ struct FViewerLODInfo
 struct MASSLOD_API FMassLODBaseLogic
 {
 protected:
-	void CacheViewerInformation(TConstArrayView<FViewerInfo> ViewerInfos, const bool bLocalViewersOnly);
+	void CacheViewerInformation(TConstArrayView<FViewerInfo> ViewerInfos);
 
-	// Per viewer LOD information conditional fragment accessors
+	/** Per viewer LOD information conditional fragment accessors */
 	DECLARE_CONDITIONAL_MEMBER_ARRAY_ACCESSORS(Condition, EMassLOD::Type, LODPerViewer);
 	DECLARE_CONDITIONAL_MEMBER_ARRAY_ACCESSORS(Condition, EMassLOD::Type, PrevLODPerViewer);
 
-	// LOD Significance conditional fragment accessors
+	/** LOD Significance conditional fragment accessors */
 	DECLARE_CONDITIONAL_MEMBER_ACCESSORS(Condition, float, LODSignificance);
 
-	// Visibility conditional fragment accessors
+	/** Visibility conditional fragment accessors */
 	DECLARE_CONDITIONAL_MEMBER_ACCESSORS(Condition, float, ClosestDistanceToFrustum);
 	DECLARE_CONDITIONAL_MEMBER_ACCESSORS(Condition, EMassVisibility, Visibility);
 	DECLARE_CONDITIONAL_MEMBER_ACCESSORS(Condition, EMassVisibility, PrevVisibility);
 
-	// Per viewer distance conditional fragment accessors
+	/** Per viewer distance conditional fragment accessors */
 	DECLARE_CONDITIONAL_MEMBER_ARRAY_ACCESSORS(Condition, float, DistanceToViewerSq);
 
-	// Per viewer visibility conditional fragment accessors
+	/** Per viewer visibility conditional fragment accessors */
 	DECLARE_CONDITIONAL_MEMBER_ARRAY_ACCESSORS(Condition, float, DistanceToFrustum);
 	DECLARE_CONDITIONAL_MEMBER_ARRAY_ACCESSORS(Condition, EMassVisibility, VisibilityPerViewer);
 	DECLARE_CONDITIONAL_MEMBER_ARRAY_ACCESSORS(Condition, EMassVisibility, PrevVisibilityPerViewer);
