Split per viewer LOD info into its own fragment

#rnx

diff --git a/Engine/Plugins/AI/MassCrowd/Source/MassCrowd/Private/MassCrowdServerRepresentationLODProcessor.cpp b/Engine/Plugins/AI/MassCrowd/Source/MassCrowd/Private/MassCrowdServerRepresentationLODProcessor.cpp
--- a/Engine/Plugins/AI/MassCrowd/Source/MassCrowd/Private/MassCrowdServerRepresentationLODProcessor.cpp
+++ b/Engine/Plugins/AI/MassCrowd/Source/MassCrowd/Private/MassCrowdServerRepresentationLODProcessor.cpp
@@ -58,7 +58,7 @@ void UMassCrowdServerRepresentationLODProcessor::Execute(UMassEntitySubsystem& E
 		{
 			const TConstArrayView<FMassViewerInfoFragment> ViewersInfoList = Context.GetFragmentView<FMassViewerInfoFragment>();
 			const TArrayView<FMassRepresentationLODFragment> RepresentationLODFragments = Context.GetMutableFragmentView<FMassRepresentationLODFragment>();
-			LODCalculator.CalculateLOD(Context, RepresentationLODFragments, ViewersInfoList);
+			LODCalculator.CalculateLOD(Context, ViewersInfoList, RepresentationLODFragments);
 		});
 	}
 
@@ -71,7 +71,7 @@ void UMassCrowdServerRepresentationLODProcessor::Execute(UMassEntitySubsystem& E
 			{
 				const TConstArrayView<FMassViewerInfoFragment> ViewersInfoList = Context.GetFragmentView<FMassViewerInfoFragment>();
 				const TArrayView<FMassRepresentationLODFragment> RepresentationLODFragments = Context.GetMutableFragmentView<FMassRepresentationLODFragment>();
-				LODCalculator.AdjustLODFromCount(Context, RepresentationLODFragments, ViewersInfoList);
+				LODCalculator.AdjustLODFromCount(Context, ViewersInfoList, RepresentationLODFragments);
 			});
 		}
 	}
diff --git a/Engine/Plugins/Runtime/MassGameplay/Source/MassLOD/Private/MassSimulationLOD.cpp b/Engine/Plugins/Runtime/MassGameplay/Source/MassLOD/Private/MassSimulationLOD.cpp
--- a/Engine/Plugins/Runtime/MassGameplay/Source/MassLOD/Private/MassSimulationLOD.cpp
+++ b/Engine/Plugins/Runtime/MassGameplay/Source/MassLOD/Private/MassSimulationLOD.cpp
@@ -128,7 +128,7 @@ void UMassSimulationLODProcessor::Execute(UMassEntitySubsystem& EntitySubsystem,
 			FMassSimulationLODSharedFragment& LODSharedFragment = Context.GetMutableSharedFragment<FMassSimulationLODSharedFragment>();
 			TConstArrayView<FMassViewerInfoFragment> ViewersInfoList = Context.GetFragmentView<FMassViewerInfoFragment>();
 			TArrayView<FMassSimulationLODFragment> SimulationLODFragments = Context.GetMutableFragmentView<FMassSimulationLODFragment>();
-			LODSharedFragment.LODCalculator.CalculateLOD(Context, SimulationLODFragments, ViewersInfoList);
+			LODSharedFragment.LODCalculator.CalculateLOD(Context, ViewersInfoList, SimulationLODFragments);
 		});
 	}
 
@@ -144,7 +144,7 @@ void UMassSimulationLODProcessor::Execute(UMassEntitySubsystem& EntitySubsystem,
 			FMassSimulationLODSharedFragment& LODSharedFragment = Context.GetMutableSharedFragment<FMassSimulationLODSharedFragment>();
 			TConstArrayView<FMassViewerInfoFragment> ViewersInfoList = Context.GetFragmentView<FMassViewerInfoFragment>();
 			TArrayView<FMassSimulationLODFragment> SimulationLODFragments = Context.GetMutableFragmentView<FMassSimulationLODFragment>();
-			LODSharedFragment.LODCalculator.AdjustLODFromCount(Context, SimulationLODFragments, ViewersInfoList);
+			LODSharedFragment.LODCalculator.AdjustLODFromCount(Context, ViewersInfoList, SimulationLODFragments);
 		});
 	}
 
diff --git a/Engine/Plugins/Runtime/MassGameplay/Source/MassLOD/Public/MassLODCalculator.h b/Engine/Plugins/Runtime/MassGameplay/Source/MassLOD/Public/MassLODCalculator.h
--- a/Engine/Plugins/Runtime/MassGameplay/Source/MassLOD/Public/MassLODCalculator.h
+++ b/Engine/Plugins/Runtime/MassGameplay/Source/MassLOD/Public/MassLODCalculator.h
@@ -43,11 +43,17 @@ public:
 	/**
 	 * Calculate LOD, called for each entity chunks
 	 * @Param Context of the chunk execution
-	 * @Param LODList is the fragment where calculation are stored
 	 * @Param ViewersInfoList is the source information fragment for LOD calculation
+	 * @Param LODList is the fragment where calculation are stored
+	 * @Param InfoPerViewerList is the Per viewer source information (optional parameter required when FLODLogic::bStoreInfoPerViewer is enabled)
 	 */
-	template <typename TLODFragment, typename TViewerInfoFragment>
-	void CalculateLOD(FMassExecutionContext& Context, TArrayView<TLODFragment> LODList, TConstArrayView<TViewerInfoFragment> ViewersInfoList);
+	template <typename TViewerInfoFragment,
+			  typename TLODFragment,
+			  typename TPerViewerInfoFragment = FOptionalDefaultType>
+	void CalculateLOD(FMassExecutionContext& Context,
+					  TConstArrayView<TViewerInfoFragment> ViewersInfoList,
+					  TArrayView<TLODFragment> LODList,
+					  TConstArrayView<TPerViewerInfoFragment> PerViewerInfoList = TConstArrayView<TPerViewerInfoFragment>());
 
 	/**
 	 * Adjust LOD distances by clamping them to respect the maximum LOD count
@@ -58,11 +64,17 @@ public:
 	/**
 	 * Adjust LOD from newly adjusted distances, only needed to be called when AdjustDistancesFromCount return true, called for each entity chunks
 	 * @Param Context of the chunk execution
-	 * @Param LODList is the fragment where calculation are stored
 	 * @Param ViewersInfoList is the source information fragment for LOD calculation
+	 * @Param LODList is the fragment where calculation are stored
+	 * @Param InfoPerViewerList is the Per viewer source information (optional parameter required when FLODLogic::bStoreInfoPerViewer is enabled)
 	 */
-	template <typename TLODFragment, typename TViewerInfoFragment>
-	void AdjustLODFromCount(FMassExecutionContext& Context, TArrayView<TLODFragment> LODList, TConstArrayView<TViewerInfoFragment> ViewersInfoList);
+	template <typename TViewerInfoFragment,
+			  typename TLODFragment,
+			  typename TPerViewerInfoFragment = FOptionalDefaultType>
+	void AdjustLODFromCount(FMassExecutionContext& Context,
+							TConstArrayView<TViewerInfoFragment> ViewersInfoList,
+							TArrayView<TLODFragment> LODList,
+							TConstArrayView<TPerViewerInfoFragment> PerViewerInfoList = TConstArrayView<TPerViewerInfoFragment>());
 
 	/**
 	 * Turn Off all LOD, called for each entity chunks
@@ -293,8 +305,11 @@ float TMassLODCalculator<FLODLogic>::AccumulateCountInRuntimeData(const EMassLOD
 }
 
 template <typename FLODLogic>
-template <typename TLODFragment, typename TViewerInfoFragment>
-void TMassLODCalculator<FLODLogic>::CalculateLOD(FMassExecutionContext& Context, TArrayView<TLODFragment> LODList, TConstArrayView<TViewerInfoFragment> ViewersInfoList)
+template <typename TViewerInfoFragment, typename TLODFragment, typename TPerViewerInfoFragment>
+void TMassLODCalculator<FLODLogic>::CalculateLOD(FMassExecutionContext& Context, 
+												 TConstArrayView<TViewerInfoFragment> ViewersInfoList, 
+												 TArrayView<TLODFragment> LODList, 
+												 TConstArrayView<TPerViewerInfoFragment> PerViewerInfoList)
 {
 #if WITH_MASSGAMEPLAY_DEBUG
 	if (UE::MassLOD::Debug::bLODCalculationsPaused)
@@ -329,6 +344,8 @@ void TMassLODCalculator<FLODLogic>::CalculateLOD(FMassExecutionContext& Context,
 		// Do per viewer logic if asked for
 		if (FLODLogic::bStoreInfoPerViewer)
 		{
+			const TPerViewerInfoFragment& EntityPerViewerInfo = PerViewerInfoList[EntityIdx];
+
 			SetLODPerViewerNum<FLODLogic::bCalculateLODPerViewer>(EntityLOD, Viewers.Num());
 			SetPrevLODPerViewerNum<FLODLogic::bCalculateLODPerViewer>(EntityLOD, Viewers.Num());
 			SetVisibilityPerViewerNum<FLODLogic::bDoVisibilityLogic&& FLODLogic::bStoreInfoPerViewer>(EntityLOD, Viewers.Num());
@@ -353,14 +370,14 @@ void TMassLODCalculator<FLODLogic>::CalculateLOD(FMassExecutionContext& Context,
 
 				if (Viewer.Handle.IsValid())
 				{
-					const float DistanceToFrustum = GetDistanceToFrustum<FLODLogic::bDoVisibilityLogic && FLODLogic::bStoreInfoPerViewer>(EntityViewersInfo, ViewerIdx, FLT_MAX);
+					const float DistanceToFrustum = GetDistanceToFrustum<FLODLogic::bDoVisibilityLogic && FLODLogic::bStoreInfoPerViewer>(EntityPerViewerInfo, ViewerIdx, FLT_MAX);
 					const EMassVisibility PrevVisibilityPerViewer = GetVisibilityPerViewer<FLODLogic::bDoVisibilityLogic && FLODLogic::bStoreInfoPerViewer>(EntityLOD, ViewerIdx, EMassVisibility::Max);
 					const bool bIsVisibleByViewer = CalculateVisibility(PrevVisibilityPerViewer == EMassVisibility::CanBeSeen, DistanceToFrustum);
 					bool bIsInVisibleRange = false;
 
 					if (FLODLogic::bCalculateLODPerViewer)
 					{
-						const float DistanceToViewerSq = GetDistanceToViewerSq<FLODLogic::bStoreInfoPerViewer>(EntityViewersInfo, ViewerIdx, FLT_MAX);
+						const float DistanceToViewerSq = GetDistanceToViewerSq<FLODLogic::bStoreInfoPerViewer>(EntityPerViewerInfo, ViewerIdx, FLT_MAX);
 						const EMassLOD::Type PrevLODPerViewer = GetLODPerViewer<FLODLogic::bCalculateLODPerViewer>(EntityLOD, ViewerIdx, EntityLOD.LOD);
 
 						// Find new LOD
@@ -520,8 +537,11 @@ bool TMassLODCalculator<FLODLogic>::AdjustDistancesFromCount()
 }
 
 template <typename FLODLogic>
-template <typename TLODFragment, typename TViewerInfoFragment>
-void TMassLODCalculator<FLODLogic>::AdjustLODFromCount(FMassExecutionContext& Context, TArrayView<TLODFragment> LODList, TConstArrayView<TViewerInfoFragment> ViewersInfoList)
+template <typename TViewerInfoFragment, typename TLODFragment, typename TPerViewerInfoFragment>
+void TMassLODCalculator<FLODLogic>::AdjustLODFromCount(FMassExecutionContext& Context,
+													   TConstArrayView<TViewerInfoFragment> ViewersInfoList,
+													   TArrayView<TLODFragment> LODList,
+													   TConstArrayView<TPerViewerInfoFragment> PerViewerInfoList)
 {
 	const int32 NumEntities = Context.GetNumEntities();
 	// Adjust LOD for each viewer and remember the new highest
@@ -532,6 +552,8 @@ void TMassLODCalculator<FLODLogic>::AdjustLODFromCount(FMassExecutionContext& Co
 		EMassLOD::Type HighestViewerLOD = EMassLOD::Off;
 		if (FLODLogic::bMaximizeCountPerViewer)
 		{
+			const TPerViewerInfoFragment& EntityPerViewerInfo = PerViewerInfoList[EntityIdx];
+
 			for (int ViewerIdx = 0; ViewerIdx < Viewers.Num(); ++ViewerIdx)
 			{
 				const FViewerLODInfo& Viewer = Viewers[ViewerIdx];
@@ -546,7 +568,7 @@ void TMassLODCalculator<FLODLogic>::AdjustLODFromCount(FMassExecutionContext& Co
 					continue;
 				}
 
-				const float DistanceToViewerSq = GetDistanceToViewerSq<FLODLogic::bStoreInfoPerViewer>(EntityViewersInfo, ViewerIdx, FLT_MAX);
+				const float DistanceToViewerSq = GetDistanceToViewerSq<FLODLogic::bStoreInfoPerViewer>(EntityPerViewerInfo, ViewerIdx, FLT_MAX);
 				// Using the prev visibility here as it was already updated for this frame in the CalculateLOD method and we want the previous one
 				const bool bIsVisibleByViewer = GetPrevVisibilityPerViewer<FLODLogic::bDoVisibilityLogic && FLODLogic::bStoreInfoPerViewer>(EntityLOD, ViewerIdx, EMassVisibility::Max) == EMassVisibility::CanBeSeen;
 				EMassLOD::Type LODPerViewer = GetLODPerViewer<FLODLogic::bCalculateLODPerViewer>(EntityLOD, ViewerIdx, EntityLOD.LOD);
diff --git a/Engine/Plugins/Runtime/MassGameplay/Source/MassLOD/Public/MassLODCollector.h b/Engine/Plugins/Runtime/MassGameplay/Source/MassLOD/Public/MassLODCollector.h
--- a/Engine/Plugins/Runtime/MassGameplay/Source/MassLOD/Public/MassLODCollector.h
+++ b/Engine/Plugins/Runtime/MassGameplay/Source/MassLOD/Public/MassLODCollector.h
@@ -26,10 +26,15 @@ struct TMassLODCollector : public FMassLODBaseLogic
 	 * @Param Context of the chunk execution
 	 * @Param TranformList is the fragment transforms of the entities
 	 * @Param ViewersInfoList is the fragment where to store source information for LOD calculation
+	 * @Param InfoPerViewerList is the per viewer information (optional parameter required when FLODLogic::bStoreInfoPerViewer is enabled)
 	 */
-	template <typename TTransformFragment, typename TViewerInfoFragment>
-	void CollectLODInfo(FMassExecutionContext& Context, TConstArrayView<TTransformFragment> TranformList, TArrayView<TViewerInfoFragment> ViewersInfoList);
-
+	template <typename TTransformFragment,
+		typename TViewerInfoFragment,
+		typename TPerViewerInfoFragment = FOptionalDefaultType>
+	void CollectLODInfo(FMassExecutionContext& Context, 
+						TConstArrayView<TTransformFragment> TranformList, 
+						TArrayView<TViewerInfoFragment> ViewersInfoList, 
+						TArrayView<TPerViewerInfoFragment> InfoPerViewerList = TArrayView<TPerViewerInfoFragment>());
 };
 
 template <typename FLODLogic>
@@ -39,9 +44,13 @@ void TMassLODCollector<FLODLogic>::PrepareExecution(TConstArrayView<FViewerInfo>
 }
 
 template <typename FLODLogic>
-template <typename TTransformFragment, typename TViewerInfoFragment>
-void TMassLODCollector<FLODLogic>::CollectLODInfo(FMassExecutionContext& Context, TConstArrayView<TTransformFragment> TranformList, TArrayView<TViewerInfoFragment> ViewersInfoList)
+template <typename TTransformFragment, typename TViewerInfoFragment, typename TPerViewerInfoFragment>
+void TMassLODCollector<FLODLogic>::CollectLODInfo(FMassExecutionContext& Context, 
+												  TConstArrayView<TTransformFragment> TranformList, 
+												  TArrayView<TViewerInfoFragment> ViewersInfoList, 
+												  TArrayView<TPerViewerInfoFragment> InfoPerViewerList /*= TArrayView<FOptionDefaultType>()*/)
 {
+	static TPerViewerInfoFragment DummyFragment;
 	const int32 NumEntities = Context.GetNumEntities();
 	for (int EntityIdx = 0; EntityIdx < NumEntities; EntityIdx++)
 	{
@@ -49,16 +58,17 @@ void TMassLODCollector<FLODLogic>::CollectLODInfo(FMassExecutionContext& Context
 		float ClosestDistanceToFrustum = FLT_MAX;
 		const TTransformFragment& EntityTransform = TranformList[EntityIdx];
 		TViewerInfoFragment& EntityViewerInfo = ViewersInfoList[EntityIdx];
+		TPerViewerInfoFragment& EntityInfoPerViewer = FLODLogic::bStoreInfoPerViewer ? InfoPerViewerList[EntityIdx] : DummyFragment;
 
-		SetDistanceToViewerSqNum<FLODLogic::bStoreInfoPerViewer>(EntityViewerInfo, Viewers.Num());
-		SetDistanceToFrustumNum<FLODLogic::bDoVisibilityLogic && FLODLogic::bStoreInfoPerViewer>(EntityViewerInfo, Viewers.Num());
+		SetDistanceToViewerSqNum<FLODLogic::bStoreInfoPerViewer>(EntityInfoPerViewer, Viewers.Num());
+		SetDistanceToFrustumNum<FLODLogic::bDoVisibilityLogic && FLODLogic::bStoreInfoPerViewer>(EntityInfoPerViewer, Viewers.Num());
 		for (int ViewerIdx = 0; ViewerIdx < Viewers.Num(); ++ViewerIdx)
 		{
 			const FViewerLODInfo& Viewer = Viewers[ViewerIdx];
 			if (Viewer.bClearData)
 			{
-				SetDistanceToViewerSq<FLODLogic::bStoreInfoPerViewer>(EntityViewerInfo, ViewerIdx, FLT_MAX);
-				SetDistanceToFrustum<FLODLogic::bDoVisibilityLogic && FLODLogic::bStoreInfoPerViewer>(EntityViewerInfo, ViewerIdx, FLT_MAX);
+				SetDistanceToViewerSq<FLODLogic::bStoreInfoPerViewer>(EntityInfoPerViewer, ViewerIdx, FLT_MAX);
+				SetDistanceToFrustum<FLODLogic::bDoVisibilityLogic && FLODLogic::bStoreInfoPerViewer>(EntityInfoPerViewer, ViewerIdx, FLT_MAX);
 			}
 
 			// Check to see if we want only local viewer only
@@ -76,17 +86,16 @@ void TMassLODCollector<FLODLogic>::CollectLODInfo(FMassExecutionContext& Context
 				{
 					ClosestViewerDistanceSq = DistanceToViewerSq;
 				}
-				SetDistanceToViewerSq<FLODLogic::bStoreInfoPerViewer>(EntityViewerInfo, ViewerIdx, DistanceToViewerSq);
+				SetDistanceToViewerSq<FLODLogic::bStoreInfoPerViewer>(EntityInfoPerViewer, ViewerIdx, DistanceToViewerSq);
 
 				if (FLODLogic::bDoVisibilityLogic)
 				{
 					const float DistanceToFrustum = Viewer.Frustum.DistanceTo(EntityLocation);
-					SetDistanceToFrustum<FLODLogic::bDoVisibilityLogic && FLODLogic::bStoreInfoPerViewer>(EntityViewerInfo, ViewerIdx, DistanceToFrustum);
+					SetDistanceToFrustum<FLODLogic::bDoVisibilityLogic && FLODLogic::bStoreInfoPerViewer>(EntityInfoPerViewer, ViewerIdx, DistanceToFrustum);
 					if (ClosestDistanceToFrustum > DistanceToFrustum)
 					{
 						ClosestDistanceToFrustum = DistanceToFrustum;
 					}
-
 				}
 			}
 		}
diff --git a/Engine/Plugins/Runtime/MassGameplay/Source/MassLOD/Public/MassLODLogic.h b/Engine/Plugins/Runtime/MassGameplay/Source/MassLOD/Public/MassLODLogic.h
--- a/Engine/Plugins/Runtime/MassGameplay/Source/MassLOD/Public/MassLODLogic.h
+++ b/Engine/Plugins/Runtime/MassGameplay/Source/MassLOD/Public/MassLODLogic.h
@@ -71,13 +71,17 @@ struct FMassViewerInfoFragment
 
 	// Square distances to each valid viewers (Required when FLODLogic::bStoreInfoPerViewer is enabled)
 	TArray<float> DistanceToViewerSq;
+};
 
-	// Closest distance to frustum (Required when FLODLogic::bDoVisibilityLogic is enabled)
-	float ClosestDistanceToFrustum;
+struct FMassInfoPerViewerFragment
+{
+	// Square distances to each valid viewers (Required when FLODLogic::bStoreInfoPerViewer is enabled)
+	TArray<float> DistanceToViewerSq;
 
 	// Distances to each valid viewers frustums (Required when FLODLogic::bDoVisibilityLogic and FLODLogic::bStoreInfoPerViewer are enabled)
 	TArray<float> DistanceToFrustum;
 };
+
 */
 
 /**
@@ -89,17 +93,17 @@ struct FMassViewerInfoFragment
 	TEnumAsByte<EMassLOD::Type> LOD;
 	TEnumAsByte<EMassLOD::Type> PrevLOD;
 
-	// Per viewer LOD information (Required when FLODLogic::bCalculateLODPerViewer is enabled)
-	TArray<EMassLOD::Type> LODPerViewer;
-	TArray<EMassLOD::Type> PrevLODPerViewer;
+	// Visibility information (Required when FLODLogic::bDoVisibilityLogic is enabled)
+	EMassVisibility Visibility;
+	EMassVisibility PrevVisibility
 
 	// Floating point LOD value, scaling from 0 to 3, 0 highest LOD and 3 being completely off LOD 
 	// (Required only when FLODLogic::bCalculateLODSignificance is enabled)
 	float LODSignificance = 0.0f; // 
 
-	// Visibility information (Required when FLODLogic::bDoVisibilityLogic is enabled)
-	EMassVisibility Visibility;
-	EMassVisibility PrevVisibility
+	// Per viewer LOD information (Required when FLODLogic::bCalculateLODPerViewer is enabled)
+	TArray<EMassLOD::Type> LODPerViewer;
+	TArray<EMassLOD::Type> PrevLODPerViewer;
 
 	// Visibility information per viewer (Required when FLODLogic::bDoVisibilityLogic and FLODLogicbStoreInfoPerViewer are enabled)
 	TArray<EMassVisibility> VisibilityPerViewer;
@@ -166,4 +170,7 @@ protected:
 	DECLARE_CONDITIONAL_MEMBER_ARRAY_ACCESSORS(Condition, EMassVisibility, PrevVisibilityPerViewer);
 
 	TArray<FViewerLODInfo> Viewers;
+
+	/** Optional type to use when the templated parameter is not used */
+	typedef void* FOptionalDefaultType;
 };
\ No newline at end of file
diff --git a/Engine/Plugins/Runtime/MassGameplay/Source/MassReplication/Private/MassReplicationProcessorBase.cpp b/Engine/Plugins/Runtime/MassGameplay/Source/MassReplication/Private/MassReplicationProcessorBase.cpp
--- a/Engine/Plugins/Runtime/MassGameplay/Source/MassReplication/Private/MassReplicationProcessorBase.cpp
+++ b/Engine/Plugins/Runtime/MassGameplay/Source/MassReplication/Private/MassReplicationProcessorBase.cpp
@@ -147,24 +147,24 @@ void UMassReplicationProcessorBase::Execute(UMassEntitySubsystem& EntitySubsyste
 	{
 		const TConstArrayView<FDataFragment_Transform> LocationList = Context.GetFragmentView<FDataFragment_Transform>();
 		const TArrayView<FMassReplicationViewerInfoFragment> ViewersInfoList = Context.GetMutableFragmentView<FMassReplicationViewerInfoFragment>();
-		LODCollector.CollectLODInfo(Context, LocationList, ViewersInfoList);
+		LODCollector.CollectLODInfo(Context, LocationList, ViewersInfoList, ViewersInfoList);
 	});
 
 	CalculateLODQuery.ForEachEntityChunk(EntitySubsystem, Context, [this](FMassExecutionContext& Context)
-		{
-			const TConstArrayView<FMassReplicationViewerInfoFragment> ViewersInfoList = Context.GetFragmentView<FMassReplicationViewerInfoFragment>();
-			const TArrayView<FMassReplicationLODFragment> ViewerLODList = Context.GetMutableFragmentView<FMassReplicationLODFragment>();
-			LODCalculator.CalculateLOD(Context, ViewerLODList, ViewersInfoList);
-		});
+	{
+		const TConstArrayView<FMassReplicationViewerInfoFragment> ViewersInfoList = Context.GetFragmentView<FMassReplicationViewerInfoFragment>();
+		const TArrayView<FMassReplicationLODFragment> ViewerLODList = Context.GetMutableFragmentView<FMassReplicationLODFragment>();
+		LODCalculator.CalculateLOD(Context, ViewersInfoList, ViewerLODList, ViewersInfoList);
+	});
 
 	if (LODCalculator.AdjustDistancesFromCount())
 	{
 		CalculateLODQuery.ForEachEntityChunk(EntitySubsystem, Context, [this](FMassExecutionContext& Context)
-			{
-				const TConstArrayView<FMassReplicationViewerInfoFragment> ViewersInfoList = Context.GetFragmentView<FMassReplicationViewerInfoFragment>();
-				const TArrayView<FMassReplicationLODFragment> ViewerLODList = Context.GetMutableFragmentView<FMassReplicationLODFragment>();
-				LODCalculator.AdjustLODFromCount(Context, ViewerLODList, ViewersInfoList);
-			});
+		{
+			const TConstArrayView<FMassReplicationViewerInfoFragment> ViewersInfoList = Context.GetFragmentView<FMassReplicationViewerInfoFragment>();
+			const TArrayView<FMassReplicationLODFragment> ViewerLODList = Context.GetMutableFragmentView<FMassReplicationLODFragment>();
+			LODCalculator.AdjustLODFromCount(Context, ViewersInfoList, ViewerLODList, ViewersInfoList);
+		});
 	}
 #endif //UE_REPLICATION_COMPILE_SERVER_CODE
 }
diff --git a/Engine/Plugins/Runtime/MassGameplay/Source/MassRepresentation/Public/MassVisualizationLODProcessor.h b/Engine/Plugins/Runtime/MassGameplay/Source/MassRepresentation/Public/MassVisualizationLODProcessor.h
--- a/Engine/Plugins/Runtime/MassGameplay/Source/MassRepresentation/Public/MassVisualizationLODProcessor.h
+++ b/Engine/Plugins/Runtime/MassGameplay/Source/MassRepresentation/Public/MassVisualizationLODProcessor.h
@@ -107,7 +107,7 @@ void UMassVisualizationLODProcessor::ExecuteInternal(UMassEntitySubsystem& Entit
 		{
 			TArrayView<FMassRepresentationLODFragment> RepresentationLODList = Context.GetMutableFragmentView<FMassRepresentationLODFragment>();
 			TConstArrayView<TMassViewerLODInfoFragment> ViewerInfoList = Context.GetFragmentView<TMassViewerLODInfoFragment>();
-			LODCalculator.CalculateLOD(Context, RepresentationLODList, ViewerInfoList);
+			LODCalculator.CalculateLOD(Context, ViewerInfoList, RepresentationLODList);
 		};
 		CloseEntityQuery.ForEachEntityChunk(EntitySubsystem, Context, CalculateLOD);
 
@@ -129,7 +129,7 @@ void UMassVisualizationLODProcessor::ExecuteInternal(UMassEntitySubsystem& Entit
 			{
 				TConstArrayView<TMassViewerLODInfoFragment> ViewerInfoList = Context.GetFragmentView<TMassViewerLODInfoFragment>();
 				TArrayView<FMassRepresentationLODFragment> RepresentationLODList = Context.GetMutableFragmentView<FMassRepresentationLODFragment>();
-				LODCalculator.AdjustLODFromCount(Context, RepresentationLODList, ViewerInfoList);
+				LODCalculator.AdjustLODFromCount(Context, ViewerInfoList, RepresentationLODList);
 			});
 			// Far entities do not need to maximice count
 		}
