Fix Hololens build break

#rnx

diff --git a/Engine/Plugins/Runtime/MassGameplay/Source/MassLOD/Public/MassLODCalculator.h b/Engine/Plugins/Runtime/MassGameplay/Source/MassLOD/Public/MassLODCalculator.h
--- a/Engine/Plugins/Runtime/MassGameplay/Source/MassLOD/Public/MassLODCalculator.h
+++ b/Engine/Plugins/Runtime/MassGameplay/Source/MassLOD/Public/MassLODCalculator.h
@@ -42,18 +42,35 @@ public:
 
 	/**
 	 * Calculate LOD, called for each entity chunks
+	 * Use next method when FLODLogic::bStoreInfoPerViewer is enabled
 	 * @Param Context of the chunk execution
 	 * @Param ViewersInfoList is the source information fragment for LOD calculation
 	 * @Param LODList is the fragment where calculation are stored
-	 * @Param InfoPerViewerList is the Per viewer source information (optional parameter required when FLODLogic::bStoreInfoPerViewer is enabled)
+	 */
+	template <typename TViewerInfoFragment, typename TLODFragment>
+	FORCEINLINE void CalculateLOD(FMassExecutionContext& Context,
+					  TConstArrayView<TViewerInfoFragment> ViewersInfoList,
+					  TArrayView<TLODFragment> LODList)
+	{
+		CalculateLOD(Context, ViewersInfoList, LODList, TConstArrayView<void*>());
+	}
+
+	/**
+	 * Calculate LOD, called for each entity chunks
+	 * Use this version when FLODLogic::bStoreInfoPerViewer is enabled
+	 * It calculates a LOD per viewer and needs information per viewer via PerViewerInfoList fragments
+	 * @Param Context of the chunk execution
+	 * @Param ViewersInfoList is the source information fragment for LOD calculation
+	 * @Param LODList is the fragment where calculation are stored
+	 * @Param PerViewerInfoList is the Per viewer source information
 	 */
 	template <typename TViewerInfoFragment,
 			  typename TLODFragment,
-			  typename TPerViewerInfoFragment = FOptionalDefaultType>
+			  typename TPerViewerInfoFragment>
 	void CalculateLOD(FMassExecutionContext& Context,
 					  TConstArrayView<TViewerInfoFragment> ViewersInfoList,
 					  TArrayView<TLODFragment> LODList,
-					  TConstArrayView<TPerViewerInfoFragment> PerViewerInfoList = TConstArrayView<TPerViewerInfoFragment>());
+					  TConstArrayView<TPerViewerInfoFragment> PerViewerInfoList);
 
 	/**
 	 * Adjust LOD distances by clamping them to respect the maximum LOD count
@@ -63,18 +80,37 @@ public:
 
 	/**
 	 * Adjust LOD from newly adjusted distances, only needed to be called when AdjustDistancesFromCount return true, called for each entity chunks
+	 * Use next method when FLODLogic::bStoreInfoPerViewer is enabled
+	 * @Param Context of the chunk execution
+	 * @Param ViewersInfoList is the source information fragment for LOD calculation
+	 * @Param LODList is the fragment where calculation are stored
+	 */
+	template <typename TViewerInfoFragment, 
+			  typename TLODFragment>
+	FORCEINLINE void AdjustLODFromCount(FMassExecutionContext& Context,
+										TConstArrayView<TViewerInfoFragment> ViewersInfoList,
+										TArrayView<TLODFragment> LODList)
+	{
+		AdjustLODFromCount(Context, ViewersInfoList, LODList, TConstArrayView<void*>());
+	}
+
+
+	/**
+	 * Adjust LOD from newly adjusted distances, only needed to be called when AdjustDistancesFromCount return true, called for each entity chunks
+	 * Use this version when FLODLogic::bStoreInfoPerViewer is enabled
+	 * It calculates a LOD per viewer and needs information per viewer via PerViewerInfoList fragments
 	 * @Param Context of the chunk execution
 	 * @Param ViewersInfoList is the source information fragment for LOD calculation
 	 * @Param LODList is the fragment where calculation are stored
-	 * @Param InfoPerViewerList is the Per viewer source information (optional parameter required when FLODLogic::bStoreInfoPerViewer is enabled)
+	 * @Param PerViewerInfoList is the Per viewer source information
 	 */
 	template <typename TViewerInfoFragment,
 			  typename TLODFragment,
-			  typename TPerViewerInfoFragment = FOptionalDefaultType>
+			  typename TPerViewerInfoFragment>
 	void AdjustLODFromCount(FMassExecutionContext& Context,
 							TConstArrayView<TViewerInfoFragment> ViewersInfoList,
 							TArrayView<TLODFragment> LODList,
-							TConstArrayView<TPerViewerInfoFragment> PerViewerInfoList = TConstArrayView<TPerViewerInfoFragment>());
+							TConstArrayView<TPerViewerInfoFragment> PerViewerInfoList);
 
 	/**
 	 * Turn Off all LOD, called for each entity chunks
diff --git a/Engine/Plugins/Runtime/MassGameplay/Source/MassLOD/Public/MassLODCollector.h b/Engine/Plugins/Runtime/MassGameplay/Source/MassLOD/Public/MassLODCollector.h
--- a/Engine/Plugins/Runtime/MassGameplay/Source/MassLOD/Public/MassLODCollector.h
+++ b/Engine/Plugins/Runtime/MassGameplay/Source/MassLOD/Public/MassLODCollector.h
@@ -23,18 +23,37 @@ struct TMassLODCollector : public FMassLODBaseLogic
 
 	/**
 	 * Collects the information for LOD calculation, called for each entity chunks
+	 * Use next method when FLODLogic::bStoreInfoPerViewer is enabled
 	 * @Param Context of the chunk execution
-	 * @Param TranformList is the fragment transforms of the entities
+	 * @Param TransformList is the fragment transforms of the entities
 	 * @Param ViewersInfoList is the fragment where to store source information for LOD calculation
-	 * @Param InfoPerViewerList is the per viewer information (optional parameter required when FLODLogic::bStoreInfoPerViewer is enabled)
+	 */
+	template <typename TTransformFragment, 
+			  typename TViewerInfoFragment>
+	FORCEINLINE void CollectLODInfo(FMassExecutionContext& Context, 
+									TConstArrayView<TTransformFragment> TransformList, 
+									TArrayView<TViewerInfoFragment> ViewersInfoList)
+	{
+		CollectLODInfo(Context, TransformList, ViewersInfoList, TArrayView<void*>());
+	}
+
+
+	/**
+	 * Collects the information for LOD calculation, called for each entity chunks
+	 * Use this version when FLODLogic::bStoreInfoPerViewer is enabled
+	 * It collects information per viewer into the PerViewerInfoList fragments
+	 * @Param Context of the chunk execution
+	 * @Param TransformList is the fragment transforms of the entities
+	 * @Param ViewersInfoList is the fragment where to store source information for LOD calculation
+	 * @Param PerViewerInfoList is the per viewer information
 	 */
 	template <typename TTransformFragment,
-		typename TViewerInfoFragment,
-		typename TPerViewerInfoFragment = FOptionalDefaultType>
+			  typename TViewerInfoFragment,
+			  typename TPerViewerInfoFragment>
 	void CollectLODInfo(FMassExecutionContext& Context, 
-						TConstArrayView<TTransformFragment> TranformList, 
+						TConstArrayView<TTransformFragment> TransformList, 
 						TArrayView<TViewerInfoFragment> ViewersInfoList, 
-						TArrayView<TPerViewerInfoFragment> InfoPerViewerList = TArrayView<TPerViewerInfoFragment>());
+						TArrayView<TPerViewerInfoFragment> PerViewerInfoList);
 };
 
 template <typename FLODLogic>
@@ -44,11 +63,13 @@ void TMassLODCollector<FLODLogic>::PrepareExecution(TConstArrayView<FViewerInfo>
 }
 
 template <typename FLODLogic>
-template <typename TTransformFragment, typename TViewerInfoFragment, typename TPerViewerInfoFragment>
+template <typename TTransformFragment, 
+		  typename TViewerInfoFragment, 
+		  typename TPerViewerInfoFragment>
 void TMassLODCollector<FLODLogic>::CollectLODInfo(FMassExecutionContext& Context, 
-												  TConstArrayView<TTransformFragment> TranformList, 
+												  TConstArrayView<TTransformFragment> TransformList, 
 												  TArrayView<TViewerInfoFragment> ViewersInfoList, 
-												  TArrayView<TPerViewerInfoFragment> InfoPerViewerList /*= TArrayView<FOptionDefaultType>()*/)
+												  TArrayView<TPerViewerInfoFragment> PerViewerInfoList)
 {
 	static TPerViewerInfoFragment DummyFragment;
 	const int32 NumEntities = Context.GetNumEntities();
@@ -56,9 +77,9 @@ void TMassLODCollector<FLODLogic>::CollectLODInfo(FMassExecutionContext& Context
 	{
 		float ClosestViewerDistanceSq = FLT_MAX;
 		float ClosestDistanceToFrustum = FLT_MAX;
-		const TTransformFragment& EntityTransform = TranformList[EntityIdx];
+		const TTransformFragment& EntityTransform = TransformList[EntityIdx];
 		TViewerInfoFragment& EntityViewerInfo = ViewersInfoList[EntityIdx];
-		TPerViewerInfoFragment& EntityInfoPerViewer = FLODLogic::bStoreInfoPerViewer ? InfoPerViewerList[EntityIdx] : DummyFragment;
+		TPerViewerInfoFragment& EntityInfoPerViewer = FLODLogic::bStoreInfoPerViewer ? PerViewerInfoList[EntityIdx] : DummyFragment;
 
 		SetDistanceToViewerSqNum<FLODLogic::bStoreInfoPerViewer>(EntityInfoPerViewer, Viewers.Num());
 		SetDistanceToFrustumNum<FLODLogic::bDoVisibilityLogic && FLODLogic::bStoreInfoPerViewer>(EntityInfoPerViewer, Viewers.Num());
diff --git a/Engine/Plugins/Runtime/MassGameplay/Source/MassLOD/Public/MassLODLogic.h b/Engine/Plugins/Runtime/MassGameplay/Source/MassLOD/Public/MassLODLogic.h
--- a/Engine/Plugins/Runtime/MassGameplay/Source/MassLOD/Public/MassLODLogic.h
+++ b/Engine/Plugins/Runtime/MassGameplay/Source/MassLOD/Public/MassLODLogic.h
@@ -170,7 +170,4 @@ protected:
 	DECLARE_CONDITIONAL_MEMBER_ARRAY_ACCESSORS(Condition, EMassVisibility, PrevVisibilityPerViewer);
 
 	TArray<FViewerLODInfo> Viewers;
-
-	/** Optional type to use when the templated parameter is not used */
-	typedef void* FOptionalDefaultType;
 };
\ No newline at end of file
