Moved the LODLogic template parameters into the method to support multiple gathering type from a single LODCollector/Calculator class.

#rnx

diff --git a/Engine/Plugins/Runtime/MassGameplay/Source/MassLOD/Public/MassLODCalculator.h b/Engine/Plugins/Runtime/MassGameplay/Source/MassLOD/Public/MassLODCalculator.h
--- a/Engine/Plugins/Runtime/MassGameplay/Source/MassLOD/Public/MassLODCalculator.h
+++ b/Engine/Plugins/Runtime/MassGameplay/Source/MassLOD/Public/MassLODCalculator.h
@@ -47,12 +47,22 @@ public:
 	 * @Param ViewersInfoList is the source information fragment for LOD calculation
 	 * @Param LODList is the fragment where calculation are stored
 	 */
-	template <typename TViewerInfoFragment, typename TLODFragment>
+	template <typename TViewerInfoFragment,
+			  typename TLODFragment,
+			  bool bCalculateLocalViewers = FLODLogic::bLocalViewersOnly,
+			  bool bCalculateVisibility = FLODLogic::bDoVisibilityLogic>
 	FORCEINLINE void CalculateLOD(FMassExecutionContext& Context,
-					  TConstArrayView<TViewerInfoFragment> ViewersInfoList,
-					  TArrayView<TLODFragment> LODList)
+					  			  TConstArrayView<TViewerInfoFragment> ViewersInfoList,
+					  			  TArrayView<TLODFragment> LODList)
 	{
-		CalculateLOD(Context, ViewersInfoList, LODList, TConstArrayView<void*>());
+		CalculateLOD<TViewerInfoFragment,
+			TLODFragment,
+			/*TPerViewerInfoFragment*/void*,
+			bCalculateLocalViewers,
+			bCalculateVisibility,
+			/*bCalculateLODPerViewer*/false,
+			/*bCalculateVisibilityPerViewer*/false,
+			/*bMaximizeCountPerViewer*/false>(Context, ViewersInfoList, LODList, TConstArrayView<void*>());
 	}
 
 	/**
@@ -66,7 +76,12 @@ public:
 	 */
 	template <typename TViewerInfoFragment,
 			  typename TLODFragment,
-			  typename TPerViewerInfoFragment>
+			  typename TPerViewerInfoFragment,
+			  bool bCalculateLocalViewers = FLODLogic::bLocalViewersOnly,
+			  bool bCalculateVisibility = FLODLogic::bDoVisibilityLogic,
+			  bool bCalculateLODPerViewer = FLODLogic::bCalculateLODPerViewer,
+			  bool bCalculateVisibilityPerViewer = FLODLogic::bDoVisibilityLogic && FLODLogic::bStoreInfoPerViewer,
+			  bool bMaximizeCountPerViewer = FLODLogic::bMaximizeCountPerViewer>
 	void CalculateLOD(FMassExecutionContext& Context,
 					  TConstArrayView<TViewerInfoFragment> ViewersInfoList,
 					  TArrayView<TLODFragment> LODList,
@@ -76,6 +91,9 @@ public:
 	 * Adjust LOD distances by clamping them to respect the maximum LOD count
 	 * @Return true if any LOD distances clamping was done
 	 */
+	template <bool bCalculateVisibility = FLODLogic::bDoVisibilityLogic,
+			  bool bCalculateVisibilityPerViewer = FLODLogic::bDoVisibilityLogic && FLODLogic::bStoreInfoPerViewer,
+			  bool bMaximizeCountPerViewer = FLODLogic::bMaximizeCountPerViewer>
 	bool AdjustDistancesFromCount();
 
 	/**
@@ -86,12 +104,21 @@ public:
 	 * @Param LODList is the fragment where calculation are stored
 	 */
 	template <typename TViewerInfoFragment, 
-			  typename TLODFragment>
+			  typename TLODFragment,
+			  bool bCalculateLocalViewers = FLODLogic::bLocalViewersOnly,
+			  bool bCalculateVisibility = FLODLogic::bDoVisibilityLogic>
 	FORCEINLINE void AdjustLODFromCount(FMassExecutionContext& Context,
 										TConstArrayView<TViewerInfoFragment> ViewersInfoList,
 										TArrayView<TLODFragment> LODList)
 	{
-		AdjustLODFromCount(Context, ViewersInfoList, LODList, TConstArrayView<void*>());
+		AdjustLODFromCount<TViewerInfoFragment,
+			TLODFragment,
+			/*TPerViewerInfoFragment*/void*,
+			bCalculateLocalViewers,
+			bCalculateVisibility,
+			/*bCalculateLODPerViewer*/false,
+			/*bCalculateVisibilityPerViewer*/false,
+			/*bMaximizeCountPerViewer*/false>(Context, ViewersInfoList, LODList, TConstArrayView<void*>());
 	}
 
 
@@ -106,7 +133,12 @@ public:
 	 */
 	template <typename TViewerInfoFragment,
 			  typename TLODFragment,
-			  typename TPerViewerInfoFragment>
+			  typename TPerViewerInfoFragment,
+			  bool bCalculateLocalViewers = FLODLogic::bLocalViewersOnly,
+			  bool bCalculateVisibility = FLODLogic::bDoVisibilityLogic,
+			  bool bCalculateLODPerViewer = FLODLogic::bCalculateLODPerViewer,
+			  bool bCalculateVisibilityPerViewer = FLODLogic::bDoVisibilityLogic && FLODLogic::bStoreInfoPerViewer,
+			  bool bMaximizeCountPerViewer = FLODLogic::bMaximizeCountPerViewer>
 	void AdjustLODFromCount(FMassExecutionContext& Context,
 							TConstArrayView<TViewerInfoFragment> ViewersInfoList,
 							TArrayView<TLODFragment> LODList,
@@ -172,11 +204,13 @@ protected:
 	};
 
 
-	template <bool bCalculateLODSignificance>
+	template <bool bCalculateVisibility, bool bCalculateLODSignificance>
 	float AccumulateCountInRuntimeData(const EMassLOD::Type LOD, const float ViewerDistanceSq, const bool bIsVisible, FMassLODRuntimeData& Data) const;
 
+	template <bool bCalculateVisibility>
 	bool AdjustDistancesFromCountForRuntimeData(const TStaticArray<int32, EMassLOD::Max>& MaxCount, FMassLODRuntimeData& RuntimeData) const;
 
+	template<bool bCalculateVisibility>
 	EMassLOD::Type ComputeLODFromSettings(const EMassLOD::Type PrevLOD, const float DistanceToViewerSq, const bool bIsVisible, bool* bIsInAVisibleRange, const FMassLODRuntimeData& Data) const;
 
 	bool CalculateVisibility(const bool bWasVisible, const float DistanceToFrustum) const;
@@ -304,10 +338,10 @@ void TMassLODCalculator<FLODLogic>::PrepareExecution(TConstArrayView<FViewerInfo
 
 
 template <typename FLODLogic>
-template <bool bCalculateLODSignificance>
+template <bool bCalculateVisibility, bool bCalculateLODSignificance>
 float TMassLODCalculator<FLODLogic>::AccumulateCountInRuntimeData(const EMassLOD::Type LOD, const float ViewerDistanceSq, const bool bIsVisible, FMassLODRuntimeData& Data) const
 {
-	TStaticArray< TStaticArray<int32, UE::MassLOD::MaxBucketsPerLOD>, EMassLOD::Max>& BucketCounts = FLODLogic::bDoVisibilityLogic && bIsVisible ? Data.VisibleBucketCounts : Data.BaseBucketCounts;
+	TStaticArray< TStaticArray<int32, UE::MassLOD::MaxBucketsPerLOD>, EMassLOD::Max>& BucketCounts = bCalculateVisibility && bIsVisible ? Data.VisibleBucketCounts : Data.BaseBucketCounts;
 
 	// Cumulate LOD in buckets for Max LOD count calculation
 	if (LOD == EMassLOD::Off)
@@ -321,8 +355,8 @@ float TMassLODCalculator<FLODLogic>::AccumulateCountInRuntimeData(const EMassLOD
 	}
 	else
 	{
-		const TStaticArray<float, EMassLOD::Max>& BucketSize = FLODLogic::bDoVisibilityLogic && bIsVisible ? VisibleBucketSize : BaseBucketSize;
-		const TStaticArray<float, EMassLOD::Max>& AdjustedLODDistance = FLODLogic::bDoVisibilityLogic && bIsVisible ? Data.AdjustedVisibleLODDistance : Data.AdjustedBaseLODDistance;
+		const TStaticArray<float, EMassLOD::Max>& BucketSize = bCalculateVisibility && bIsVisible ? VisibleBucketSize : BaseBucketSize;
+		const TStaticArray<float, EMassLOD::Max>& AdjustedLODDistance = bCalculateVisibility && bIsVisible ? Data.AdjustedVisibleLODDistance : Data.AdjustedBaseLODDistance;
 
 		const int32 LODDistIdx = (int32)LOD;
 
@@ -341,12 +375,18 @@ float TMassLODCalculator<FLODLogic>::AccumulateCountInRuntimeData(const EMassLOD
 }
 
 template <typename FLODLogic>
-template <typename TViewerInfoFragment, typename TLODFragment, typename TPerViewerInfoFragment>
+template <typename TViewerInfoFragment, typename TLODFragment, typename TPerViewerInfoFragment,
+		  bool bCalculateLocalViewers, bool bCalculateVisibility, bool bCalculateLODPerViewer, bool bCalculateVisibilityPerViewer,bool bMaximizeCountPerViewer>
 void TMassLODCalculator<FLODLogic>::CalculateLOD(FMassExecutionContext& Context, 
 												 TConstArrayView<TViewerInfoFragment> ViewersInfoList, 
 												 TArrayView<TLODFragment> LODList, 
 												 TConstArrayView<TPerViewerInfoFragment> PerViewerInfoList)
 {
+	static_assert(!bCalculateVisibility || FLODLogic::bDoVisibilityLogic, "FLODLogic must have bDoVisibilityLogic enabled to calculate visibility.");
+	static_assert(!bCalculateLODPerViewer || FLODLogic::bCalculateLODPerViewer, "FLODLogic must have bCalculateLODPerViewer enabled to calculate LOD Per viewer.");
+	static_assert(!bCalculateVisibilityPerViewer || (FLODLogic::bDoVisibilityLogic && FLODLogic::bStoreInfoPerViewer), "FLODLogic must have bDoVisibilityLogic and bStoreInfoPerViewer enabled to calculate visibility per viewer.");
+	static_assert(!bMaximizeCountPerViewer || FLODLogic::bMaximizeCountPerViewer, "FLODLogic must have bMaximizeCountPerViewer enabled to maximize count per viewer.");
+
 #if WITH_MASSGAMEPLAY_DEBUG
 	if (UE::MassLOD::Debug::bLODCalculationsPaused)
 	{
@@ -360,79 +400,79 @@ void TMassLODCalculator<FLODLogic>::CalculateLOD(FMassExecutionContext& Context,
 		// Calculate the LOD purely upon distances
 		const TViewerInfoFragment& EntityViewersInfo = ViewersInfoList[EntityIdx];
 		TLODFragment& EntityLOD = LODList[EntityIdx];
-		const float ClosestDistanceToFrustum = GetClosestDistanceToFrustum<FLODLogic::bDoVisibilityLogic>(EntityViewersInfo, FLT_MAX);
-		const EMassVisibility PrevVisibility = GetVisibility<FLODLogic::bDoVisibilityLogic>(EntityLOD, EMassVisibility::Max);
+		const float ClosestDistanceToFrustum = GetClosestDistanceToFrustum<bCalculateVisibility>(EntityViewersInfo, FLT_MAX);
+		const EMassVisibility PrevVisibility = GetVisibility<bCalculateVisibility>(EntityLOD, EMassVisibility::Max);
 		const bool bIsVisibleByAViewer = CalculateVisibility(PrevVisibility == EMassVisibility::CanBeSeen, ClosestDistanceToFrustum);
 		bool bIsInAVisibleRange = false;
 
 		// Find new LOD
 		EntityLOD.PrevLOD = EntityLOD.LOD;
-		EntityLOD.LOD = ComputeLODFromSettings(EntityLOD.PrevLOD, EntityViewersInfo.ClosestViewerDistanceSq, bIsVisibleByAViewer, &bIsInAVisibleRange, RuntimeData);
+		EntityLOD.LOD = ComputeLODFromSettings<bCalculateVisibility>(EntityLOD.PrevLOD, EntityViewersInfo.ClosestViewerDistanceSq, bIsVisibleByAViewer, &bIsInAVisibleRange, RuntimeData);
 
 		// Set visibility
-		SetPrevVisibility<FLODLogic::bDoVisibilityLogic>(EntityLOD, PrevVisibility);
-		SetVisibility<FLODLogic::bDoVisibilityLogic>(EntityLOD, bIsVisibleByAViewer ? EMassVisibility::CanBeSeen : (bIsInAVisibleRange ? EMassVisibility::CulledByFrustum : EMassVisibility::CulledByDistance));
+		SetPrevVisibility<bCalculateVisibility>(EntityLOD, PrevVisibility);
+		SetVisibility<bCalculateVisibility>(EntityLOD, bIsVisibleByAViewer ? EMassVisibility::CanBeSeen : (bIsInAVisibleRange ? EMassVisibility::CulledByFrustum : EMassVisibility::CulledByDistance));
 
 		// Accumulate in buckets
-		const float LODSignificance = AccumulateCountInRuntimeData<FLODLogic::bCalculateLODSignificance>(EntityLOD.LOD, EntityViewersInfo.ClosestViewerDistanceSq, bIsVisibleByAViewer, RuntimeData);
+		const float LODSignificance = AccumulateCountInRuntimeData<bCalculateVisibility, FLODLogic::bCalculateLODSignificance>(EntityLOD.LOD, EntityViewersInfo.ClosestViewerDistanceSq, bIsVisibleByAViewer, RuntimeData);
 		SetLODSignificance<FLODLogic::bCalculateLODSignificance>(EntityLOD, LODSignificance);
 
 		// Do per viewer logic if asked for
-		if (FLODLogic::bStoreInfoPerViewer)
+		if (bCalculateLODPerViewer || bCalculateVisibilityPerViewer)
 		{
 			const TPerViewerInfoFragment& EntityPerViewerInfo = PerViewerInfoList[EntityIdx];
 
-			SetLODPerViewerNum<FLODLogic::bCalculateLODPerViewer>(EntityLOD, Viewers.Num());
-			SetPrevLODPerViewerNum<FLODLogic::bCalculateLODPerViewer>(EntityLOD, Viewers.Num());
-			SetVisibilityPerViewerNum<FLODLogic::bDoVisibilityLogic&& FLODLogic::bStoreInfoPerViewer>(EntityLOD, Viewers.Num());
-			SetPrevVisibilityPerViewerNum<FLODLogic::bDoVisibilityLogic&& FLODLogic::bStoreInfoPerViewer>(EntityLOD, Viewers.Num());
+			SetLODPerViewerNum<bCalculateLODPerViewer>(EntityLOD, Viewers.Num());
+			SetPrevLODPerViewerNum<bCalculateLODPerViewer>(EntityLOD, Viewers.Num());
+			SetVisibilityPerViewerNum<bCalculateVisibilityPerViewer>(EntityLOD, Viewers.Num());
+			SetPrevVisibilityPerViewerNum<bCalculateVisibilityPerViewer>(EntityLOD, Viewers.Num());
 
 			for (int ViewerIdx = 0; ViewerIdx < Viewers.Num(); ++ViewerIdx)
 			{
 				const FViewerLODInfo& Viewer = Viewers[ViewerIdx];
 				if (Viewer.bClearData)
 				{
-					SetLODPerViewer<FLODLogic::bCalculateLODPerViewer>(EntityLOD, ViewerIdx, EMassLOD::Max);
-					SetPrevLODPerViewer<FLODLogic::bCalculateLODPerViewer>(EntityLOD, ViewerIdx, EMassLOD::Max);
-					SetPrevVisibilityPerViewer<FLODLogic::bDoVisibilityLogic&& FLODLogic::bStoreInfoPerViewer>(EntityLOD, ViewerIdx, EMassVisibility::Max);
-					SetVisibilityPerViewer<FLODLogic::bDoVisibilityLogic && FLODLogic::bStoreInfoPerViewer>(EntityLOD, ViewerIdx, EMassVisibility::Max);
+					SetLODPerViewer<bCalculateLODPerViewer>(EntityLOD, ViewerIdx, EMassLOD::Max);
+					SetPrevLODPerViewer<bCalculateLODPerViewer>(EntityLOD, ViewerIdx, EMassLOD::Max);
+					SetPrevVisibilityPerViewer<bCalculateVisibilityPerViewer>(EntityLOD, ViewerIdx, EMassVisibility::Max);
+					SetVisibilityPerViewer<bCalculateVisibilityPerViewer>(EntityLOD, ViewerIdx, EMassVisibility::Max);
 				}
 
 				// Check to see if we want only local viewer only
-				if (FLODLogic::bLocalViewersOnly && !Viewer.bLocal)
+				if (bCalculateLocalViewers && !Viewer.bLocal)
 				{
 					continue;
 				}
 
 				if (Viewer.Handle.IsValid())
 				{
-					const float DistanceToFrustum = GetDistanceToFrustum<FLODLogic::bDoVisibilityLogic && FLODLogic::bStoreInfoPerViewer>(EntityPerViewerInfo, ViewerIdx, FLT_MAX);
-					const EMassVisibility PrevVisibilityPerViewer = GetVisibilityPerViewer<FLODLogic::bDoVisibilityLogic && FLODLogic::bStoreInfoPerViewer>(EntityLOD, ViewerIdx, EMassVisibility::Max);
+					const float DistanceToFrustum = GetDistanceToFrustum<bCalculateVisibilityPerViewer>(EntityPerViewerInfo, ViewerIdx, FLT_MAX);
+					const EMassVisibility PrevVisibilityPerViewer = GetVisibilityPerViewer<bCalculateVisibilityPerViewer>(EntityLOD, ViewerIdx, EMassVisibility::Max);
 					const bool bIsVisibleByViewer = CalculateVisibility(PrevVisibilityPerViewer == EMassVisibility::CanBeSeen, DistanceToFrustum);
 					bool bIsInVisibleRange = false;
 
-					if (FLODLogic::bCalculateLODPerViewer)
+					if (bCalculateLODPerViewer)
 					{
-						const float DistanceToViewerSq = GetDistanceToViewerSq<FLODLogic::bStoreInfoPerViewer>(EntityPerViewerInfo, ViewerIdx, FLT_MAX);
-						const EMassLOD::Type PrevLODPerViewer = GetLODPerViewer<FLODLogic::bCalculateLODPerViewer>(EntityLOD, ViewerIdx, EntityLOD.LOD);
+						const float DistanceToViewerSq = GetDistanceToViewerSq<bCalculateLODPerViewer>(EntityPerViewerInfo, ViewerIdx, FLT_MAX);
+						const EMassLOD::Type PrevLODPerViewer = GetLODPerViewer<bCalculateLODPerViewer>(EntityLOD, ViewerIdx, EntityLOD.LOD);
 
 						// Find new LOD
-						const EMassLOD::Type LODPerViewer = ComputeLODFromSettings(PrevLODPerViewer, DistanceToViewerSq, bIsInVisibleRange, &bIsInAVisibleRange, RuntimeData);
+						const EMassLOD::Type LODPerViewer = ComputeLODFromSettings<bCalculateVisibilityPerViewer>(PrevLODPerViewer, DistanceToViewerSq, bIsInVisibleRange, &bIsInAVisibleRange, RuntimeData);
 
 						// Set Per viewer LOD
-						SetPrevLODPerViewer<FLODLogic::bCalculateLODPerViewer>(EntityLOD, ViewerIdx, PrevLODPerViewer);
-						SetLODPerViewer<FLODLogic::bCalculateLODPerViewer>(EntityLOD, ViewerIdx, LODPerViewer);
+						SetPrevLODPerViewer<bCalculateLODPerViewer>(EntityLOD, ViewerIdx, PrevLODPerViewer);
+						SetLODPerViewer<bCalculateLODPerViewer>(EntityLOD, ViewerIdx, LODPerViewer);
 
-						if (FLODLogic::bMaximizeCountPerViewer)
+						if (bMaximizeCountPerViewer)
 						{
 							// Accumulate in buckets
-							AccumulateCountInRuntimeData<false>(LODPerViewer, DistanceToViewerSq, bIsInVisibleRange, RuntimeDataPerViewer[ViewerIdx]);
+							AccumulateCountInRuntimeData<bCalculateVisibilityPerViewer, false>(LODPerViewer, DistanceToViewerSq, bIsInVisibleRange, RuntimeDataPerViewer[ViewerIdx]);
 						}
 					}
 
 					// Set visibility
-					SetPrevVisibilityPerViewer<FLODLogic::bDoVisibilityLogic&& FLODLogic::bStoreInfoPerViewer>(EntityLOD, ViewerIdx, bIsVisibleByViewer ? EMassVisibility::CanBeSeen : (bIsInVisibleRange ? EMassVisibility::CulledByFrustum : EMassVisibility::CulledByDistance));
-					SetVisibilityPerViewer<FLODLogic::bDoVisibilityLogic&& FLODLogic::bStoreInfoPerViewer>(EntityLOD, ViewerIdx, PrevVisibilityPerViewer);
+					SetPrevVisibilityPerViewer<bCalculateVisibilityPerViewer>(EntityLOD, ViewerIdx, bIsVisibleByViewer ? EMassVisibility::CanBeSeen : (bIsInVisibleRange ? EMassVisibility::CulledByFrustum : EMassVisibility::CulledByDistance));
+					SetVisibilityPerViewer<bCalculateVisibilityPerViewer>(EntityLOD, ViewerIdx, PrevVisibilityPerViewer);
 				}
 			}
 		}
@@ -440,6 +480,7 @@ void TMassLODCalculator<FLODLogic>::CalculateLOD(FMassExecutionContext& Context,
 }
 
 template <typename FLODLogic>
+template <bool bCalculateVisibility>
 bool TMassLODCalculator<FLODLogic>::AdjustDistancesFromCountForRuntimeData(const TStaticArray<int32, EMassLOD::Max>& MaxCount, FMassLODRuntimeData& Data) const
 {
 	int32 Count = 0;
@@ -468,7 +509,7 @@ bool TMassLODCalculator<FLODLogic>::AdjustDistancesFromCountForRuntimeData(const
 		// Count entities through all buckets of this LOD
 		for (int32 BucketIdx = 0; BucketIdx < UE::MassLOD::MaxBucketsPerLOD; ++BucketIdx)
 		{
-			if (FLODLogic::bDoVisibilityLogic)
+			if (bCalculateVisibility)
 			{
 				// Do visible count first to prioritize them over none visible for that bucket idx
 				Count += Data.VisibleBucketCounts[BucketLODIdx][BucketIdx];
@@ -519,7 +560,7 @@ bool TMassLODCalculator<FLODLogic>::AdjustDistancesFromCountForRuntimeData(const
 				// Adjust distance for this LOD
 				Data.AdjustedBaseLODDistance[ProcessingLODIdx] = BaseLODDistance[BucketLODIdx] + (BucketIdx * BaseBucketSize[BucketLODIdx]);
 				Data.AdjustedBaseLODDistanceSq[ProcessingLODIdx] = FMath::Square(Data.AdjustedBaseLODDistance[ProcessingLODIdx]);
-				if (FLODLogic::bDoVisibilityLogic)
+				if (bCalculateVisibility)
 				{
 					Data.AdjustedVisibleLODDistance[ProcessingLODIdx] = VisibleLODDistance[BucketLODIdx] + ((BucketIdx + 1) * VisibleBucketSize[BucketLODIdx]);
 					Data.AdjustedVisibleLODDistanceSq[ProcessingLODIdx] = FMath::Square(Data.AdjustedVisibleLODDistance[ProcessingLODIdx]);
@@ -552,10 +593,15 @@ bool TMassLODCalculator<FLODLogic>::AdjustDistancesFromCountForRuntimeData(const
 
 
 template <typename FLODLogic>
+template <bool bCalculateVisibility, bool bCalculateVisibilityPerViewer, bool bMaximizeCountPerViewer>
 bool TMassLODCalculator<FLODLogic>::AdjustDistancesFromCount()
 {
+	static_assert(!bCalculateVisibility || FLODLogic::bDoVisibilityLogic, "FLODLogic must have bDoVisibilityLogic enabled to calculate visibility.");
+	static_assert(!bCalculateVisibilityPerViewer || (FLODLogic::bDoVisibilityLogic && FLODLogic::bStoreInfoPerViewer), "FLODLogic must have bDoVisibilityLogic and bStoreInfoPerViewer enabled to calculate visibility per viewer.");
+	static_assert(!bMaximizeCountPerViewer || FLODLogic::bMaximizeCountPerViewer, "FLODLogic must have bMaximizeCountPerViewer enabled to maximize count per viewer.");
+
 	bool bDistanceAdjusted = false;
-	if (FLODLogic::bMaximizeCountPerViewer)
+	if (bMaximizeCountPerViewer)
 	{
 		for (int ViewerIdx = 0; ViewerIdx < Viewers.Num(); ++ViewerIdx)
 		{
@@ -564,21 +610,27 @@ bool TMassLODCalculator<FLODLogic>::AdjustDistancesFromCount()
 				continue;
 			}
 
-			bDistanceAdjusted |= AdjustDistancesFromCountForRuntimeData(LODMaxCountPerViewer, RuntimeDataPerViewer[ViewerIdx]);
+			bDistanceAdjusted |= AdjustDistancesFromCountForRuntimeData<bCalculateVisibilityPerViewer>(LODMaxCountPerViewer, RuntimeDataPerViewer[ViewerIdx]);
 		}
 	}
 
-	bDistanceAdjusted |= AdjustDistancesFromCountForRuntimeData(LODMaxCount, RuntimeData);
+	bDistanceAdjusted |= AdjustDistancesFromCountForRuntimeData<bCalculateVisibility>(LODMaxCount, RuntimeData);
 	return bDistanceAdjusted;
 }
 
 template <typename FLODLogic>
-template <typename TViewerInfoFragment, typename TLODFragment, typename TPerViewerInfoFragment>
+template <typename TViewerInfoFragment, typename TLODFragment, typename TPerViewerInfoFragment,
+		  bool bCalculateLocalViewers, bool bCalculateVisibility, bool bCalculateLODPerViewer, bool bCalculateVisibilityPerViewer, bool bMaximizeCountPerViewer>
 void TMassLODCalculator<FLODLogic>::AdjustLODFromCount(FMassExecutionContext& Context,
 													   TConstArrayView<TViewerInfoFragment> ViewersInfoList,
 													   TArrayView<TLODFragment> LODList,
 													   TConstArrayView<TPerViewerInfoFragment> PerViewerInfoList)
 {
+	static_assert(!bCalculateVisibility || FLODLogic::bDoVisibilityLogic, "FLODLogic must have bDoVisibilityLogic enabled to calculate visibility.");
+	static_assert(!bCalculateLODPerViewer || FLODLogic::bCalculateLODPerViewer, "FLODLogic must have bCalculateLODPerViewer enabled to calculate LOD Per viewer.");
+	static_assert(!bCalculateVisibilityPerViewer || (FLODLogic::bDoVisibilityLogic && FLODLogic::bStoreInfoPerViewer), "FLODLogic must have bDoVisibilityLogic and bStoreInfoPerViewer enabled to calculate visibility per viewer.");
+	static_assert(!bMaximizeCountPerViewer || FLODLogic::bMaximizeCountPerViewer, "FLODLogic must have bMaximizeCountPerViewer enabled to maximize count per viewer.");
+
 	const int32 NumEntities = Context.GetNumEntities();
 	// Adjust LOD for each viewer and remember the new highest
 	for (int EntityIdx = 0; EntityIdx < NumEntities; EntityIdx++)
@@ -586,7 +638,7 @@ void TMassLODCalculator<FLODLogic>::AdjustLODFromCount(FMassExecutionContext& Co
 		const TViewerInfoFragment& EntityViewersInfo = ViewersInfoList[EntityIdx];
 		TLODFragment& EntityLOD = LODList[EntityIdx];
 		EMassLOD::Type HighestViewerLOD = EMassLOD::Off;
-		if (FLODLogic::bMaximizeCountPerViewer)
+		if (bMaximizeCountPerViewer)
 		{
 			const TPerViewerInfoFragment& EntityPerViewerInfo = PerViewerInfoList[EntityIdx];
 
@@ -599,33 +651,33 @@ void TMassLODCalculator<FLODLogic>::AdjustLODFromCount(FMassExecutionContext& Co
 				}
 
 				// Check to see if we want only local viewer only
-				if (FLODLogic::bLocalViewersOnly && !Viewer.bLocal)
+				if (bCalculateLocalViewers && !Viewer.bLocal)
 				{
 					continue;
 				}
 
-				const float DistanceToViewerSq = GetDistanceToViewerSq<FLODLogic::bStoreInfoPerViewer>(EntityPerViewerInfo, ViewerIdx, FLT_MAX);
+				const float DistanceToViewerSq = GetDistanceToViewerSq<bMaximizeCountPerViewer>(EntityPerViewerInfo, ViewerIdx, FLT_MAX);
 				// Using the prev visibility here as it was already updated for this frame in the CalculateLOD method and we want the previous one
-				const bool bIsVisibleByViewer = GetPrevVisibilityPerViewer<FLODLogic::bDoVisibilityLogic && FLODLogic::bStoreInfoPerViewer>(EntityLOD, ViewerIdx, EMassVisibility::Max) == EMassVisibility::CanBeSeen;
-				EMassLOD::Type LODPerViewer = GetLODPerViewer<FLODLogic::bCalculateLODPerViewer>(EntityLOD, ViewerIdx, EntityLOD.LOD);
+				const bool bIsVisibleByViewer = GetPrevVisibilityPerViewer<bCalculateVisibilityPerViewer>(EntityLOD, ViewerIdx, EMassVisibility::Max) == EMassVisibility::CanBeSeen;
+				EMassLOD::Type LODPerViewer = GetLODPerViewer<bCalculateLODPerViewer>(EntityLOD, ViewerIdx, EntityLOD.LOD);
 
-				LODPerViewer = ComputeLODFromSettings(LODPerViewer, DistanceToViewerSq, bIsVisibleByViewer, nullptr, RuntimeDataPerViewer[ViewerIdx]);
+				LODPerViewer = ComputeLODFromSettings<bCalculateVisibilityPerViewer>(LODPerViewer, DistanceToViewerSq, bIsVisibleByViewer, nullptr, RuntimeDataPerViewer[ViewerIdx]);
 
 				if (HighestViewerLOD < LODPerViewer)
 				{
 					HighestViewerLOD = LODPerViewer;
 				}
 
-				SetLODPerViewer<FLODLogic::bCalculateLODPerViewer>(EntityLOD, ViewerIdx, LODPerViewer);
+				SetLODPerViewer<bCalculateLODPerViewer>(EntityLOD, ViewerIdx, LODPerViewer);
 			}
 		}
 
 		// Using the prev visibility here as it was already updated for this frame in the CalculateLOD method and we want the previous one
-		const bool bIsVisibleByAViewer = GetPrevVisibility<FLODLogic::bDoVisibilityLogic>(EntityLOD, EMassVisibility::Max) == EMassVisibility::CanBeSeen;
-		EMassLOD::Type NewLOD = ComputeLODFromSettings(EntityLOD.PrevLOD, EntityViewersInfo.ClosestViewerDistanceSq, bIsVisibleByAViewer, nullptr, RuntimeData);
+		const bool bIsVisibleByAViewer = GetPrevVisibility<bCalculateVisibility>(EntityLOD, EMassVisibility::Max) == EMassVisibility::CanBeSeen;
+		EMassLOD::Type NewLOD = ComputeLODFromSettings<bCalculateVisibility>(EntityLOD.PrevLOD, EntityViewersInfo.ClosestViewerDistanceSq, bIsVisibleByAViewer, nullptr, RuntimeData);
 
 		// Maybe the highest of all the viewers is now lower than the global entity LOD, make sure to update the it accordingly
-		if (FLODLogic::bMaximizeCountPerViewer && NewLOD < HighestViewerLOD)
+		if (bMaximizeCountPerViewer && NewLOD < HighestViewerLOD)
 		{
 			NewLOD = HighestViewerLOD;
 		}
@@ -641,7 +693,7 @@ void TMassLODCalculator<FLODLogic>::AdjustLODFromCount(FMassExecutionContext& Co
 				}
 				else
 				{
-					const TStaticArray<float, EMassLOD::Max>& AdjustedLODDistance = FLODLogic::bDoVisibilityLogic && bIsVisibleByAViewer ? RuntimeData.AdjustedVisibleLODDistance : RuntimeData.AdjustedBaseLODDistance;
+					const TStaticArray<float, EMassLOD::Max>& AdjustedLODDistance = bCalculateVisibility && bIsVisibleByAViewer ? RuntimeData.AdjustedVisibleLODDistance : RuntimeData.AdjustedBaseLODDistance;
 
 					// Need to clamp as the Sqrt is not precise enough and always end up with floating calculation errors
 					const float DistanceBetweenLODThresholdAndEntity = FMath::Max(FMath::Sqrt(EntityViewersInfo.ClosestViewerDistanceSq) - AdjustedLODDistance[NewLOD], 0.f);
@@ -684,13 +736,10 @@ void TMassLODCalculator<FLODLogic>::ForceOffLOD(FMassExecutionContext& Context,
 					continue;
 				}
 
-				if(FLODLogic::bCalculateLODPerViewer)
-				{
-					SetLODPerViewer<FLODLogic::bCalculateLODPerViewer>(EntityLOD, ViewerIdx, EMassLOD::Off);
-					SetPrevLODPerViewer<FLODLogic::bCalculateLODPerViewer>(EntityLOD, ViewerIdx, EMassLOD::Off);
-				}
-				SetPrevVisibilityPerViewer<FLODLogic::bDoVisibilityLogic&& FLODLogic::bStoreInfoPerViewer>(EntityLOD, ViewerIdx, EMassVisibility::Max);
-				SetVisibilityPerViewer<FLODLogic::bDoVisibilityLogic&& FLODLogic::bStoreInfoPerViewer>(EntityLOD, ViewerIdx, EMassVisibility::Max);
+				SetLODPerViewer<FLODLogic::bCalculateLODPerViewer>(EntityLOD, ViewerIdx, EMassLOD::Off);
+				SetPrevLODPerViewer<FLODLogic::bCalculateLODPerViewer>(EntityLOD, ViewerIdx, EMassLOD::Off);
+				SetPrevVisibilityPerViewer<FLODLogic::bDoVisibilityLogic && FLODLogic::bStoreInfoPerViewer>(EntityLOD, ViewerIdx, EMassVisibility::Max);
+				SetVisibilityPerViewer<FLODLogic::bDoVisibilityLogic && FLODLogic::bStoreInfoPerViewer>(EntityLOD, ViewerIdx, EMassVisibility::Max);
 			}
 		}
 
@@ -713,9 +762,10 @@ void TMassLODCalculator<FLODLogic>::DebugDisplayLOD(FMassExecutionContext& Conte
 }
 
 template <typename FLODLogic>
+template<bool bCalculateVisibility>
 EMassLOD::Type TMassLODCalculator<FLODLogic>::ComputeLODFromSettings(const EMassLOD::Type PrevLOD, const float DistanceToViewerSq, const bool bIsVisible, bool* bIsInAVisibleRange, const FMassLODRuntimeData& Data) const
 {
-	const TStaticArray<float, EMassLOD::Max>& AdjustedLODDistanceSq = FLODLogic::bDoVisibilityLogic && bIsVisible ? Data.AdjustedVisibleLODDistanceSq : Data.AdjustedBaseLODDistanceSq;
+	const TStaticArray<float, EMassLOD::Max>& AdjustedLODDistanceSq = bCalculateVisibility && bIsVisible ? Data.AdjustedVisibleLODDistanceSq : Data.AdjustedBaseLODDistanceSq;
 	int32 LODDistIdx = EMassLOD::Max - 1;
 	for (; LODDistIdx > 0; LODDistIdx--)
 	{
@@ -724,7 +774,7 @@ EMassLOD::Type TMassLODCalculator<FLODLogic>::ComputeLODFromSettings(const EMass
 			// Validate that we allow going to a single higher LOD level after considering an extended buffer hysteresis on distance for the LOD level we are about to quit to prevent oscillating LOD states
 			if (PrevLOD != EMassLOD::Max && (PrevLOD - (EMassLOD::Type)LODDistIdx) == 1)
 			{
-				const TStaticArray<float, EMassLOD::Max>& AdjustedLODDistance = FLODLogic::bDoVisibilityLogic && bIsVisible ? Data.AdjustedVisibleLODDistance : Data.AdjustedBaseLODDistance;
+				const TStaticArray<float, EMassLOD::Max>& AdjustedLODDistance = bCalculateVisibility && bIsVisible ? Data.AdjustedVisibleLODDistance : Data.AdjustedBaseLODDistance;
 				float HysteresisDistance = FMath::Lerp(AdjustedLODDistance[LODDistIdx], AdjustedLODDistance[LODDistIdx + 1], 1.f - BufferHysteresisOnDistanceRatio);
 				if (DistanceToViewerSq >= FMath::Square(HysteresisDistance))
 				{
@@ -738,7 +788,7 @@ EMassLOD::Type TMassLODCalculator<FLODLogic>::ComputeLODFromSettings(const EMass
 	}
 
 	EMassLOD::Type NewLOD = (EMassLOD::Type)LODDistIdx;
-	if(FLODLogic::bDoVisibilityLogic && bIsInAVisibleRange)
+	if(bCalculateVisibility && bIsInAVisibleRange)
 	{
 		*bIsInAVisibleRange = bIsVisible ? (NewLOD != EMassLOD::Off) : (DistanceToViewerSq < Data.AdjustedVisibleLODDistanceSq[EMassLOD::Off]);
 	}
diff --git a/Engine/Plugins/Runtime/MassGameplay/Source/MassLOD/Public/MassLODCollector.h b/Engine/Plugins/Runtime/MassGameplay/Source/MassLOD/Public/MassLODCollector.h
--- a/Engine/Plugins/Runtime/MassGameplay/Source/MassLOD/Public/MassLODCollector.h
+++ b/Engine/Plugins/Runtime/MassGameplay/Source/MassLOD/Public/MassLODCollector.h
@@ -29,12 +29,20 @@ struct TMassLODCollector : public FMassLODBaseLogic
 	 * @Param ViewersInfoList is the fragment where to store source information for LOD calculation
 	 */
 	template <typename TTransformFragment, 
-			  typename TViewerInfoFragment>
+			  typename TViewerInfoFragment,
+			  bool bCollectLocalViewers = FLODLogic::bLocalViewersOnly,
+			  bool bCollectDistanceToFrustum = FLODLogic::bDoVisibilityLogic>
 	FORCEINLINE void CollectLODInfo(FMassExecutionContext& Context, 
 									TConstArrayView<TTransformFragment> TransformList, 
 									TArrayView<TViewerInfoFragment> ViewersInfoList)
 	{
-		CollectLODInfo(Context, TransformList, ViewersInfoList, TArrayView<void*>());
+		CollectLODInfo<TTransformFragment,
+			TViewerInfoFragment,
+			/*TPerViewerInfoFragment*/void*,
+			bCollectLocalViewers,
+			bCollectDistanceToFrustum,
+			/*bCollectDistancePerViewer*/false,
+			/*bCollectDistanceToFrustumPerViewer*/false>(Context, TransformList, ViewersInfoList, TArrayView<void*>());
 	}
 
 
@@ -49,7 +57,11 @@ struct TMassLODCollector : public FMassLODBaseLogic
 	 */
 	template <typename TTransformFragment,
 			  typename TViewerInfoFragment,
-			  typename TPerViewerInfoFragment>
+			  typename TPerViewerInfoFragment,
+			  bool bCollectLocalViewers = FLODLogic::bLocalViewersOnly,
+			  bool bCollectDistanceToFrustum = FLODLogic::bDoVisibilityLogic,
+			  bool bCollectDistancePerViewer = FLODLogic::bStoreInfoPerViewer,
+			  bool bCollectDistanceToFrustumPerViewer = FLODLogic::bDoVisibilityLogic && FLODLogic::bStoreInfoPerViewer>
 	void CollectLODInfo(FMassExecutionContext& Context, 
 						TConstArrayView<TTransformFragment> TransformList, 
 						TArrayView<TViewerInfoFragment> ViewersInfoList, 
@@ -65,7 +77,11 @@ void TMassLODCollector<FLODLogic>::PrepareExecution(TConstArrayView<FViewerInfo>
 template <typename FLODLogic>
 template <typename TTransformFragment, 
 		  typename TViewerInfoFragment, 
-		  typename TPerViewerInfoFragment>
+		  typename TPerViewerInfoFragment,
+		  bool bCollectLocalViewers,
+		  bool bCollectDistanceToFrustum,
+		  bool bCollectDistancePerViewer,
+		  bool bCollectDistanceToFrustumPerViewer>
 void TMassLODCollector<FLODLogic>::CollectLODInfo(FMassExecutionContext& Context, 
 												  TConstArrayView<TTransformFragment> TransformList, 
 												  TArrayView<TViewerInfoFragment> ViewersInfoList, 
@@ -81,19 +97,19 @@ void TMassLODCollector<FLODLogic>::CollectLODInfo(FMassExecutionContext& Context
 		TViewerInfoFragment& EntityViewerInfo = ViewersInfoList[EntityIdx];
 		TPerViewerInfoFragment& EntityInfoPerViewer = FLODLogic::bStoreInfoPerViewer ? PerViewerInfoList[EntityIdx] : DummyFragment;
 
-		SetDistanceToViewerSqNum<FLODLogic::bStoreInfoPerViewer>(EntityInfoPerViewer, Viewers.Num());
-		SetDistanceToFrustumNum<FLODLogic::bDoVisibilityLogic && FLODLogic::bStoreInfoPerViewer>(EntityInfoPerViewer, Viewers.Num());
+		SetDistanceToViewerSqNum<bCollectDistancePerViewer>(EntityInfoPerViewer, Viewers.Num());
+		SetDistanceToFrustumNum<bCollectDistanceToFrustumPerViewer>(EntityInfoPerViewer, Viewers.Num());
 		for (int ViewerIdx = 0; ViewerIdx < Viewers.Num(); ++ViewerIdx)
 		{
 			const FViewerLODInfo& Viewer = Viewers[ViewerIdx];
 			if (Viewer.bClearData)
 			{
-				SetDistanceToViewerSq<FLODLogic::bStoreInfoPerViewer>(EntityInfoPerViewer, ViewerIdx, FLT_MAX);
-				SetDistanceToFrustum<FLODLogic::bDoVisibilityLogic && FLODLogic::bStoreInfoPerViewer>(EntityInfoPerViewer, ViewerIdx, FLT_MAX);
+				SetDistanceToViewerSq<bCollectDistancePerViewer>(EntityInfoPerViewer, ViewerIdx, FLT_MAX);
+				SetDistanceToFrustum<bCollectDistanceToFrustumPerViewer>(EntityInfoPerViewer, ViewerIdx, FLT_MAX);
 			}
 
 			// Check to see if we want only local viewer only
-			if (FLODLogic::bLocalViewersOnly && !Viewer.bLocal)
+			if (bCollectLocalViewers && !Viewer.bLocal)
 			{
 				continue;
 			}
@@ -107,12 +123,12 @@ void TMassLODCollector<FLODLogic>::CollectLODInfo(FMassExecutionContext& Context
 				{
 					ClosestViewerDistanceSq = DistanceToViewerSq;
 				}
-				SetDistanceToViewerSq<FLODLogic::bStoreInfoPerViewer>(EntityInfoPerViewer, ViewerIdx, DistanceToViewerSq);
+				SetDistanceToViewerSq<bCollectDistancePerViewer>(EntityInfoPerViewer, ViewerIdx, DistanceToViewerSq);
 
-				if (FLODLogic::bDoVisibilityLogic)
+				if (bCollectDistanceToFrustum)
 				{
 					const float DistanceToFrustum = Viewer.Frustum.DistanceTo(EntityLocation);
-					SetDistanceToFrustum<FLODLogic::bDoVisibilityLogic && FLODLogic::bStoreInfoPerViewer>(EntityInfoPerViewer, ViewerIdx, DistanceToFrustum);
+					SetDistanceToFrustum<bCollectDistanceToFrustumPerViewer>(EntityInfoPerViewer, ViewerIdx, DistanceToFrustum);
 					if (ClosestDistanceToFrustum > DistanceToFrustum)
 					{
 						ClosestDistanceToFrustum = DistanceToFrustum;
@@ -121,6 +137,6 @@ void TMassLODCollector<FLODLogic>::CollectLODInfo(FMassExecutionContext& Context
 			}
 		}
 		EntityViewerInfo.ClosestViewerDistanceSq = ClosestViewerDistanceSq;
-		SetClosestDistanceToFrustum<FLODLogic::bDoVisibilityLogic>(EntityViewerInfo, ClosestDistanceToFrustum);
+		SetClosestDistanceToFrustum<bCollectDistanceToFrustum>(EntityViewerInfo, ClosestDistanceToFrustum);
 	}
 }
\ No newline at end of file
