LODCollectorTrait

#rnx

diff --git a/Engine/Plugins/Runtime/MassGameplay/Source/MassLOD/Private/MassLODCollectorProcessor.cpp b/Engine/Plugins/Runtime/MassGameplay/Source/MassLOD/Private/MassLODCollectorProcessor.cpp
--- a/Engine/Plugins/Runtime/MassGameplay/Source/MassLOD/Private/MassLODCollectorProcessor.cpp
+++ b/Engine/Plugins/Runtime/MassGameplay/Source/MassLOD/Private/MassLODCollectorProcessor.cpp
@@ -16,76 +16,57 @@ UMassLODCollectorProcessor::UMassLODCollectorProcessor()
 
 void UMassLODCollectorProcessor::ConfigureQueries()
 {
-	for (FMassCollectorLODConfig& LODConfig : LODConfigs)
+	FMassEntityQuery BaseQuery;
+	BaseQuery.AddTagRequirement<FMassCollectLODViewerInfoTag>(EMassFragmentPresence::All);
+	BaseQuery.AddRequirement<FDataFragment_Transform>(EMassFragmentAccess::ReadOnly);
+	BaseQuery.AddRequirement<FMassViewerInfoFragment>(EMassFragmentAccess::ReadWrite);
+	BaseQuery.AddChunkRequirement<FMassSimulationVariableTickChunkFragment>(EMassFragmentAccess::ReadOnly, EMassFragmentPresence::Optional);
+	BaseQuery.AddChunkRequirement<FMassVisualizationChunkFragment>(EMassFragmentAccess::ReadOnly, EMassFragmentPresence::Optional);
+	BaseQuery.SetChunkFilter([](const FMassExecutionContext& Context)
 	{
-		FMassEntityQuery BaseQuery;
-		if (LODConfig.TagFilter.GetScriptStruct())
-		{
-			BaseQuery.AddTagRequirement(*LODConfig.TagFilter.GetScriptStruct(), EMassFragmentPresence::All);
-		}
-		BaseQuery.AddRequirement<FDataFragment_Transform>(EMassFragmentAccess::ReadOnly);
-		BaseQuery.AddRequirement<FMassViewerInfoFragment>(EMassFragmentAccess::ReadWrite);
+		return FMassVisualizationChunkFragment::ShouldUpdateVisualizationForChunk(Context)
+			|| FMassSimulationVariableTickChunkFragment::ShouldTickChunkThisFrame(Context);
+	});
 
-		LODConfig.OnLODEntityQuery = BaseQuery;
-		LODConfig.OnLODEntityQuery.AddTagRequirement<FMassOffLODTag>(EMassFragmentPresence::None);
-		LODConfig.OffLODEntityQuery_Conditional = BaseQuery;
-		LODConfig.OffLODEntityQuery_Conditional.AddTagRequirement<FMassOffLODTag>(EMassFragmentPresence::All);
-		LODConfig.OffLODEntityQuery_Conditional.AddChunkRequirement<FMassSimulationVariableTickChunkFragment>(EMassFragmentAccess::ReadOnly, EMassFragmentPresence::Optional);
-		LODConfig.OffLODEntityQuery_Conditional.SetChunkFilter(&FMassSimulationVariableTickChunkFragment::ShouldTickChunkThisFrame);
+	EntityQuery_VisibleRangeAndOnLOD = BaseQuery;
+	EntityQuery_VisibleRangeAndOnLOD.AddTagRequirement<FMassVisibilityCulledByDistanceTag>(EMassFragmentPresence::None);
+	EntityQuery_VisibleRangeAndOnLOD.AddTagRequirement<FMassOffLODTag>(EMassFragmentPresence::None);
 
-		LODConfig.CloseEntityQuery = BaseQuery;
-		LODConfig.CloseEntityQuery.AddTagRequirement<FMassVisibilityCulledByDistanceTag>(EMassFragmentPresence::None);
-		LODConfig.FarEntityQuery_Conditional = BaseQuery;
-		LODConfig.FarEntityQuery_Conditional.AddTagRequirement<FMassVisibilityCulledByDistanceTag>(EMassFragmentPresence::All);
-		LODConfig.FarEntityQuery_Conditional.AddChunkRequirement<FMassVisualizationChunkFragment>(EMassFragmentAccess::ReadOnly);
-		LODConfig.FarEntityQuery_Conditional.SetChunkFilter(&FMassVisualizationChunkFragment::ShouldUpdateVisualizationForChunk);
+	EntityQuery_VisibleRangeOnly = BaseQuery;
+	EntityQuery_VisibleRangeOnly.AddTagRequirement<FMassVisibilityCulledByDistanceTag>(EMassFragmentPresence::None);
+	EntityQuery_VisibleRangeOnly.AddTagRequirement<FMassOffLODTag>(EMassFragmentPresence::All);
 
-		// @todo: Cleanup once we can do more complex logic with entity queries. There is no other way to do "Notall", we need to create by hand 3 queries to achieve it.
-		LODConfig.CloseAndOnLODEntityQuery = BaseQuery;
-		LODConfig.CloseAndOnLODEntityQuery.AddTagRequirement<FMassVisibilityCulledByDistanceTag>(EMassFragmentPresence::None);
-		LODConfig.CloseAndOnLODEntityQuery.AddTagRequirement<FMassOffLODTag>(EMassFragmentPresence::None);
-		LODConfig.CloseAndOffLODEntityQuery = BaseQuery;
-		LODConfig.CloseAndOffLODEntityQuery.AddTagRequirement<FMassVisibilityCulledByDistanceTag>(EMassFragmentPresence::None);
-		LODConfig.CloseAndOffLODEntityQuery.AddTagRequirement<FMassOffLODTag>(EMassFragmentPresence::All);
-		LODConfig.FarAndOnLODEntityQuery = BaseQuery;
-		LODConfig.FarAndOnLODEntityQuery.AddTagRequirement<FMassVisibilityCulledByDistanceTag>(EMassFragmentPresence::All);
-		LODConfig.FarAndOnLODEntityQuery.AddTagRequirement<FMassOffLODTag>(EMassFragmentPresence::None);
-		LODConfig.FarAndOffLODEntityQuery_Conditional = BaseQuery;
-		LODConfig.FarAndOffLODEntityQuery_Conditional.AddTagRequirement<FMassVisibilityCulledByDistanceTag>(EMassFragmentPresence::All);
-		LODConfig.FarAndOffLODEntityQuery_Conditional.AddTagRequirement<FMassOffLODTag>(EMassFragmentPresence::All);
-		LODConfig.FarAndOffLODEntityQuery_Conditional.AddChunkRequirement<FMassSimulationVariableTickChunkFragment>(EMassFragmentAccess::ReadOnly, EMassFragmentPresence::Optional);
-		LODConfig.FarAndOffLODEntityQuery_Conditional.AddChunkRequirement<FMassVisualizationChunkFragment>(EMassFragmentAccess::ReadOnly);
-		LODConfig.FarAndOffLODEntityQuery_Conditional.SetChunkFilter([](const FMassExecutionContext& Context)
-			{
-				return Context.GetChunkFragment<FMassVisualizationChunkFragment>().ShouldUpdateVisualization()
-					|| FMassSimulationVariableTickChunkFragment::ShouldTickChunkThisFrame(Context);
-			});
-	}
+	EntityQuery_OnLODOnly = BaseQuery;
+	EntityQuery_OnLODOnly.AddTagRequirement<FMassVisibilityCulledByDistanceTag>(EMassFragmentPresence::All);
+	EntityQuery_OnLODOnly.AddTagRequirement<FMassOffLODTag>(EMassFragmentPresence::None);
+
+	EntityQuery_NotVisibleRangeAndOffLOD = BaseQuery;
+	EntityQuery_NotVisibleRangeAndOffLOD.AddTagRequirement<FMassVisibilityCulledByDistanceTag>(EMassFragmentPresence::All);
+	EntityQuery_NotVisibleRangeAndOffLOD.AddTagRequirement<FMassOffLODTag>(EMassFragmentPresence::All);
 }
 
-template <typename TCollector>
-void CollectLOD(UMassEntitySubsystem& EntitySubsystem, FMassExecutionContext& Context, const TArray<FViewerInfo>& ViewersInfo, TCollector& Collector, TArrayView<FMassEntityQuery> HighFrequencyTickingEntityQueries, FMassEntityQuery& LowFrequencyTickingEntityQuery)
+template <bool bLocalViewersOnly>
+void UMassLODCollectorProcessor::CollectLODForChunk(FMassExecutionContext& Context)
 {
-	Collector.PrepareExecution(ViewersInfo);
-
-	auto CollectLODInfo = [&Collector](FMassExecutionContext& Context)
-	{
-		TConstArrayView<FDataFragment_Transform> LocationList = Context.GetFragmentView<FDataFragment_Transform>();
-		TArrayView<FMassViewerInfoFragment> ViewerInfoList = Context.GetMutableFragmentView<FMassViewerInfoFragment>();
+	TConstArrayView<FDataFragment_Transform> LocationList = Context.GetFragmentView<FDataFragment_Transform>();
+	TArrayView<FMassViewerInfoFragment> ViewerInfoList = Context.GetMutableFragmentView<FMassViewerInfoFragment>();
 
-		Collector.CollectLODInfo(Context, LocationList, ViewerInfoList);
-	};
+	Collector.CollectLODInfo<FDataFragment_Transform, FMassViewerInfoFragment, bLocalViewersOnly, true/*bCollectDistanceToViewer*/>(Context, LocationList, ViewerInfoList);
+}
 
+template <bool bLocalViewersOnly>
+void UMassLODCollectorProcessor::ExecuteInternal(UMassEntitySubsystem& EntitySubsystem, FMassExecutionContext& Context)
+{
 	{
-		TRACE_CPUPROFILER_EVENT_SCOPE(TEXT("HighFrequency"))
-		for(FMassEntityQuery& HighFrequencyTickingEntityQuery : HighFrequencyTickingEntityQueries)
-		{
-			HighFrequencyTickingEntityQuery.ForEachEntityChunk(EntitySubsystem, Context, CollectLODInfo);
-		}
+		TRACE_CPUPROFILER_EVENT_SCOPE(TEXT("Close"));
+		EntityQuery_VisibleRangeAndOnLOD.ForEachEntityChunk(EntitySubsystem, Context, [this](FMassExecutionContext& Context) { CollectLODForChunk<bLocalViewersOnly>(Context); });
+		EntityQuery_VisibleRangeOnly.ForEachEntityChunk(EntitySubsystem, Context, [this](FMassExecutionContext& Context) { CollectLODForChunk<bLocalViewersOnly>(Context); });
+		EntityQuery_OnLODOnly.ForEachEntityChunk(EntitySubsystem, Context, [this](FMassExecutionContext& Context) { CollectLODForChunk<bLocalViewersOnly>(Context); });
 	}
+
 	{
-		TRACE_CPUPROFILER_EVENT_SCOPE(TEXT("LowFrequency"))
-		LowFrequencyTickingEntityQuery.ForEachEntityChunk(EntitySubsystem, Context, CollectLODInfo);
+		TRACE_CPUPROFILER_EVENT_SCOPE(TEXT("Far"));
+		EntityQuery_NotVisibleRangeAndOffLOD.ForEachEntityChunk(EntitySubsystem, Context, [this](FMassExecutionContext& Context) { CollectLODForChunk<bLocalViewersOnly>(Context); });
 	}
 }
 
@@ -93,25 +74,16 @@ void UMassLODCollectorProcessor::Execute(UMassEntitySubsystem& EntitySubsystem,
 {
 	check(LODManager);
 	const TArray<FViewerInfo>& Viewers = LODManager->GetViewers();
+	Collector.PrepareExecution(Viewers);
 
 	check(World);
-
-	for (FMassCollectorLODConfig& LODConfig : LODConfigs)
+	if (World->IsNetMode(NM_DedicatedServer))
 	{
-		if (World->IsNetMode(NM_Client))
-		{
-			TRACE_CPUPROFILER_EVENT_SCOPE(TEXT("RepresentationLODCollector"))
-			CollectLOD(EntitySubsystem, Context, Viewers, LODConfig.RepresentationLODCollector, MakeArrayView(&LODConfig.CloseEntityQuery,1), LODConfig.FarEntityQuery_Conditional);
-		}
-		else if (World->IsNetMode(NM_DedicatedServer))
-		{
-			TRACE_CPUPROFILER_EVENT_SCOPE(TEXT("SimulationLODCollector"))
-			CollectLOD(EntitySubsystem, Context, Viewers, LODConfig.SimulationLODCollector, MakeArrayView(&LODConfig.OnLODEntityQuery,1), LODConfig.OffLODEntityQuery_Conditional);
-		}
-		else
-		{
-			TRACE_CPUPROFILER_EVENT_SCOPE(TEXT("CombinedLODCollector"))
-			CollectLOD(EntitySubsystem, Context, Viewers, LODConfig.CombinedLODCollector, MakeArrayView(&LODConfig.CloseAndOnLODEntityQuery,3), LODConfig.FarAndOffLODEntityQuery_Conditional);
-		}
+		ExecuteInternal<false/*bLocalViewersOnly*/>(EntitySubsystem, Context);
 	}
+	else
+	{
+		ExecuteInternal<true/*bLocalViewersOnly*/>(EntitySubsystem, Context);
+	}
+
 }
diff --git a/Engine/Plugins/Runtime/MassGameplay/Source/MassLOD/Private/MassLODTrait.cpp b/Engine/Plugins/Runtime/MassGameplay/Source/MassLOD/Private/MassLODTrait.cpp
--- a/Engine/Plugins/Runtime/MassGameplay/Source/MassLOD/Private/MassLODTrait.cpp
+++ b/Engine/Plugins/Runtime/MassGameplay/Source/MassLOD/Private/MassLODTrait.cpp
@@ -5,11 +5,19 @@
 #include "Engine/World.h"
 #include "MassCommonFragments.h"
 #include "StructUtilsTypes.h"
+#include "MassLODTypes.h"
+
+void UMassLODCollectorTrait::BuildTemplate(FMassEntityTemplateBuildContext& BuildContext, UWorld& World) const
+{
+	BuildContext.AddFragment<FMassViewerInfoFragment>();
+	BuildContext.AddTag<FMassCollectLODViewerInfoTag>();
+}
 
 void UMassSimulationLODTrait::BuildTemplate(FMassEntityTemplateBuildContext& BuildContext, UWorld& World) const
 {
-	BuildContext.AddFragmentWithDefaultInitializer<FDataFragment_Transform>();
-	BuildContext.AddFragmentWithDefaultInitializer<FMassViewerInfoFragment>();
+	BuildContext.AddTag<FMassCollectLODViewerInfoTag>(); // Depends on FMassViewerInfoFragment
+	BuildContext.AddFragment<FDataFragment_Transform>();
+
 	FMassSimulationLODFragment& LODFragment = BuildContext.AddFragment_GetRef<FMassSimulationLODFragment>();
 
 	// Start all simulation LOD in the Off 
@@ -39,4 +47,4 @@ void UMassSimulationLODTrait::BuildTemplate(FMassEntityTemplateBuildContext& Bui
 		FSharedStruct VariableTickSharedFragment = EntitySubsystem->GetOrCreateSharedFragment<FMassSimulationVariableTickSharedFragment>(VariableTickConfigHash, VariableTickConfig);
 		BuildContext.AddSharedFragment(VariableTickSharedFragment);
 	}
-}
+}
\ No newline at end of file
diff --git a/Engine/Plugins/Runtime/MassGameplay/Source/MassLOD/Public/MassLODCollectorProcessor.h b/Engine/Plugins/Runtime/MassGameplay/Source/MassLOD/Public/MassLODCollectorProcessor.h
--- a/Engine/Plugins/Runtime/MassGameplay/Source/MassLOD/Public/MassLODCollectorProcessor.h
+++ b/Engine/Plugins/Runtime/MassGameplay/Source/MassLOD/Public/MassLODCollectorProcessor.h
@@ -9,28 +9,12 @@
 
 #include "MassLODCollectorProcessor.generated.h"
 
-USTRUCT()
-struct FMassCollectorLODConfig
+struct FMassGenericCollectorLogic : public FLODDefaultLogic
 {
-	GENERATED_BODY()
-
-	/** Component Tag that will be used to associate LOD config */
-	UPROPERTY(EditAnywhere, Category = "Mass|LOD", config, meta = (BaseStruct = "MassTag"))
-	FInstancedStruct TagFilter;
-
-	/** Runtime data for matching the LOD config */
-	TMassLODCollector<FMassRepresentationLODLogic> RepresentationLODCollector;
-	FMassEntityQuery CloseEntityQuery;
-	FMassEntityQuery FarEntityQuery_Conditional;
-	TMassLODCollector<FMassSimulationLODLogic> SimulationLODCollector;
-	FMassEntityQuery OnLODEntityQuery;
-	FMassEntityQuery OffLODEntityQuery_Conditional;
-	TMassLODCollector<FMassCombinedLODLogic> CombinedLODCollector;
-	// Keep these in this order as we are creating ArrayView from the first 3 queries.
-	FMassEntityQuery CloseAndOnLODEntityQuery;
-	FMassEntityQuery CloseAndOffLODEntityQuery;
-	FMassEntityQuery FarAndOnLODEntityQuery;
-	FMassEntityQuery FarAndOffLODEntityQuery_Conditional;
+	enum
+	{
+		bDoVisibilityLogic = true,
+	};
 };
 
 /*
@@ -48,6 +32,21 @@ protected:
 	virtual void ConfigureQueries() override;
 	virtual void Execute(UMassEntitySubsystem& EntitySubsystem, FMassExecutionContext& Context) override;
 
-	UPROPERTY(EditAnywhere, Category = "Mass|LOD", config)
-	TArray<FMassCollectorLODConfig> LODConfigs;
+	template <bool bLocalViewersOnly>
+	void CollectLODForChunk(FMassExecutionContext& Context);
+
+	template <bool bLocalViewersOnly>
+	void ExecuteInternal(UMassEntitySubsystem& EntitySubsystem, FMassExecutionContext& Context);
+
+	TMassLODCollector<FMassGenericCollectorLogic> Collector;
+
+	// Queries for visualization and simulation calculation
+	/** All entities that are in visible range and are On LOD*/
+	FMassEntityQuery EntityQuery_VisibleRangeAndOnLOD;
+	/** All entities that are in visible range but are Off LOD */
+	FMassEntityQuery EntityQuery_VisibleRangeOnly;
+	/** All entities that are NOT in visible range but are On LOD */
+	FMassEntityQuery EntityQuery_OnLODOnly;
+	/** All entities that are Not in visible range and are at Off LOD */
+	FMassEntityQuery EntityQuery_NotVisibleRangeAndOffLOD;
 };
diff --git a/Engine/Plugins/Runtime/MassGameplay/Source/MassLOD/Public/MassLODTrait.h b/Engine/Plugins/Runtime/MassGameplay/Source/MassLOD/Public/MassLODTrait.h
--- a/Engine/Plugins/Runtime/MassGameplay/Source/MassLOD/Public/MassLODTrait.h
+++ b/Engine/Plugins/Runtime/MassGameplay/Source/MassLOD/Public/MassLODTrait.h
@@ -7,6 +7,15 @@
 
 #include "MassLODTrait.generated.h"
 
+UCLASS(meta = (DisplayName = "LODCollector"))
+class MASSLOD_API UMassLODCollectorTrait : public UMassEntityTraitBase
+{
+	GENERATED_BODY()
+
+protected:
+	virtual void BuildTemplate(FMassEntityTemplateBuildContext& BuildContext, UWorld& World) const override;
+};
+
 UCLASS(meta = (DisplayName = "SimulationLOD"))
 class MASSLOD_API UMassSimulationLODTrait : public UMassEntityTraitBase
 {
diff --git a/Engine/Plugins/Runtime/MassGameplay/Source/MassLOD/Public/MassLODTypes.h b/Engine/Plugins/Runtime/MassGameplay/Source/MassLOD/Public/MassLODTypes.h
--- a/Engine/Plugins/Runtime/MassGameplay/Source/MassLOD/Public/MassLODTypes.h
+++ b/Engine/Plugins/Runtime/MassGameplay/Source/MassLOD/Public/MassLODTypes.h
@@ -142,6 +142,12 @@ enum class EMassVisibility : uint8
 	Max
 };
 
+USTRUCT()
+struct FMassCollectLODViewerInfoTag : public FMassTag
+{
+	GENERATED_BODY();
+};
+
 USTRUCT()
 struct MASSLOD_API FMassVisibilityCanBeSeenTag : public FMassTag
 {
diff --git a/Engine/Plugins/Runtime/MassGameplay/Source/MassRepresentation/Private/MassVisualizationTrait.cpp b/Engine/Plugins/Runtime/MassGameplay/Source/MassRepresentation/Private/MassVisualizationTrait.cpp
--- a/Engine/Plugins/Runtime/MassGameplay/Source/MassRepresentation/Private/MassVisualizationTrait.cpp
+++ b/Engine/Plugins/Runtime/MassGameplay/Source/MassRepresentation/Private/MassVisualizationTrait.cpp
@@ -27,6 +27,7 @@ void UMassVisualizationTrait::BuildTemplate(FMassEntityTemplateBuildContext& Bui
 	// the following needs to be always there for mesh vis to work. Adding following fragments after already 
 	// adding Config.AdditionalDataFragments to let user configure the fragments first. Calling BuildContext.Add() 
 	// won't override any fragments that are already there
+	BuildContext.AddTag<FMassCollectLODViewerInfoTag>(); // Depends on FMassViewerInfoFragment
 	BuildContext.AddFragment<FDataFragment_Transform>();
 
 	if (UMassRepresentationSubsystem* RepresentationSubsystem = Cast<UMassRepresentationSubsystem>(World.GetSubsystemBase(RepresentationSubsystemClass)))
@@ -43,7 +44,6 @@ void UMassVisualizationTrait::BuildTemplate(FMassEntityTemplateBuildContext& Bui
 			BuildContext.AddTag(*RepresentationDestructorTag);
 		}
 
-		BuildContext.AddFragment<FMassViewerInfoFragment>();
 		BuildContext.AddFragment<FMassRepresentationLODFragment>();
 		BuildContext.AddTag<FMassVisibilityCulledByDistanceTag>();
 
