Renamed MassFragmentDestructor to MassFragmentDeinitializer and cleaned up related headers a bit

diff --git a/Engine/Plugins/AI/MassAI/Source/MassAIBehavior/Public/MassStateTreeProcessors.h b/Engine/Plugins/AI/MassAI/Source/MassAIBehavior/Public/MassStateTreeProcessors.h
--- a/Engine/Plugins/AI/MassAI/Source/MassAIBehavior/Public/MassStateTreeProcessors.h
+++ b/Engine/Plugins/AI/MassAI/Source/MassAIBehavior/Public/MassStateTreeProcessors.h
@@ -4,7 +4,7 @@
 
 #include "MassSignals/Public/MassSignalProcessorBase.h"
 #include "MassStateTreeFragments.h"
-#include "MassTranslator.h"
+#include "MassObserverProcessor.h"
 #include "MassLODTypes.h"
 #include "MassStateTreeProcessors.generated.h"
 
@@ -14,7 +14,7 @@ struct FMassStateTreeExecutionContext;
  * Processor to stop and uninitialize StateTrees on entities.
  */
 UCLASS()
-class MASSAIBEHAVIOR_API UMassStateTreeFragmentDestructor : public UMassFragmentDestructor
+class MASSAIBEHAVIOR_API UMassStateTreeFragmentDestructor : public UMassFragmentDeinitializer
 {
 	GENERATED_BODY()
 
diff --git a/Engine/Plugins/AI/MassAI/Source/MassAIBehavior/Public/MassZoneGraphAnnotationProcessors.h b/Engine/Plugins/AI/MassAI/Source/MassAIBehavior/Public/MassZoneGraphAnnotationProcessors.h
--- a/Engine/Plugins/AI/MassAI/Source/MassAIBehavior/Public/MassZoneGraphAnnotationProcessors.h
+++ b/Engine/Plugins/AI/MassAI/Source/MassAIBehavior/Public/MassZoneGraphAnnotationProcessors.h
@@ -3,7 +3,7 @@
 #pragma once
 
 #include "MassSignalProcessorBase.h"
-#include "MassTranslator.h"
+#include "MassObserverProcessor.h"
 #include "MassZoneGraphAnnotationProcessors.generated.h"
 
 class UMassSignalSubsystem;
diff --git a/Engine/Plugins/AI/MassAI/Source/MassAIMovement/Public/MassMoveTargetFragmentInitializer.h b/Engine/Plugins/AI/MassAI/Source/MassAIMovement/Public/MassMoveTargetFragmentInitializer.h
--- a/Engine/Plugins/AI/MassAI/Source/MassAIMovement/Public/MassMoveTargetFragmentInitializer.h
+++ b/Engine/Plugins/AI/MassAI/Source/MassAIMovement/Public/MassMoveTargetFragmentInitializer.h
@@ -2,7 +2,7 @@
 
 #pragma once 
 
-#include "MassTranslator.h"
+#include "MassObserverProcessor.h"
 #include "MassMoveTargetFragmentInitializer.generated.h"
 
 
diff --git a/Engine/Plugins/AI/MassAI/Source/MassAIMovement/Public/MassZoneGraphMovementProcessors.h b/Engine/Plugins/AI/MassAI/Source/MassAIMovement/Public/MassZoneGraphMovementProcessors.h
--- a/Engine/Plugins/AI/MassAI/Source/MassAIMovement/Public/MassZoneGraphMovementProcessors.h
+++ b/Engine/Plugins/AI/MassAI/Source/MassAIMovement/Public/MassZoneGraphMovementProcessors.h
@@ -3,7 +3,7 @@
 #pragma once
 
 #include "MassProcessor.h"
-#include "MassTranslator.h"
+#include "MassObserverProcessor.h"
 #include "MassZoneGraphMovementProcessors.generated.h"
 
 class UMassSignalSubsystem;
diff --git a/Engine/Plugins/AI/MassAI/Source/MassAIMovement/Public/MovementProcessor.h b/Engine/Plugins/AI/MassAI/Source/MassAIMovement/Public/MovementProcessor.h
--- a/Engine/Plugins/AI/MassAI/Source/MassAIMovement/Public/MovementProcessor.h
+++ b/Engine/Plugins/AI/MassAI/Source/MassAIMovement/Public/MovementProcessor.h
@@ -3,7 +3,7 @@
 #pragma once
 
 #include "MassProcessor.h"
-#include "MassTranslator.h"
+#include "MassObserverProcessor.h"
 #include "MassMovementSubsystem.h"
 #include "MovementProcessor.generated.h"
 
@@ -37,7 +37,7 @@ protected:
 
 /** Destructor processor to remove avoidance obstacles from the avoidance obstacle grid */
 UCLASS()
-class MASSAIMOVEMENT_API UMassAvoidanceObstacleRemoverFragmentDestructor : public UMassFragmentDestructor
+class MASSAIMOVEMENT_API UMassAvoidanceObstacleRemoverFragmentDestructor : public UMassFragmentDeinitializer
 {
 	GENERATED_BODY()
 
diff --git a/Engine/Plugins/AI/MassCrowd/Source/MassCrowd/Public/MassCrowdNavigationProcessor.h b/Engine/Plugins/AI/MassCrowd/Source/MassCrowd/Public/MassCrowdNavigationProcessor.h
--- a/Engine/Plugins/AI/MassCrowd/Source/MassCrowd/Public/MassCrowdNavigationProcessor.h
+++ b/Engine/Plugins/AI/MassCrowd/Source/MassCrowd/Public/MassCrowdNavigationProcessor.h
@@ -3,7 +3,7 @@
 #pragma once
 
 #include "MassSignalProcessorBase.h"
-#include "MassTranslator.h"
+#include "MassObserverProcessor.h"
 #include "NavigationProcessor.h"
 #include "MassCrowdNavigationProcessor.generated.h"
 
@@ -31,7 +31,7 @@ protected:
 
 /** Processors that cleans up the lane tracking on entity destruction. */
 UCLASS()
-class MASSCROWD_API UMassCrowdLaneTrackingDestructor : public UMassFragmentDestructor
+class MASSCROWD_API UMassCrowdLaneTrackingDestructor : public UMassFragmentDeinitializer
 {
 	GENERATED_BODY()
 
diff --git a/Engine/Plugins/Runtime/MassEntity/Source/MassEntity/Private/MassObserverProcessor.cpp b/Engine/Plugins/Runtime/MassEntity/Source/MassEntity/Private/MassObserverProcessor.cpp
--- a/Engine/Plugins/Runtime/MassEntity/Source/MassEntity/Private/MassObserverProcessor.cpp
+++ b/Engine/Plugins/Runtime/MassEntity/Source/MassEntity/Private/MassObserverProcessor.cpp
@@ -34,9 +34,9 @@ void UMassFragmentInitializer::Register()
 }
 
 //----------------------------------------------------------------------//
-//  UMassFragmentDestructor
+//  UMassFragmentDeinitializer
 //----------------------------------------------------------------------//
-void UMassFragmentDestructor::Register()
+void UMassFragmentDeinitializer::Register()
 {
 	check(FragmentType);
 	UMassObserverRegistry::GetMutable().RegisterFragmentDeinitializer(*FragmentType, GetClass());
diff --git a/Engine/Plugins/Runtime/MassEntity/Source/MassEntity/Public/MassObserverProcessor.h b/Engine/Plugins/Runtime/MassEntity/Source/MassEntity/Public/MassObserverProcessor.h
--- a/Engine/Plugins/Runtime/MassEntity/Source/MassEntity/Public/MassObserverProcessor.h
+++ b/Engine/Plugins/Runtime/MassEntity/Source/MassEntity/Public/MassObserverProcessor.h
@@ -41,7 +41,7 @@ protected:
 };
 
 UCLASS(abstract)
-class MASSENTITY_API UMassFragmentDestructor : public UMassObserverProcessor
+class MASSENTITY_API UMassFragmentDeinitializer : public UMassObserverProcessor
 {
 	GENERATED_BODY()
 
diff --git a/Engine/Plugins/Runtime/MassGameplay/Source/MassActors/Public/Translators/MassCharacterMovementTranslators.h b/Engine/Plugins/Runtime/MassGameplay/Source/MassActors/Public/Translators/MassCharacterMovementTranslators.h
--- a/Engine/Plugins/Runtime/MassGameplay/Source/MassActors/Public/Translators/MassCharacterMovementTranslators.h
+++ b/Engine/Plugins/Runtime/MassGameplay/Source/MassActors/Public/Translators/MassCharacterMovementTranslators.h
@@ -4,6 +4,7 @@
 
 #include "MassCommonFragments.h"
 #include "MassTranslator.h"
+#include "MassObserverProcessor.h"
 #include "MassCharacterMovementTranslators.generated.h"
 
 class UCharacterMovementComponent;
diff --git a/Engine/Plugins/Runtime/MassGameplay/Source/MassGameplayDebug/Public/AssignDebugVisProcessor.h b/Engine/Plugins/Runtime/MassGameplay/Source/MassGameplayDebug/Public/AssignDebugVisProcessor.h
--- a/Engine/Plugins/Runtime/MassGameplay/Source/MassGameplayDebug/Public/AssignDebugVisProcessor.h
+++ b/Engine/Plugins/Runtime/MassGameplay/Source/MassGameplayDebug/Public/AssignDebugVisProcessor.h
@@ -2,7 +2,7 @@
 
 #pragma once
 
-#include "MassTranslator.h"
+#include "MassObserverProcessor.h"
 #include "AssignDebugVisProcessor.generated.h"
 
 
diff --git a/Engine/Plugins/Runtime/MassGameplay/Source/MassMovement/Public/MassVelocityRandomizerTrait.h b/Engine/Plugins/Runtime/MassGameplay/Source/MassMovement/Public/MassVelocityRandomizerTrait.h
--- a/Engine/Plugins/Runtime/MassGameplay/Source/MassMovement/Public/MassVelocityRandomizerTrait.h
+++ b/Engine/Plugins/Runtime/MassGameplay/Source/MassMovement/Public/MassVelocityRandomizerTrait.h
@@ -3,7 +3,7 @@
 #pragma once
 
 #include "MassEntityTraitBase.h"
-#include "MassTranslator.h"
+#include "MassObserverProcessor.h"
 #include "MassVelocityRandomizerTrait.generated.h"
 
 
diff --git a/Engine/Plugins/Runtime/MassGameplay/Source/MassReplication/Public/MassReplicationTypes.h b/Engine/Plugins/Runtime/MassGameplay/Source/MassReplication/Public/MassReplicationTypes.h
--- a/Engine/Plugins/Runtime/MassGameplay/Source/MassReplication/Public/MassReplicationTypes.h
+++ b/Engine/Plugins/Runtime/MassGameplay/Source/MassReplication/Public/MassReplicationTypes.h
@@ -6,7 +6,7 @@
 #include "MassLODManager.h"
 #include "IndexedHandle.h"
 #include "MassEntityTemplate.h"
-#include "MassTranslator.h"
+#include "MassObserverProcessor.h"
 #include "MassLODLogic.h"
 
 #include "MassReplicationTypes.generated.h"
diff --git a/Engine/Plugins/Runtime/MassGameplay/Source/MassRepresentation/Public/MassRepresentationProcessor.h b/Engine/Plugins/Runtime/MassGameplay/Source/MassRepresentation/Public/MassRepresentationProcessor.h
--- a/Engine/Plugins/Runtime/MassGameplay/Source/MassRepresentation/Public/MassRepresentationProcessor.h
+++ b/Engine/Plugins/Runtime/MassGameplay/Source/MassRepresentation/Public/MassRepresentationProcessor.h
@@ -4,7 +4,7 @@
 
 #include "MassRepresentationFragments.h"
 #include "MassActorSpawnerSubsystem.h"
-#include "MassTranslator.h"
+#include "MassObserverProcessor.h"
 
 #include "MassRepresentationProcessor.generated.h"
 
@@ -216,7 +216,7 @@ struct FMassRepresentationDefaultDestructorTag : public FMassTag
 };
 
 UCLASS()
-class MASSREPRESENTATION_API UMassRepresentationFragmentDestructor : public UMassFragmentDestructor
+class MASSREPRESENTATION_API UMassRepresentationFragmentDestructor : public UMassFragmentDeinitializer
 {
 	GENERATED_BODY()
 
diff --git a/Engine/Plugins/Runtime/MassGameplay/Source/MassSpawner/Public/MassTranslator.h b/Engine/Plugins/Runtime/MassGameplay/Source/MassSpawner/Public/MassTranslator.h
--- a/Engine/Plugins/Runtime/MassGameplay/Source/MassSpawner/Public/MassTranslator.h
+++ b/Engine/Plugins/Runtime/MassGameplay/Source/MassSpawner/Public/MassTranslator.h
@@ -3,9 +3,7 @@
 #pragma once
 
 #include "MassProcessor.h"
-#include "MassCommonTypes.h"
 #include "MassEntityTypes.h"
-#include "MassObserverProcessor.h"
 #include "MassTranslator.generated.h"
 
 
