Small cleanup of execution context being passed around into sub methods.

#rnx

diff --git a/Engine/Plugins/AI/MassCrowd/Source/MassCrowd/Private/MassCrowdRepresentationProcessor.cpp b/Engine/Plugins/AI/MassCrowd/Source/MassCrowd/Private/MassCrowdRepresentationProcessor.cpp
--- a/Engine/Plugins/AI/MassCrowd/Source/MassCrowd/Private/MassCrowdRepresentationProcessor.cpp
+++ b/Engine/Plugins/AI/MassCrowd/Source/MassCrowd/Private/MassCrowdRepresentationProcessor.cpp
@@ -69,9 +69,9 @@ void UMassCrowdRepresentationProcessor::InitializeVelocity(UMassEntitySubsystem&
 		});
 }
 
-void UMassCrowdRepresentationProcessor::SetActorEnabled(const EActorEnabledType EnabledType, AActor& Actor, const int32 EntityIdx, FMassExecutionContext& Context)
+void UMassCrowdRepresentationProcessor::SetActorEnabled(const EActorEnabledType EnabledType, AActor& Actor, const int32 EntityIdx, FMassCommandBuffer& CommandBuffer)
 {
-	Super::SetActorEnabled(EnabledType, Actor, EntityIdx, Context);
+	Super::SetActorEnabled(EnabledType, Actor, EntityIdx, CommandBuffer);
 
 	const bool bEnabled = EnabledType != EActorEnabledType::Disabled;
 
@@ -140,7 +140,7 @@ AActor* UMassCrowdRepresentationProcessor::GetOrSpawnActor(const FMassEntityHand
 	return Super::GetOrSpawnActor(MassAgent, ActorInfo, RootTransform, TemplateActorIndex, SpawnRequestHandle, Priority);
 }
 
-void UMassCrowdRepresentationProcessor::TeleportActor(const FTransform& Transform, AActor& Actor, FMassExecutionContext& Context)
+void UMassCrowdRepresentationProcessor::TeleportActor(const FTransform& Transform, AActor& Actor, FMassCommandBuffer& CommandBuffer)
 {
 	FTransform RootTransform = Transform;
 
@@ -160,5 +160,5 @@ void UMassCrowdRepresentationProcessor::TeleportActor(const FTransform& Transfor
 			RootTransform.SetLocation(OutHit.Location);
 		}
 	}
-	Super::TeleportActor(RootTransform, Actor, Context);
+	Super::TeleportActor(RootTransform, Actor, CommandBuffer);
 }
diff --git a/Engine/Plugins/AI/MassCrowd/Source/MassCrowd/Public/MassCrowdRepresentationProcessor.h b/Engine/Plugins/AI/MassCrowd/Source/MassCrowd/Public/MassCrowdRepresentationProcessor.h
--- a/Engine/Plugins/AI/MassCrowd/Source/MassCrowd/Public/MassCrowdRepresentationProcessor.h
+++ b/Engine/Plugins/AI/MassCrowd/Source/MassCrowd/Public/MassCrowdRepresentationProcessor.h
@@ -36,7 +36,7 @@ protected:
 	 * @param EntityIdx is the entity index currently processing
 	 * @param Context is the current Mass execution context
 	 */
-	virtual void SetActorEnabled(const EActorEnabledType EnabledType, AActor& Actor, const int32 EntityIdx, FMassExecutionContext& Context) override;
+	virtual void SetActorEnabled(const EActorEnabledType EnabledType, AActor& Actor, const int32 EntityIdx, FMassCommandBuffer& CommandBuffer) override;
 
 	/** 
 	 * Returns an actor of the template type and setup fragments values from it
@@ -58,7 +58,7 @@ protected:
 	 * @param Actor is the actual actor to teleport
 	 * @param Context is the current Mass execution context
 	 */
-	virtual void TeleportActor(const FTransform& Transform, AActor& Actor, FMassExecutionContext& Context) override;
+	virtual void TeleportActor(const FTransform& Transform, AActor& Actor, FMassCommandBuffer& CommandBuffer) override;
 
 	/**
 	 * Initializes velocity for each entity when just switch to use movement component, must be called within a ForEachEntityChunk
diff --git a/Engine/Plugins/Runtime/MassGameplay/Source/MassRepresentation/Private/MassRepresentationProcessor.cpp b/Engine/Plugins/Runtime/MassGameplay/Source/MassRepresentation/Private/MassRepresentationProcessor.cpp
--- a/Engine/Plugins/Runtime/MassGameplay/Source/MassRepresentation/Private/MassRepresentationProcessor.cpp
+++ b/Engine/Plugins/Runtime/MassGameplay/Source/MassRepresentation/Private/MassRepresentationProcessor.cpp
@@ -105,12 +105,12 @@ void UMassRepresentationProcessor::UpdateRepresentation(FMassExecutionContext& C
 				if (Representation.HighResTemplateActorIndex != Representation.LowResTemplateActorIndex || !bKeepLowResActors)
 				{
 					// Try releasing the high actor or any high res spawning request
-					if (ReleaseActorOrCancelSpawning(MassAgent, ActorInfo, Representation.HighResTemplateActorIndex, Representation.ActorSpawnRequestHandle, Context))
+					if (ReleaseActorOrCancelSpawning(MassAgent, ActorInfo, Representation.HighResTemplateActorIndex, Representation.ActorSpawnRequestHandle, Context.Defer()))
 					{
 						Actor = ActorInfo.GetOwnedByMassMutable();
 					}
 					// Do not do the same with low res if indicated so
-					if (!bKeepLowResActors && ReleaseActorOrCancelSpawning(MassAgent, ActorInfo, Representation.LowResTemplateActorIndex, Representation.ActorSpawnRequestHandle, Context))
+					if (!bKeepLowResActors && ReleaseActorOrCancelSpawning(MassAgent, ActorInfo, Representation.LowResTemplateActorIndex, Representation.ActorSpawnRequestHandle, Context.Defer()))
 					{
 						Actor = ActorInfo.GetOwnedByMassMutable();
 					}
@@ -123,7 +123,7 @@ void UMassRepresentationProcessor::UpdateRepresentation(FMassExecutionContext& C
 			}
 			if (Actor != nullptr)
 			{
-				SetActorEnabled(EActorEnabledType::Disabled, *Actor, EntityIdx, Context);
+				SetActorEnabled(EActorEnabledType::Disabled, *Actor, EntityIdx, Context.Defer());
 			}
 		};
 
@@ -153,7 +153,7 @@ void UMassRepresentationProcessor::UpdateRepresentation(FMassExecutionContext& C
 						// If the low res is different than the high res, cancel any pending spawn request that is the opposite of what is needed.
 						if (Representation.LowResTemplateActorIndex != Representation.HighResTemplateActorIndex)
 						{
-							ReleaseActorOrCancelSpawning(MassAgent, ActorInfo, bHighResActor ? Representation.LowResTemplateActorIndex : Representation.HighResTemplateActorIndex, Representation.ActorSpawnRequestHandle, Context, /*bCancelSpawningOnly*/true);
+							ReleaseActorOrCancelSpawning(MassAgent, ActorInfo, bHighResActor ? Representation.LowResTemplateActorIndex : Representation.HighResTemplateActorIndex, Representation.ActorSpawnRequestHandle, Context.Defer(), /*bCancelSpawningOnly*/true);
 							Actor = ActorInfo.GetOwnedByMassMutable();
 						}
 
@@ -183,10 +183,10 @@ void UMassRepresentationProcessor::UpdateRepresentation(FMassExecutionContext& C
 						// Needs to be done before enabling the actor so the animation initialization can use the new values
 						if (Representation.CurrentRepresentation == ERepresentationType::StaticMeshInstance)
 						{
-							TeleportActor(Representation.PrevTransform, *NewActor, Context);
+							TeleportActor(Representation.PrevTransform, *NewActor, Context.Defer());
 						}
 
-						SetActorEnabled(bHighResActor ? EActorEnabledType::HighRes : EActorEnabledType::LowRes, *NewActor, EntityIdx, Context);
+						SetActorEnabled(bHighResActor ? EActorEnabledType::HighRes : EActorEnabledType::LowRes, *NewActor, EntityIdx, Context.Defer());
 						Representation.CurrentRepresentation = WantedRepresentationType;
 					}
 					else if (!Actor)
@@ -208,12 +208,12 @@ void UMassRepresentationProcessor::UpdateRepresentation(FMassExecutionContext& C
 					if (!Actor || ActorInfo.IsOwnedByMass())
 					{
 						// Try releasing both, could have an high res spawned actor and a spawning request for a low res one
-						ReleaseActorOrCancelSpawning(MassAgent, ActorInfo, Representation.LowResTemplateActorIndex, Representation.ActorSpawnRequestHandle, Context);
-						ReleaseActorOrCancelSpawning(MassAgent, ActorInfo, Representation.HighResTemplateActorIndex, Representation.ActorSpawnRequestHandle, Context);
+						ReleaseActorOrCancelSpawning(MassAgent, ActorInfo, Representation.LowResTemplateActorIndex, Representation.ActorSpawnRequestHandle, Context.Defer());
+						ReleaseActorOrCancelSpawning(MassAgent, ActorInfo, Representation.HighResTemplateActorIndex, Representation.ActorSpawnRequestHandle, Context.Defer());
 					}
 					else
 					{
-						SetActorEnabled(EActorEnabledType::Disabled, *Actor, EntityIdx, Context);
+						SetActorEnabled(EActorEnabledType::Disabled, *Actor, EntityIdx, Context.Defer());
 					}
 					Representation.CurrentRepresentation = ERepresentationType::None;
 					break;
@@ -252,7 +252,7 @@ AActor* UMassRepresentationProcessor::GetOrSpawnActor(const FMassEntityHandle Ma
 		FMassActorPostSpawnDelegate::CreateUObject(this, &UMassRepresentationProcessor::OnActorPostSpawn));
 }
 
-bool UMassRepresentationProcessor::ReleaseActorOrCancelSpawning(const FMassEntityHandle MassAgent, FDataFragment_Actor& ActorInfo, const int16 TemplateActorIndex, FMassActorSpawnRequestHandle& SpawnRequestHandle, FMassExecutionContext& Context, bool bCancelSpawningOnly /*= false*/)
+bool UMassRepresentationProcessor::ReleaseActorOrCancelSpawning(const FMassEntityHandle MassAgent, FDataFragment_Actor& ActorInfo, const int16 TemplateActorIndex, FMassActorSpawnRequestHandle& SpawnRequestHandle, FMassCommandBuffer& CommandBuffer, bool bCancelSpawningOnly /*= false*/)
 {
 	check(RepresentationSubsystem);
 	check(!ActorInfo.IsValid() || ActorInfo.IsOwnedByMass());
@@ -272,7 +272,7 @@ bool UMassRepresentationProcessor::ReleaseActorOrCancelSpawning(const FMassEntit
 			TObjectKey<const AActor> ActorKey(Actor); 
 			if (UMassActorSubsystem* MassActorSubsystem = UWorld::GetSubsystem<UMassActorSubsystem>(World))
 			{
-				Context.Defer().EmplaceCommand<FDeferredCommand>([MassActorSubsystem, ActorKey](UMassEntitySubsystem&)
+				CommandBuffer.EmplaceCommand<FDeferredCommand>([MassActorSubsystem, ActorKey](UMassEntitySubsystem&)
 				{
 					MassActorSubsystem->RemoveHandleForActor(ActorKey);
 				});
@@ -283,7 +283,7 @@ bool UMassRepresentationProcessor::ReleaseActorOrCancelSpawning(const FMassEntit
 	return false;
 }
 
-void UMassRepresentationProcessor::SetActorEnabled(const EActorEnabledType EnabledType, AActor& Actor, const int32 EntityIdx, FMassExecutionContext& Context)
+void UMassRepresentationProcessor::SetActorEnabled(const EActorEnabledType EnabledType, AActor& Actor, const int32 EntityIdx, FMassCommandBuffer& CommandBuffer)
 {
 	const bool bEnabled = EnabledType != EActorEnabledType::Disabled;
 	if (Actor.IsActorTickEnabled() != bEnabled)
@@ -293,18 +293,18 @@ void UMassRepresentationProcessor::SetActorEnabled(const EActorEnabledType Enabl
 	if (Actor.GetActorEnableCollision() != bEnabled)
 	{
 		// Deferring this as there is a callback internally that could end up doing things outside of the game thread and will fire checks(Chaos mostly)
-		Context.Defer().EmplaceCommand<FDeferredCommand>( [&Actor,bEnabled](UMassEntitySubsystem&)
+		CommandBuffer.EmplaceCommand<FDeferredCommand>( [&Actor,bEnabled](UMassEntitySubsystem&)
 		{
 			Actor.SetActorEnableCollision(bEnabled);
 		});
 	}
 }
 
-void UMassRepresentationProcessor::TeleportActor(const FTransform& Transform, AActor& Actor, FMassExecutionContext& Context)
+void UMassRepresentationProcessor::TeleportActor(const FTransform& Transform, AActor& Actor, FMassCommandBuffer& CommandBuffer)
 {
 	if (!Actor.GetTransform().Equals(Transform))
 	{
-		Context.Defer().EmplaceCommand<FDeferredCommand>([&Actor, Transform](UMassEntitySubsystem&)
+		CommandBuffer.EmplaceCommand<FDeferredCommand>([&Actor, Transform](UMassEntitySubsystem&)
 		{
 			Actor.SetActorTransform(Transform, /*bSweep*/false, /*OutSweepHitResult*/nullptr, ETeleportType::TeleportPhysics);
 		});
@@ -364,7 +364,7 @@ void UMassRepresentationProcessor::UpdateVisualization(FMassExecutionContext& Co
 		const FMassEntityHandle Entity = Context.GetEntity(EntityIdx);
 		FMassRepresentationFragment& Representation = RepresentationList[EntityIdx];
 		const FMassRepresentationLODFragment& RepresentationLOD = RepresentationLODList[EntityIdx];
-		UpdateEntityVisibility(Entity, Representation, RepresentationLOD, ChunkData, Context);
+		UpdateEntityVisibility(Entity, Representation, RepresentationLOD, ChunkData, Context.Defer());
 	}
 }
 
@@ -410,7 +410,7 @@ FMassVisualizationChunkFragment& UMassRepresentationProcessor::UpdateChunkVisibi
 	return ChunkData;
 }
 
-void UMassRepresentationProcessor::UpdateEntityVisibility(const FMassEntityHandle Entity, const FMassRepresentationFragment& Representation, const FMassRepresentationLODFragment& RepresentationLOD, FMassVisualizationChunkFragment& ChunkData, FMassExecutionContext& Context)
+void UMassRepresentationProcessor::UpdateEntityVisibility(const FMassEntityHandle Entity, const FMassRepresentationFragment& Representation, const FMassRepresentationLODFragment& RepresentationLOD, FMassVisualizationChunkFragment& ChunkData, FMassCommandBuffer& CommandBuffer)
 {
 	// Move the visible entities together into same chunks so we can skip entire chunk when not visible as an optimization
 	const EMassVisibility Visibility = Representation.CurrentRepresentation != ERepresentationType::None ? 
@@ -418,7 +418,7 @@ void UMassRepresentationProcessor::UpdateEntityVisibility(const FMassEntityHandl
 	const EMassVisibility ChunkVisibility = ChunkData.GetVisibility();
 	if (ChunkVisibility != Visibility)
 	{
-		Context.Defer().PushCommand(FCommandSwapTags(Entity, UE::MassRepresentation::GetTagFromVisibility(ChunkVisibility), UE::MassRepresentation::GetTagFromVisibility(Visibility)));
+		CommandBuffer.PushCommand(FCommandSwapTags(Entity, UE::MassRepresentation::GetTagFromVisibility(ChunkVisibility), UE::MassRepresentation::GetTagFromVisibility(Visibility)));
 		ChunkData.SetContainsNewlyVisibleEntity(Visibility == EMassVisibility::CanBeSeen);
 	}
 }
diff --git a/Engine/Plugins/Runtime/MassGameplay/Source/MassRepresentation/Public/MassRepresentationProcessor.h b/Engine/Plugins/Runtime/MassGameplay/Source/MassRepresentation/Public/MassRepresentationProcessor.h
--- a/Engine/Plugins/Runtime/MassGameplay/Source/MassRepresentation/Public/MassRepresentationProcessor.h
+++ b/Engine/Plugins/Runtime/MassGameplay/Source/MassRepresentation/Public/MassRepresentationProcessor.h
@@ -62,29 +62,29 @@ protected:
 	 * @param ActorInfo is the fragment where we are going to store the actor pointer
 	 * @param TemplateActorIndex is the index of the type to release
 	 * @param SpawnRequestHandle (in/out) In: previously requested spawn to cancel if any
-	 * @param Context of the execution from the entity sub system
+	 * @param CommandBuffer to queue up anything that is thread sensitive
 	 * @param bCancelSpawningOnly tell to only cancel the existing spawning request and to not release the associated actor it any.
 	 * @return if the actor was release or the spawning was canceled.
 	 */
-	virtual bool ReleaseActorOrCancelSpawning(const FMassEntityHandle MassAgent, FDataFragment_Actor& ActorInfo, const int16 TemplateActorIndex, FMassActorSpawnRequestHandle& SpawnRequestHandle, FMassExecutionContext& Context, const bool bCancelSpawningOnly = false);
+	virtual bool ReleaseActorOrCancelSpawning(const FMassEntityHandle MassAgent, FDataFragment_Actor& ActorInfo, const int16 TemplateActorIndex, FMassActorSpawnRequestHandle& SpawnRequestHandle, FMassCommandBuffer& CommandBuffer, const bool bCancelSpawningOnly = false);
 
 	/** 
 	 * Enable/disable a spawned actor
 	 * @param EnabledType is the type of enabling to do on this actor
 	 * @param Actor is the actual actor to perform enabling type on
 	 * @param EntityIdx is the entity index currently processing
-	 * @param Context is the current Mass execution context 
+	 * @param CommandBuffer to queue up anything that is thread sensitive
 	 */
-	virtual void SetActorEnabled(const EActorEnabledType EnabledType, AActor& Actor, const int32 EntityIdx, FMassExecutionContext& Context);
+	virtual void SetActorEnabled(const EActorEnabledType EnabledType, AActor& Actor, const int32 EntityIdx, FMassCommandBuffer& CommandBuffer);
 
 	/**
 	 * Teleports the actor at the specified transform by preserving its velocity and without collision.
 	 * The destination will be adjusted to fit an existing capsule.
 	 * @param Transform is the new actor's transform 
 	 * @param Actor is the actual actor to teleport
-	 * @param Context is the current Mass execution context
+	 * @param CommandBuffer to queue up anything that is thread sensitive
 	 */
-	virtual void TeleportActor(const FTransform& Transform, AActor& Actor, FMassExecutionContext& Context);
+	virtual void TeleportActor(const FTransform& Transform, AActor& Actor, FMassCommandBuffer& CommandBuffer);
 
 	/**
 	 * Method that will be bound to a delegate called before the spawning of an actor to let the requester prepare it
@@ -121,9 +121,9 @@ protected:
 	 * @param Representation fragment containing the current and previous visual state
 	 * @param RepresentationLOD fragment containing the visibility information
 	 * @param ChunkData is the visualization chunk fragment
-     * @param Context of the execution from the entity sub system
+	 * @param CommandBuffer to queue up anything that is thread sensitive
 	 */
-	static void UpdateEntityVisibility(const FMassEntityHandle Entity, const FMassRepresentationFragment& Representation, const FMassRepresentationLODFragment& RepresentationLOD, FMassVisualizationChunkFragment& ChunkData, FMassExecutionContext& Context);
+	static void UpdateEntityVisibility(const FMassEntityHandle Entity, const FMassRepresentationFragment& Representation, const FMassRepresentationLODFragment& RepresentationLOD, FMassVisualizationChunkFragment& ChunkData, FMassCommandBuffer& CommandBuffer);
 
 public:
 	/**
