Fix for crashes happening in network games.

#rnx

diff --git a/Engine/Plugins/Runtime/MassEntity/Source/MassEntity/Private/MassCommandBuffer.cpp b/Engine/Plugins/Runtime/MassEntity/Source/MassEntity/Private/MassCommandBuffer.cpp
--- a/Engine/Plugins/Runtime/MassEntity/Source/MassEntity/Private/MassCommandBuffer.cpp
+++ b/Engine/Plugins/Runtime/MassEntity/Source/MassEntity/Private/MassCommandBuffer.cpp
@@ -267,5 +267,10 @@ void FCommandRemoveTag::Execute(UMassEntitySubsystem& System) const
 
 void FCommandRemoveComposition::Execute(UMassEntitySubsystem& System) const
 {
+	if (System.IsEntityValid(TargetEntity) == false)
+	{
+		return;
+	}
+
 	System.RemoveCompositionFromEntity(TargetEntity, Descriptor);
 }
diff --git a/Engine/Plugins/Runtime/MassGameplay/Source/MassRepresentation/Private/MassRepresentationProcessor.cpp b/Engine/Plugins/Runtime/MassGameplay/Source/MassRepresentation/Private/MassRepresentationProcessor.cpp
--- a/Engine/Plugins/Runtime/MassGameplay/Source/MassRepresentation/Private/MassRepresentationProcessor.cpp
+++ b/Engine/Plugins/Runtime/MassGameplay/Source/MassRepresentation/Private/MassRepresentationProcessor.cpp
@@ -488,7 +488,7 @@ UMassRepresentationFragmentDestructor::UMassRepresentationFragmentDestructor()
 {
 	// By putting this to null, this Deinitializer needs to be explicitly added via a Entity traits and will no longer be automatically called
 	// This FragmentType should be deprecated and we should always use the traits.
-	FragmentType = nullptr; // FMassRepresentationFragment::StaticStruct();
+	FragmentType = FMassRepresentationFragment::StaticStruct();
 	ExecutionFlags = (int32)(EProcessorExecutionFlags::All);
 }
 
