Added a check on the EntitySystem validity.

diff --git a/Engine/Plugins/Runtime/MassGameplay/Source/MassLOD/Private/MassLODTrait.cpp b/Engine/Plugins/Runtime/MassGameplay/Source/MassLOD/Private/MassLODTrait.cpp
--- a/Engine/Plugins/Runtime/MassGameplay/Source/MassLOD/Private/MassLODTrait.cpp
+++ b/Engine/Plugins/Runtime/MassGameplay/Source/MassLOD/Private/MassLODTrait.cpp
@@ -28,6 +28,7 @@ void UMassSimulationLODTrait::BuildTemplate(FMassEntityTemplateBuildContext& Bui
 	}
 
 	UMassEntitySubsystem* EntitySubsystem = UWorld::GetSubsystem<UMassEntitySubsystem>(&World);
+	check(EntitySubsystem);
 
 	uint32 ConfigHash = UE::StructUtils::GetStructCrc32(FConstStructView::Make(Config));
 	FConstSharedStruct ConfigFragment = EntitySubsystem->GetOrCreateConstSharedFragment(ConfigHash, Config);
