Non-Unity fix.

diff --git a/Engine/Plugins/Runtime/MassGameplay/Source/MassActors/Private/Translators/MassSceneComponentLocationTranslator.cpp b/Engine/Plugins/Runtime/MassGameplay/Source/MassActors/Private/Translators/MassSceneComponentLocationTranslator.cpp
--- a/Engine/Plugins/Runtime/MassGameplay/Source/MassActors/Private/Translators/MassSceneComponentLocationTranslator.cpp
+++ b/Engine/Plugins/Runtime/MassGameplay/Source/MassActors/Private/Translators/MassSceneComponentLocationTranslator.cpp
@@ -4,6 +4,7 @@
 #include "MassMovementTypes.h"
 #include "Components/SceneComponent.h"
 #include "MassEntitySubsystem.h"
+#include "MassCommonTypes.h"
 
 //----------------------------------------------------------------------//
 //  UMassSceneComponentLocationToMassTranslator
