RepresentationSystem cleanup!!!

* Created a Shared fragment with to store the pointer to the representation system to use for each entity.
* Cleanup the representation destructor need of a tag to select the right one. Now only one destructor needed.
* Removed the configs on the representation processors
* Adapt the code to fetch the representation system from the shared fragment
* It is a really nice way to let the MAssEntitySystem dependencies on the representation system. Will be Usefull for multithreading.

#rnx

diff --git a/Engine/Plugins/AI/MassCrowd/Source/MassCrowd/Private/MassCrowdRepresentationProcessor.cpp b/Engine/Plugins/AI/MassCrowd/Source/MassCrowd/Private/MassCrowdRepresentationProcessor.cpp
--- a/Engine/Plugins/AI/MassCrowd/Source/MassCrowd/Private/MassCrowdRepresentationProcessor.cpp
+++ b/Engine/Plugins/AI/MassCrowd/Source/MassCrowd/Private/MassCrowdRepresentationProcessor.cpp
@@ -9,7 +9,7 @@
 #include "GameFramework/Character.h"
 #include "Translators/MassCharacterMovementTranslators.h"
 #include "MassAgentComponent.h"
-#include "MassCrowdRepresentationSubsystem.h"
+#include "MassRepresentationSubsystem.h"
 
 UMassCrowdRepresentationProcessor::UMassCrowdRepresentationProcessor()
 {
@@ -21,9 +21,6 @@ UMassCrowdRepresentationProcessor::UMassCrowdRepresentationProcessor()
 void UMassCrowdRepresentationProcessor::Initialize(UObject& Owner)
 {
 	Super::Initialize(Owner);
-
-	// Override default representation subsystem by the crowd one to support parallelization
-	RepresentationSubsystem = UWorld::GetSubsystem<UMassCrowdRepresentationSubsystem>(Owner.GetWorld());
 }
 
 
@@ -94,11 +91,11 @@ void UMassCrowdRepresentationProcessor::SetActorEnabled(const EActorEnabledType
 	}
 }
 
-AActor* UMassCrowdRepresentationProcessor::GetOrSpawnActor(const FMassEntityHandle MassAgent, FDataFragment_Actor& ActorInfo, const FTransform& Transform, const int16 TemplateActorIndex, FMassActorSpawnRequestHandle& SpawnRequestHandle, const float Priority)
+AActor* UMassCrowdRepresentationProcessor::GetOrSpawnActor(UMassRepresentationSubsystem& RepresentationSubsystem, const FMassEntityHandle MassAgent, FDataFragment_Actor& ActorInfo, const FTransform& Transform, const int16 TemplateActorIndex, FMassActorSpawnRequestHandle& SpawnRequestHandle, const float Priority)
 {
 	FTransform RootTransform = Transform;
 	
-	if (const AActor* DefaultActor = RepresentationSubsystem->GetTemplateActorClass(TemplateActorIndex).GetDefaultObject())
+	if (const AActor* DefaultActor = RepresentationSubsystem.GetTemplateActorClass(TemplateActorIndex).GetDefaultObject())
 	{
 		if (const UCapsuleComponent* CapsuleComp = DefaultActor->FindComponentByClass<UCapsuleComponent>())
 		{
@@ -106,7 +103,7 @@ AActor* UMassCrowdRepresentationProcessor::GetOrSpawnActor(const FMassEntityHand
 		}
 	}
 
-	return Super::GetOrSpawnActor(MassAgent, ActorInfo, RootTransform, TemplateActorIndex, SpawnRequestHandle, Priority);
+	return Super::GetOrSpawnActor(RepresentationSubsystem, MassAgent, ActorInfo, RootTransform, TemplateActorIndex, SpawnRequestHandle, Priority);
 }
 
 void UMassCrowdRepresentationProcessor::TeleportActor(const FTransform& Transform, AActor& Actor, FMassCommandBuffer& CommandBuffer)
diff --git a/Engine/Plugins/AI/MassCrowd/Source/MassCrowd/Private/MassCrowdServerRepresentationTrait.cpp b/Engine/Plugins/AI/MassCrowd/Source/MassCrowd/Private/MassCrowdServerRepresentationTrait.cpp
--- a/Engine/Plugins/AI/MassCrowd/Source/MassCrowd/Private/MassCrowdServerRepresentationTrait.cpp
+++ b/Engine/Plugins/AI/MassCrowd/Source/MassCrowd/Private/MassCrowdServerRepresentationTrait.cpp
@@ -20,15 +20,25 @@ void UMassCrowdServerRepresentationTrait::BuildTemplate(FMassEntityTemplateBuild
 	// won't override any fragments that are already there
 	BuildContext.AddFragment<FDataFragment_Transform>();
 
-	if (UMassCrowdRepresentationSubsystem* RepresentationSubsystem = World.GetSubsystem<UMassCrowdRepresentationSubsystem>())
-	{
-		FMassRepresentationFragment& RepresentationFragment = BuildContext.AddFragment_GetRef<FMassRepresentationFragment>();
-		RepresentationFragment.StaticMeshDescIndex = INDEX_NONE;
-		RepresentationFragment.HighResTemplateActorIndex = TemplateActor.Get() ? RepresentationSubsystem->FindOrAddTemplateActor(TemplateActor.Get()) : INDEX_NONE;
-		RepresentationFragment.LowResTemplateActorIndex = INDEX_NONE;
+	UMassEntitySubsystem* EntitySubsystem = UWorld::GetSubsystem<UMassEntitySubsystem>(&World);
+	check(EntitySubsystem);
 
-		BuildContext.AddFragment<FMassRepresentationLODFragment>();
+	UMassCrowdRepresentationSubsystem* RepresentationSubsystem = World.GetSubsystem<UMassCrowdRepresentationSubsystem>();
+	check(RepresentationSubsystem);
 
-		BuildContext.AddChunkFragment<FMassVisualizationChunkFragment>();
-	}
+	FMassRepresentationSubsystemFragment Subsystem;
+	Subsystem.RepresentationSubsystem = RepresentationSubsystem;
+	uint32 SubsystemHash = UE::StructUtils::GetStructCrc32(FConstStructView::Make(Subsystem));
+	FSharedStruct SubsystemFragment = EntitySubsystem->GetOrCreateSharedFragment<FMassRepresentationSubsystemFragment>(SubsystemHash, Subsystem);
+	BuildContext.AddSharedFragment(SubsystemFragment);
+
+	FMassRepresentationFragment& RepresentationFragment = BuildContext.AddFragment_GetRef<FMassRepresentationFragment>();
+	RepresentationFragment.StaticMeshDescIndex = INDEX_NONE;
+	RepresentationFragment.HighResTemplateActorIndex = TemplateActor.Get() ? RepresentationSubsystem->FindOrAddTemplateActor(TemplateActor.Get()) : INDEX_NONE;
+	RepresentationFragment.LowResTemplateActorIndex = INDEX_NONE;
+
+	BuildContext.AddFragment<FMassRepresentationLODFragment>();
+
+	// @todo figure out if this chunk fragment is really needed?
+	BuildContext.AddChunkFragment<FMassVisualizationChunkFragment>();
 }
\ No newline at end of file
diff --git a/Engine/Plugins/AI/MassCrowd/Source/MassCrowd/Private/MassCrowdVisualizationProcessor.cpp b/Engine/Plugins/AI/MassCrowd/Source/MassCrowd/Private/MassCrowdVisualizationProcessor.cpp
--- a/Engine/Plugins/AI/MassCrowd/Source/MassCrowd/Private/MassCrowdVisualizationProcessor.cpp
+++ b/Engine/Plugins/AI/MassCrowd/Source/MassCrowd/Private/MassCrowdVisualizationProcessor.cpp
@@ -77,22 +77,4 @@ void UMassCrowdVisualizationProcessor::Execute(UMassEntitySubsystem& EntitySubsy
 			}
 		});
 	}
-}
-
-//----------------------------------------------------------------------//
-// UMassCrowdRepresentationFragmentDestructor
-//----------------------------------------------------------------------//
-void UMassCrowdRepresentationFragmentDestructor::Initialize(UObject& Owner)
-{
-	Super::Initialize(Owner);
-
-	// Override default representation subsystem by the crowd one to support parallelization
-	RepresentationSubsystem = UWorld::GetSubsystem<UMassCrowdRepresentationSubsystem>(Owner.GetWorld());
-}
-
-void UMassCrowdRepresentationFragmentDestructor::ConfigureQueries()
-{
-	EntityQuery.AddRequirement<FMassRepresentationFragment>(EMassFragmentAccess::ReadOnly);
-	EntityQuery.AddRequirement<FDataFragment_Actor>(EMassFragmentAccess::ReadWrite);
-	EntityQuery.AddTagRequirement<FMassCrowdRepresentationDestructorTag>(EMassFragmentPresence::All);
 }
\ No newline at end of file
diff --git a/Engine/Plugins/AI/MassCrowd/Source/MassCrowd/Private/MassCrowdVisualizationTrait.cpp b/Engine/Plugins/AI/MassCrowd/Source/MassCrowd/Private/MassCrowdVisualizationTrait.cpp
--- a/Engine/Plugins/AI/MassCrowd/Source/MassCrowd/Private/MassCrowdVisualizationTrait.cpp
+++ b/Engine/Plugins/AI/MassCrowd/Source/MassCrowd/Private/MassCrowdVisualizationTrait.cpp
@@ -6,9 +6,6 @@
 
 UMassCrowdVisualizationTrait::UMassCrowdVisualizationTrait()
 {
-	// Override the subsystem and deinitializer to support parallelization of the crowd
+	// Override the subsystem to support parallelization of the crowd
 	RepresentationSubsystemClass = UMassCrowdRepresentationSubsystem::StaticClass();
-	// @todo the following line will be cut once new de/initializers are in
-	RepresentationFragmentDeinitializerClass = UMassCrowdRepresentationFragmentDestructor::StaticClass();
-	RepresentationDestructorTag = FMassCrowdRepresentationDestructorTag::StaticStruct();
 }
diff --git a/Engine/Plugins/AI/MassCrowd/Source/MassCrowd/Public/MassCrowdRepresentationProcessor.h b/Engine/Plugins/AI/MassCrowd/Source/MassCrowd/Public/MassCrowdRepresentationProcessor.h
--- a/Engine/Plugins/AI/MassCrowd/Source/MassCrowd/Public/MassCrowdRepresentationProcessor.h
+++ b/Engine/Plugins/AI/MassCrowd/Source/MassCrowd/Public/MassCrowdRepresentationProcessor.h
@@ -40,6 +40,7 @@ protected:
 
 	/** 
 	 * Returns an actor of the template type and setup fragments values from it
+	 * @param RepresentationSubsystem to use to get or spawn the actor
 	 * @param MassAgent is the handle to the associated mass agent
 	 * @param ActorInfo is the fragment where we are going to store the actor pointer
 	 * @param Transform is the spatial information about where to spawn the actor 
@@ -49,7 +50,7 @@ protected:
 	 * @param Context of the execution from the entity sub system
 	 * @return the actor spawned
 	 */
-	virtual AActor* GetOrSpawnActor(const FMassEntityHandle MassAgent, FDataFragment_Actor& ActorInfo, const FTransform& Transform, const int16 TemplateActorIndex, FMassActorSpawnRequestHandle& SpawnRequestHandle, const float Priority) override;
+	virtual AActor* GetOrSpawnActor(UMassRepresentationSubsystem& RepresentationSubsystem, const FMassEntityHandle MassAgent, FDataFragment_Actor& ActorInfo, const FTransform& Transform, const int16 TemplateActorIndex, FMassActorSpawnRequestHandle& SpawnRequestHandle, const float Priority);
 
 	/**
 	 * Teleports the actor at the specified transform by preserving its velocity and without collision.
diff --git a/Engine/Plugins/AI/MassCrowd/Source/MassCrowd/Public/MassCrowdVisualizationProcessor.h b/Engine/Plugins/AI/MassCrowd/Source/MassCrowd/Public/MassCrowdVisualizationProcessor.h
--- a/Engine/Plugins/AI/MassCrowd/Source/MassCrowd/Public/MassCrowdVisualizationProcessor.h
+++ b/Engine/Plugins/AI/MassCrowd/Source/MassCrowd/Public/MassCrowdVisualizationProcessor.h
@@ -27,19 +27,4 @@ protected:
 	 * @param Context is the execution context to be passed when executing the lambdas
 	 */
 	virtual void Execute(UMassEntitySubsystem& EntitySubsystem, FMassExecutionContext& Context) override;
-};
-
-USTRUCT()
-struct FMassCrowdRepresentationDestructorTag : public FMassTag
-{
-	GENERATED_BODY()
-};
-
-UCLASS()
-class MASSCROWD_API UMassCrowdRepresentationFragmentDestructor : public UMassRepresentationFragmentDestructor
-{
-	GENERATED_BODY()
-protected:
-	virtual void Initialize(UObject& Owner) override;
-	virtual void ConfigureQueries() override;
-};
+};
\ No newline at end of file
diff --git a/Engine/Plugins/Runtime/MassGameplay/Source/MassRepresentation/Private/MassRepresentationProcessor.cpp b/Engine/Plugins/Runtime/MassGameplay/Source/MassRepresentation/Private/MassRepresentationProcessor.cpp
--- a/Engine/Plugins/Runtime/MassGameplay/Source/MassRepresentation/Private/MassRepresentationProcessor.cpp
+++ b/Engine/Plugins/Runtime/MassGameplay/Source/MassRepresentation/Private/MassRepresentationProcessor.cpp
@@ -38,6 +38,7 @@ void UMassRepresentationProcessor::ConfigureQueries()
 	EntityQuery.AddRequirement<FMassRepresentationLODFragment>(EMassFragmentAccess::ReadOnly);
 	EntityQuery.AddRequirement<FDataFragment_Actor>(EMassFragmentAccess::ReadWrite);
 	EntityQuery.AddChunkRequirement<FMassVisualizationChunkFragment>(EMassFragmentAccess::ReadWrite);
+	EntityQuery.AddSharedRequirement<FMassRepresentationSubsystemFragment>(EMassFragmentAccess::ReadWrite);
 }
 
 void UMassRepresentationProcessor::Initialize(UObject& Owner)
@@ -46,7 +47,6 @@ void UMassRepresentationProcessor::Initialize(UObject& Owner)
 
 	World = Owner.GetWorld();
 	CachedEntitySubsystem = UWorld::GetSubsystem<UMassEntitySubsystem>(World);
-	RepresentationSubsystem = UWorld::GetSubsystem<UMassRepresentationSubsystem>(World);
 
 	// Calculate the default representation when actor isn't spawned yet.
 	for (int32 LOD = EMassLOD::High; LOD < EMassLOD::Max; LOD++)
@@ -65,6 +65,7 @@ void UMassRepresentationProcessor::Initialize(UObject& Owner)
 
 void UMassRepresentationProcessor::UpdateRepresentation(FMassExecutionContext& Context)
 {
+	UMassRepresentationSubsystem* RepresentationSubsystem = Context.GetMutableSharedFragment<FMassRepresentationSubsystemFragment>().RepresentationSubsystem;
 	check(RepresentationSubsystem);
 
 	const TConstArrayView<FDataFragment_Transform> TransformList = Context.GetFragmentView<FDataFragment_Transform>();
@@ -105,12 +106,12 @@ void UMassRepresentationProcessor::UpdateRepresentation(FMassExecutionContext& C
 				if (Representation.HighResTemplateActorIndex != Representation.LowResTemplateActorIndex || !bKeepLowResActors)
 				{
 					// Try releasing the high actor or any high res spawning request
-					if (ReleaseActorOrCancelSpawning(MassAgent, ActorInfo, Representation.HighResTemplateActorIndex, Representation.ActorSpawnRequestHandle, Context.Defer()))
+					if (ReleaseActorOrCancelSpawning(*RepresentationSubsystem, MassAgent, ActorInfo, Representation.HighResTemplateActorIndex, Representation.ActorSpawnRequestHandle, Context.Defer()))
 					{
 						Actor = ActorInfo.GetOwnedByMassMutable();
 					}
 					// Do not do the same with low res if indicated so
-					if (!bKeepLowResActors && ReleaseActorOrCancelSpawning(MassAgent, ActorInfo, Representation.LowResTemplateActorIndex, Representation.ActorSpawnRequestHandle, Context.Defer()))
+					if (!bKeepLowResActors && ReleaseActorOrCancelSpawning(*RepresentationSubsystem, MassAgent, ActorInfo, Representation.LowResTemplateActorIndex, Representation.ActorSpawnRequestHandle, Context.Defer()))
 					{
 						Actor = ActorInfo.GetOwnedByMassMutable();
 					}
@@ -118,7 +119,7 @@ void UMassRepresentationProcessor::UpdateRepresentation(FMassExecutionContext& C
 				// If we already queued spawn request but have changed our mind, continue with it but once we get the actor back, disable it immediately
 				if (Representation.ActorSpawnRequestHandle.IsValid())
 				{
-					Actor = GetOrSpawnActor(MassAgent, ActorInfo, TransformFragment.GetTransform(), Representation.LowResTemplateActorIndex, Representation.ActorSpawnRequestHandle, GetSpawnPriority(RepresentationLOD));
+					Actor = GetOrSpawnActor(*RepresentationSubsystem, MassAgent, ActorInfo, TransformFragment.GetTransform(), Representation.LowResTemplateActorIndex, Representation.ActorSpawnRequestHandle, GetSpawnPriority(RepresentationLOD));
 				}
 			}
 			if (Actor != nullptr)
@@ -153,7 +154,7 @@ void UMassRepresentationProcessor::UpdateRepresentation(FMassExecutionContext& C
 						// If the low res is different than the high res, cancel any pending spawn request that is the opposite of what is needed.
 						if (Representation.LowResTemplateActorIndex != Representation.HighResTemplateActorIndex)
 						{
-							ReleaseActorOrCancelSpawning(MassAgent, ActorInfo, bHighResActor ? Representation.LowResTemplateActorIndex : Representation.HighResTemplateActorIndex, Representation.ActorSpawnRequestHandle, Context.Defer(), /*bCancelSpawningOnly*/true);
+							ReleaseActorOrCancelSpawning(*RepresentationSubsystem, MassAgent, ActorInfo, bHighResActor ? Representation.LowResTemplateActorIndex : Representation.HighResTemplateActorIndex, Representation.ActorSpawnRequestHandle, Context.Defer(), /*bCancelSpawningOnly*/true);
 							Actor = ActorInfo.GetOwnedByMassMutable();
 						}
 
@@ -165,7 +166,7 @@ void UMassRepresentationProcessor::UpdateRepresentation(FMassExecutionContext& C
 							!RepresentationSubsystem->DoesActorMatchTemplate(*Actor, WantedTemplateActorIndex) ||
 							Representation.ActorSpawnRequestHandle.IsValid())
 						{
-							NewActor = GetOrSpawnActor(MassAgent, ActorInfo, TransformFragment.GetTransform(), WantedTemplateActorIndex, Representation.ActorSpawnRequestHandle, GetSpawnPriority(RepresentationLOD));
+							NewActor = GetOrSpawnActor(*RepresentationSubsystem, MassAgent, ActorInfo, TransformFragment.GetTransform(), WantedTemplateActorIndex, Representation.ActorSpawnRequestHandle, GetSpawnPriority(RepresentationLOD));
 						}
 						else
 						{
@@ -208,8 +209,8 @@ void UMassRepresentationProcessor::UpdateRepresentation(FMassExecutionContext& C
 					if (!Actor || ActorInfo.IsOwnedByMass())
 					{
 						// Try releasing both, could have an high res spawned actor and a spawning request for a low res one
-						ReleaseActorOrCancelSpawning(MassAgent, ActorInfo, Representation.LowResTemplateActorIndex, Representation.ActorSpawnRequestHandle, Context.Defer());
-						ReleaseActorOrCancelSpawning(MassAgent, ActorInfo, Representation.HighResTemplateActorIndex, Representation.ActorSpawnRequestHandle, Context.Defer());
+						ReleaseActorOrCancelSpawning(*RepresentationSubsystem, MassAgent, ActorInfo, Representation.LowResTemplateActorIndex, Representation.ActorSpawnRequestHandle, Context.Defer());
+						ReleaseActorOrCancelSpawning(*RepresentationSubsystem, MassAgent, ActorInfo, Representation.HighResTemplateActorIndex, Representation.ActorSpawnRequestHandle, Context.Defer());
 					}
 					else
 					{
@@ -234,8 +235,6 @@ void UMassRepresentationProcessor::UpdateRepresentation(FMassExecutionContext& C
 
 void UMassRepresentationProcessor::Execute(UMassEntitySubsystem& InEntitySubsystem, FMassExecutionContext& Context)
 {
-	check(RepresentationSubsystem);
-
 	// Visualize entities
 	EntityQuery.ForEachEntityChunk(InEntitySubsystem, Context, [this](FMassExecutionContext& Context)
 	{
@@ -243,29 +242,26 @@ void UMassRepresentationProcessor::Execute(UMassEntitySubsystem& InEntitySubsyst
 	});
 }
 
-AActor* UMassRepresentationProcessor::GetOrSpawnActor(const FMassEntityHandle MassAgent, FDataFragment_Actor& ActorInfo, const FTransform& Transform, const int16 TemplateActorIndex, FMassActorSpawnRequestHandle& SpawnRequestHandle, const float Priority)
+AActor* UMassRepresentationProcessor::GetOrSpawnActor(UMassRepresentationSubsystem& RepresentationSubsystem, const FMassEntityHandle MassAgent, FDataFragment_Actor& ActorInfo, const FTransform& Transform, const int16 TemplateActorIndex, FMassActorSpawnRequestHandle& SpawnRequestHandle, const float Priority)
 {
-	check(RepresentationSubsystem);
-
-	return RepresentationSubsystem->GetOrSpawnActorFromTemplate(MassAgent, Transform, TemplateActorIndex, SpawnRequestHandle, Priority, 
+	return RepresentationSubsystem.GetOrSpawnActorFromTemplate(MassAgent, Transform, TemplateActorIndex, SpawnRequestHandle, Priority, 
 		FMassActorPreSpawnDelegate::CreateUObject(this, &UMassRepresentationProcessor::OnActorPreSpawn), 
 		FMassActorPostSpawnDelegate::CreateUObject(this, &UMassRepresentationProcessor::OnActorPostSpawn));
 }
 
-bool UMassRepresentationProcessor::ReleaseActorOrCancelSpawning(const FMassEntityHandle MassAgent, FDataFragment_Actor& ActorInfo, const int16 TemplateActorIndex, FMassActorSpawnRequestHandle& SpawnRequestHandle, FMassCommandBuffer& CommandBuffer, bool bCancelSpawningOnly /*= false*/)
+bool UMassRepresentationProcessor::ReleaseActorOrCancelSpawning(UMassRepresentationSubsystem& RepresentationSubsystem, const FMassEntityHandle MassAgent, FDataFragment_Actor& ActorInfo, const int16 TemplateActorIndex, FMassActorSpawnRequestHandle& SpawnRequestHandle, FMassCommandBuffer& CommandBuffer, bool bCancelSpawningOnly /*= false*/)
 {
-	check(RepresentationSubsystem);
 	check(!ActorInfo.IsValid() || ActorInfo.IsOwnedByMass());
 
 	AActor* Actor = ActorInfo.GetOwnedByMassMutable();
 	// note that it's fine for Actor to be null. That means the RepresentationSubsystem will try to stop 
 	// the spawning of whatever SpawnRequestHandle reference to
-	const bool bSuccess = bCancelSpawningOnly ? RepresentationSubsystem->CancelSpawning(MassAgent, TemplateActorIndex, SpawnRequestHandle) :
-			RepresentationSubsystem->ReleaseTemplateActorOrCancelSpawning(MassAgent, TemplateActorIndex, Actor, SpawnRequestHandle);
+	const bool bSuccess = bCancelSpawningOnly ? RepresentationSubsystem.CancelSpawning(MassAgent, TemplateActorIndex, SpawnRequestHandle) :
+			RepresentationSubsystem.ReleaseTemplateActorOrCancelSpawning(MassAgent, TemplateActorIndex, Actor, SpawnRequestHandle);
 	if (bSuccess)
 	{
 		Actor = ActorInfo.GetOwnedByMassMutable();
-		if (Actor && RepresentationSubsystem->DoesActorMatchTemplate(*Actor, TemplateActorIndex))
+		if (Actor && RepresentationSubsystem.DoesActorMatchTemplate(*Actor, TemplateActorIndex))
 		{
 			ActorInfo.ResetNoHandleMapUpdate();
 			
@@ -311,18 +307,14 @@ void UMassRepresentationProcessor::TeleportActor(const FTransform& Transform, AA
 	}
 }
 
-void UMassRepresentationProcessor::ReleaseAnyActorOrCancelAnySpawning(UMassRepresentationSubsystem& RepresentationSubsystem, const FMassEntityHandle MassAgent)
+void UMassRepresentationProcessor::ReleaseAnyActorOrCancelAnySpawning(UMassEntitySubsystem& EntitySubsystem, const FMassEntityHandle MassAgent)
 {
-	if (const UMassEntitySubsystem* EntitySubsystem = UWorld::GetSubsystem<UMassEntitySubsystem>(RepresentationSubsystem.GetWorld()))
-	{
-		if (FDataFragment_Actor* ActorInfo = EntitySubsystem->GetFragmentDataPtr<FDataFragment_Actor>(MassAgent))
-		{
-			if (FMassRepresentationFragment* Representation = EntitySubsystem->GetFragmentDataPtr<FMassRepresentationFragment>(MassAgent))
-			{
-				ReleaseAnyActorOrCancelAnySpawning(RepresentationSubsystem, MassAgent, *ActorInfo, *Representation);
-			}
-		}
-	}
+	FMassEntityView EntityView(EntitySubsystem,MassAgent);
+	FDataFragment_Actor& ActorInfo = EntityView.GetFragmentData<FDataFragment_Actor>();
+	FMassRepresentationFragment& Representation = EntityView.GetFragmentData<FMassRepresentationFragment>();
+	UMassRepresentationSubsystem* RepresentationSubsystem = EntityView.GetSharedFragmentData<FMassRepresentationSubsystemFragment>().RepresentationSubsystem;
+	check(RepresentationSubsystem);
+	ReleaseAnyActorOrCancelAnySpawning(*RepresentationSubsystem, MassAgent, ActorInfo, Representation);
 }
 
 void UMassRepresentationProcessor::ReleaseAnyActorOrCancelAnySpawning(UMassRepresentationSubsystem& RepresentationSubsystem, const FMassEntityHandle MassAgent, FDataFragment_Actor& ActorInfo, FMassRepresentationFragment& Representation)
@@ -342,6 +334,7 @@ void UMassRepresentationProcessor::ReleaseAnyActorOrCancelAnySpawning(UMassRepre
 	{
 		RepresentationSubsystem.ReleaseTemplateActorOrCancelSpawning(MassAgent, Representation.LowResTemplateActorIndex, Actor, Representation.ActorSpawnRequestHandle);
 	}
+	check(!Representation.ActorSpawnRequestHandle.IsValid());
 }
 
 void UMassRepresentationProcessor::UpdateVisualization(FMassExecutionContext& Context)
@@ -426,32 +419,29 @@ void UMassRepresentationProcessor::UpdateEntityVisibility(const FMassEntityHandl
 void UMassRepresentationProcessor::OnActorPreSpawn(const FMassActorSpawnRequestHandle& SpawnRequestHandle, const FStructView& SpawnRequest)
 {
 	check(CachedEntitySubsystem);
-	check(RepresentationSubsystem);
 
 	const FMassActorSpawnRequest& MassActorSpawnRequest = SpawnRequest.Get<FMassActorSpawnRequest>();
-	const FMassEntityView View(*CachedEntitySubsystem, MassActorSpawnRequest.MassAgent);
+	const FMassEntityView EntityView(*CachedEntitySubsystem, MassActorSpawnRequest.MassAgent);
+	FDataFragment_Actor& ActorInfo = EntityView.GetFragmentData<FDataFragment_Actor>();
+	FMassRepresentationFragment& Representation = EntityView.GetFragmentData<FMassRepresentationFragment>();
+	UMassRepresentationSubsystem* RepresentationSubsystem = EntityView.GetSharedFragmentData<FMassRepresentationSubsystemFragment>().RepresentationSubsystem;
+	check(RepresentationSubsystem);
 
-	if (FDataFragment_Actor* ActorInfo = View.GetFragmentDataPtr<FDataFragment_Actor>())
+	// Release any existing actor
+	if (AActor* Actor = ActorInfo.GetMutable())
 	{
-		// Release any existing actor
-		if (AActor* Actor = ActorInfo->GetMutable())
-		{
-			checkf(ActorInfo->IsOwnedByMass(), TEXT("If we reach here, we expect the actor to be owned by mass, otherwise we should not be spawning a new one one top of this one."));
+		checkf(ActorInfo.IsOwnedByMass(), TEXT("If we reach here, we expect the actor to be owned by mass, otherwise we should not be spawning a new one one top of this one."));
 
-			// WARNING!
-			// Need to reset before ReleaseTemplateActor as this action might move the entity to a new archetype and
-			// so the Fragment passed in parameters would not be valid anymore.
-			ActorInfo->ResetAndUpdateHandleMap();
+		// WARNING!
+		// Need to reset before ReleaseTemplateActor as this action might move the entity to a new archetype and
+		// so the Fragment passed in parameters would not be valid anymore.
+		ActorInfo.ResetAndUpdateHandleMap();
 
-			if (const FMassRepresentationFragment* Representation = View.GetFragmentDataPtr<FMassRepresentationFragment>())
+		if (!RepresentationSubsystem->ReleaseTemplateActor(MassActorSpawnRequest.MassAgent, Representation.HighResTemplateActorIndex, Actor, /*bImmediate*/ true))
+		{
+			if (!RepresentationSubsystem->ReleaseTemplateActor(MassActorSpawnRequest.MassAgent, Representation.LowResTemplateActorIndex, Actor, /*bImmediate*/ true))
 			{
-				if (!RepresentationSubsystem->ReleaseTemplateActor(MassActorSpawnRequest.MassAgent, Representation->HighResTemplateActorIndex, Actor, /*bImmediate*/ true))
-				{
-					if (!RepresentationSubsystem->ReleaseTemplateActor(MassActorSpawnRequest.MassAgent, Representation->LowResTemplateActorIndex, Actor, /*bImmediate*/ true))
-					{
-						checkf(false, TEXT("Expecting to be able to release spawned actor either the high res or low res one"));
-					}
-				}
+				checkf(false, TEXT("Expecting to be able to release spawned actor either the high res or low res one"));
 			}
 		}
 	}
@@ -492,27 +482,20 @@ UMassRepresentationFragmentDestructor::UMassRepresentationFragmentDestructor()
 	ExecutionFlags = (int32)(EProcessorExecutionFlags::All);
 }
 
-void UMassRepresentationFragmentDestructor::Initialize(UObject& Owner)
-{
-	Super::Initialize(Owner);
-
-	RepresentationSubsystem = UWorld::GetSubsystem<UMassRepresentationSubsystem>(Owner.GetWorld());
-	check(RepresentationSubsystem);
-}
-
 void UMassRepresentationFragmentDestructor::ConfigureQueries()
 {
 	EntityQuery.AddRequirement<FMassRepresentationFragment>(EMassFragmentAccess::ReadOnly);
 	EntityQuery.AddRequirement<FDataFragment_Actor>(EMassFragmentAccess::ReadWrite);
-	EntityQuery.AddTagRequirement<FMassRepresentationDefaultDestructorTag>(EMassFragmentPresence::All);
+	EntityQuery.AddSharedRequirement<FMassRepresentationSubsystemFragment>(EMassFragmentAccess::ReadWrite);
 }
 
 void UMassRepresentationFragmentDestructor::Execute(UMassEntitySubsystem& EntitySubsystem, FMassExecutionContext& Context)
 {
-	check(RepresentationSubsystem);
-
 	EntityQuery.ForEachEntityChunk(EntitySubsystem, Context, [this](FMassExecutionContext& Context)
 	{
+		UMassRepresentationSubsystem* RepresentationSubsystem = Context.GetMutableSharedFragment<FMassRepresentationSubsystemFragment>().RepresentationSubsystem;
+		check(RepresentationSubsystem);
+
 		const TArrayView<FMassRepresentationFragment> RepresentationList = Context.GetMutableFragmentView<FMassRepresentationFragment>();
 		const TArrayView<FDataFragment_Actor> ActorList = Context.GetMutableFragmentView<FDataFragment_Actor>();
 
diff --git a/Engine/Plugins/Runtime/MassGameplay/Source/MassRepresentation/Private/MassRepresentationSubsystem.cpp b/Engine/Plugins/Runtime/MassGameplay/Source/MassRepresentation/Private/MassRepresentationSubsystem.cpp
--- a/Engine/Plugins/Runtime/MassGameplay/Source/MassRepresentation/Private/MassRepresentationSubsystem.cpp
+++ b/Engine/Plugins/Runtime/MassGameplay/Source/MassRepresentation/Private/MassRepresentationSubsystem.cpp
@@ -18,20 +18,20 @@
 #include "WorldPartition/WorldPartitionSubsystem.h"
 #include "WorldPartition/WorldPartitionRuntimeCell.h"
 
-int16 UMassRepresentationSubsystem::FindOrAddStaticMeshDesc(const FStaticMeshInstanceVisualizationDesc& Desc) const
+int16 UMassRepresentationSubsystem::FindOrAddStaticMeshDesc(const FStaticMeshInstanceVisualizationDesc& Desc)
 {
 	check(VisualizationComponent);
 	return VisualizationComponent->FindOrAddVisualDesc(Desc);
 }
 
-FMassInstancedStaticMeshInfoArrayView UMassRepresentationSubsystem::GetMutableInstancedStaticMeshInfos() const
+FMassInstancedStaticMeshInfoArrayView UMassRepresentationSubsystem::GetMutableInstancedStaticMeshInfos()
 {
 	check(VisualizationComponent);
 	FMassInstancedStaticMeshInfoArrayView View = VisualizationComponent->GetMutableVisualInfos();
 	return MoveTemp(View);
 }
 
-void UMassRepresentationSubsystem::DirtyStaticMeshInstances() const
+void UMassRepresentationSubsystem::DirtyStaticMeshInstances()
 {
 	check(VisualizationComponent);
 	return VisualizationComponent->DirtyVisuals();
@@ -280,7 +280,7 @@ bool UMassRepresentationSubsystem::CancelSpawningInternal(const int16 TemplateAc
 	return true;
 }
 
-bool UMassRepresentationSubsystem::DoesActorMatchTemplate(const AActor& Actor, const int16 TemplateActorIndex)
+bool UMassRepresentationSubsystem::DoesActorMatchTemplate(const AActor& Actor, const int16 TemplateActorIndex) const
 {
 	UE_MT_SCOPED_READ_ACCESS(TemplateActorsMTAccessDetector);
 	if (!TemplateActors.IsValidIndex(TemplateActorIndex))
@@ -379,7 +379,7 @@ void UMassRepresentationSubsystem::OnMassAgentComponentEntityAssociated(const UM
 		// Check if this mass agent already handled by this sub system, if yes than release any local spawned actor or cancel any spawn requests
 		if (HandledMassAgents.Find(MassAgent))
 		{
-			UMassRepresentationProcessor::ReleaseAnyActorOrCancelAnySpawning(*this, MassAgent);
+			UMassRepresentationProcessor::ReleaseAnyActorOrCancelAnySpawning(*EntitySubsystem, MassAgent);
 		}
 	}
 }
diff --git a/Engine/Plugins/Runtime/MassGameplay/Source/MassRepresentation/Private/MassUpdateISMProcessor.cpp b/Engine/Plugins/Runtime/MassGameplay/Source/MassRepresentation/Private/MassUpdateISMProcessor.cpp
--- a/Engine/Plugins/Runtime/MassGameplay/Source/MassRepresentation/Private/MassUpdateISMProcessor.cpp
+++ b/Engine/Plugins/Runtime/MassGameplay/Source/MassRepresentation/Private/MassUpdateISMProcessor.cpp
@@ -16,65 +16,43 @@ UMassUpdateISMProcessor::UMassUpdateISMProcessor()
 	ExecutionOrder.ExecuteAfter.Add(UE::Mass::ProcessorGroupNames::Representation);
 }
 
-void UMassUpdateISMProcessor::Initialize(UObject& Owner)
-{
-	Super::Initialize(Owner);
-
-	if (const UWorld* World = Owner.GetWorld())
-	{
-		for (FMassUpdateISMConfig& Config : ISMConfigs)
-		{
-			Config.RepresentationSubsystem = Cast<UMassRepresentationSubsystem>(World->GetSubsystemBase(Config.RepresentationSubsystemClass ? Config.RepresentationSubsystemClass : TSubclassOf<UMassRepresentationSubsystem>(UMassRepresentationSubsystem::StaticClass())));
-		}
-	}
-}
-
 void UMassUpdateISMProcessor::ConfigureQueries()
 {
-	for (FMassUpdateISMConfig& Config : ISMConfigs)
-	{
-		if (Config.TagFilter.GetScriptStruct())
-		{
-			Config.EntityQuery.AddTagRequirement(*Config.TagFilter.GetScriptStruct(), EMassFragmentPresence::All);
-		}
-		Config.EntityQuery.AddRequirement<FDataFragment_Transform>(EMassFragmentAccess::ReadOnly);
-		Config.EntityQuery.AddRequirement<FMassRepresentationFragment>(EMassFragmentAccess::ReadWrite);
-		Config.EntityQuery.AddRequirement<FMassRepresentationLODFragment>(EMassFragmentAccess::ReadWrite);
-		Config.EntityQuery.AddChunkRequirement<FMassVisualizationChunkFragment>(EMassFragmentAccess::ReadWrite);
-		Config.EntityQuery.SetChunkFilter(&FMassVisualizationChunkFragment::AreAnyEntitiesVisibleInChunk);
-	}
+	EntityQuery.AddRequirement<FDataFragment_Transform>(EMassFragmentAccess::ReadOnly);
+	EntityQuery.AddRequirement<FMassRepresentationFragment>(EMassFragmentAccess::ReadWrite);
+	EntityQuery.AddRequirement<FMassRepresentationLODFragment>(EMassFragmentAccess::ReadWrite);
+	EntityQuery.AddChunkRequirement<FMassVisualizationChunkFragment>(EMassFragmentAccess::ReadWrite);
+	EntityQuery.SetChunkFilter(&FMassVisualizationChunkFragment::AreAnyEntitiesVisibleInChunk);
+	EntityQuery.AddSharedRequirement<FMassRepresentationSubsystemFragment>(EMassFragmentAccess::ReadWrite);
 }
 
 void UMassUpdateISMProcessor::Execute(UMassEntitySubsystem& EntitySubsystem, FMassExecutionContext& Context)
 {
-
-	for (FMassUpdateISMConfig& Config : ISMConfigs)
+	EntityQuery.ForEachEntityChunk(EntitySubsystem, Context, [](FMassExecutionContext& Context)
 	{
-		check(Config.RepresentationSubsystem);
-		FMassInstancedStaticMeshInfoArrayView ISMInfo = Config.RepresentationSubsystem->GetMutableInstancedStaticMeshInfos();
+		UMassRepresentationSubsystem* RepresentationSubsystem = Context.GetSharedFragment<FMassRepresentationSubsystemFragment>().RepresentationSubsystem;
+		check(RepresentationSubsystem);
+		FMassInstancedStaticMeshInfoArrayView ISMInfo = RepresentationSubsystem->GetMutableInstancedStaticMeshInfos();
 
-		Config.EntityQuery.ForEachEntityChunk(EntitySubsystem, Context, [&ISMInfo](FMassExecutionContext& Context)
+		const TConstArrayView<FDataFragment_Transform> TransformList = Context.GetFragmentView<FDataFragment_Transform>();
+		const TArrayView<FMassRepresentationFragment> RepresentationList = Context.GetMutableFragmentView<FMassRepresentationFragment>();
+		const TConstArrayView<FMassRepresentationLODFragment> RepresentationLODList = Context.GetFragmentView<FMassRepresentationLODFragment>();
+
+		const int32 NumEntities = Context.GetNumEntities();
+		for (int32 EntityIdx = 0; EntityIdx < NumEntities; EntityIdx++)
 		{
-			const TConstArrayView<FDataFragment_Transform> TransformList = Context.GetFragmentView<FDataFragment_Transform>();
-			const TArrayView<FMassRepresentationFragment> RepresentationList = Context.GetMutableFragmentView<FMassRepresentationFragment>();
-			const TConstArrayView<FMassRepresentationLODFragment> RepresentationLODList = Context.GetFragmentView<FMassRepresentationLODFragment>();
+			const FDataFragment_Transform& TransformFragment = TransformList[EntityIdx];
+			const FMassRepresentationLODFragment& RepresentationLOD = RepresentationLODList[EntityIdx];
+			FMassRepresentationFragment& Representation = RepresentationList[EntityIdx];
 
-			const int32 NumEntities = Context.GetNumEntities();
-			for (int32 EntityIdx = 0; EntityIdx < NumEntities; EntityIdx++)
+			if (Representation.CurrentRepresentation == ERepresentationType::StaticMeshInstance)
 			{
-				const FDataFragment_Transform& TransformFragment = TransformList[EntityIdx];
-				const FMassRepresentationLODFragment& RepresentationLOD = RepresentationLODList[EntityIdx];
-				FMassRepresentationFragment& Representation = RepresentationList[EntityIdx];
-
-				if (Representation.CurrentRepresentation == ERepresentationType::StaticMeshInstance)
-				{
-					UpdateISMTransform(GetTypeHash(Context.GetEntity(EntityIdx)), ISMInfo[Representation.StaticMeshDescIndex], TransformFragment.GetTransform(), Representation.PrevTransform, RepresentationLOD.LODSignificance, Representation.PrevLODSignificance);
-				}
-				Representation.PrevTransform = TransformFragment.GetTransform();
-				Representation.PrevLODSignificance = RepresentationLOD.LODSignificance;
+				UpdateISMTransform(GetTypeHash(Context.GetEntity(EntityIdx)), ISMInfo[Representation.StaticMeshDescIndex], TransformFragment.GetTransform(), Representation.PrevTransform, RepresentationLOD.LODSignificance, Representation.PrevLODSignificance);
 			}
-		});
-	}
+			Representation.PrevTransform = TransformFragment.GetTransform();
+			Representation.PrevLODSignificance = RepresentationLOD.LODSignificance;
+		}
+	});
 }
 
 void UMassUpdateISMProcessor::UpdateISMTransform(int32 EntityId, FMassInstancedStaticMeshInfo& ISMInfo, const FTransform& Transform, const FTransform& PrevTransform, const float LODSignificance, const float PrevLODSignificance/* = -1.0f*/)
diff --git a/Engine/Plugins/Runtime/MassGameplay/Source/MassRepresentation/Private/MassVisualizationTrait.cpp b/Engine/Plugins/Runtime/MassGameplay/Source/MassRepresentation/Private/MassVisualizationTrait.cpp
--- a/Engine/Plugins/Runtime/MassGameplay/Source/MassRepresentation/Private/MassVisualizationTrait.cpp
+++ b/Engine/Plugins/Runtime/MassGameplay/Source/MassRepresentation/Private/MassVisualizationTrait.cpp
@@ -11,7 +11,6 @@
 UMassVisualizationTrait::UMassVisualizationTrait()
 {
 	RepresentationSubsystemClass = UMassRepresentationSubsystem::StaticClass();
-	RepresentationDestructorTag = FMassRepresentationDefaultDestructorTag::StaticStruct();
 }
 
 void UMassVisualizationTrait::BuildTemplate(FMassEntityTemplateBuildContext& BuildContext, UWorld& World) const
@@ -28,23 +27,26 @@ void UMassVisualizationTrait::BuildTemplate(FMassEntityTemplateBuildContext& Bui
 	BuildContext.AddTag<FMassCollectLODViewerInfoTag>(); // Depends on FMassViewerInfoFragment
 	BuildContext.AddFragment<FDataFragment_Transform>();
 
-	if (UMassRepresentationSubsystem* RepresentationSubsystem = Cast<UMassRepresentationSubsystem>(World.GetSubsystemBase(RepresentationSubsystemClass)))
-	{
-		FMassRepresentationFragment& RepresentationFragment = BuildContext.AddFragment_GetRef<FMassRepresentationFragment>();
-		RepresentationFragment.StaticMeshDescIndex = RepresentationSubsystem->FindOrAddStaticMeshDesc(StaticMeshInstanceDesc);
-		RepresentationFragment.HighResTemplateActorIndex = HighResTemplateActor.Get() ? RepresentationSubsystem->FindOrAddTemplateActor(HighResTemplateActor.Get()) : INDEX_NONE;
-		RepresentationFragment.LowResTemplateActorIndex = LowResTemplateActor.Get() ? RepresentationSubsystem->FindOrAddTemplateActor(LowResTemplateActor.Get()) : INDEX_NONE;
+	UMassEntitySubsystem* EntitySubsystem = UWorld::GetSubsystem<UMassEntitySubsystem>(&World);
+	check(EntitySubsystem);
 
-		if (ensureMsgf(RepresentationDestructorTag, TEXT("RepresentationDestructorTag is never expected to be empty")))
-		{
-			BuildContext.AddTag(*RepresentationDestructorTag);
-		}
+	UMassRepresentationSubsystem* RepresentationSubsystem = Cast<UMassRepresentationSubsystem>(World.GetSubsystemBase(RepresentationSubsystemClass));
+	check(RepresentationSubsystem);
 
-		BuildContext.AddFragment<FMassRepresentationLODFragment>();
-		BuildContext.AddTag<FMassVisibilityCulledByDistanceTag>();
+	FMassRepresentationSubsystemFragment Subsystem;
+	Subsystem.RepresentationSubsystem = RepresentationSubsystem;
+	uint32 SubsystemHash = UE::StructUtils::GetStructCrc32(FConstStructView::Make(Subsystem));
+	FSharedStruct SubsystemFragment = EntitySubsystem->GetOrCreateSharedFragment<FMassRepresentationSubsystemFragment>(SubsystemHash, Subsystem);
+	BuildContext.AddSharedFragment(SubsystemFragment);
 
-		BuildContext.AddChunkFragment<FMassVisualizationChunkFragment>();
-	}
+	FMassRepresentationFragment& RepresentationFragment = BuildContext.AddFragment_GetRef<FMassRepresentationFragment>();
+	RepresentationFragment.StaticMeshDescIndex = RepresentationSubsystem->FindOrAddStaticMeshDesc(StaticMeshInstanceDesc);
+	RepresentationFragment.HighResTemplateActorIndex = HighResTemplateActor.Get() ? RepresentationSubsystem->FindOrAddTemplateActor(HighResTemplateActor.Get()) : INDEX_NONE;
+	RepresentationFragment.LowResTemplateActorIndex = LowResTemplateActor.Get() ? RepresentationSubsystem->FindOrAddTemplateActor(LowResTemplateActor.Get()) : INDEX_NONE;
+
+	BuildContext.AddFragment<FMassRepresentationLODFragment>();
+	BuildContext.AddTag<FMassVisibilityCulledByDistanceTag>();
+	BuildContext.AddChunkFragment<FMassVisualizationChunkFragment>();
 }
 
 
diff --git a/Engine/Plugins/Runtime/MassGameplay/Source/MassRepresentation/Public/MassRepresentationFragments.h b/Engine/Plugins/Runtime/MassGameplay/Source/MassRepresentation/Public/MassRepresentationFragments.h
--- a/Engine/Plugins/Runtime/MassGameplay/Source/MassRepresentation/Public/MassRepresentationFragments.h
+++ b/Engine/Plugins/Runtime/MassGameplay/Source/MassRepresentation/Public/MassRepresentationFragments.h
@@ -10,6 +10,8 @@
 
 #include "MassRepresentationFragments.generated.h"
 
+class UMassRepresentationSubsystem;
+
 USTRUCT()
 struct MASSREPRESENTATION_API FMassRepresentationLODFragment : public FMassFragment
 {
@@ -48,4 +50,13 @@ struct MASSREPRESENTATION_API FMassRepresentationFragment : public FMassFragment
 
 	/** Value scaling from 0 to 3, 0 highest LOD we support and 3 being completely off LOD */
 	float PrevLODSignificance = -1.0f;
-};
\ No newline at end of file
+};
+
+USTRUCT()
+struct FMassRepresentationSubsystemFragment : public FMassSharedFragment
+{
+	GENERATED_BODY()
+
+	UPROPERTY(Transient)
+	UMassRepresentationSubsystem* RepresentationSubsystem = nullptr;
+};
diff --git a/Engine/Plugins/Runtime/MassGameplay/Source/MassRepresentation/Public/MassRepresentationProcessor.h b/Engine/Plugins/Runtime/MassGameplay/Source/MassRepresentation/Public/MassRepresentationProcessor.h
--- a/Engine/Plugins/Runtime/MassGameplay/Source/MassRepresentation/Public/MassRepresentationProcessor.h
+++ b/Engine/Plugins/Runtime/MassGameplay/Source/MassRepresentation/Public/MassRepresentationProcessor.h
@@ -46,6 +46,7 @@ protected:
 
 	/** 
 	 * Returns an actor of the template type and setup fragments values from it
+	 * @param RepresentationSubsystem to use to get or spawn the actor
 	 * @param MassAgent is the handle to the associated mass agent
 	 * @param ActorInfo is the fragment where we are going to store the actor pointer
 	 * @param Transform is the spatial information about where to spawn the actor 
@@ -54,10 +55,11 @@ protected:
 	 * @param Priority of this spawn request in comparison with the others, lower value means higher priority
 	 * @return the actor spawned
 	 */
-	virtual AActor* GetOrSpawnActor(const FMassEntityHandle MassAgent, FDataFragment_Actor& ActorInfo, const FTransform& Transform, const int16 TemplateActorIndex, FMassActorSpawnRequestHandle& SpawnRequestHandle, const float Priority);
+	virtual AActor* GetOrSpawnActor(UMassRepresentationSubsystem& RepresentationSubsystem, const FMassEntityHandle MassAgent, FDataFragment_Actor& ActorInfo, const FTransform& Transform, const int16 TemplateActorIndex, FMassActorSpawnRequestHandle& SpawnRequestHandle, const float Priority);
 
 	/**
 	 * Release the actor to the subsystem, will only release it the actor or spawn request matches the template actor
+	 * @param RepresentationSubsystem to use to release the actor or cancel the spawning
 	 * @param MassAgent is the handle to the associated mass agent
 	 * @param ActorInfo is the fragment where we are going to store the actor pointer
 	 * @param TemplateActorIndex is the index of the type to release
@@ -66,7 +68,7 @@ protected:
 	 * @param bCancelSpawningOnly tell to only cancel the existing spawning request and to not release the associated actor it any.
 	 * @return if the actor was release or the spawning was canceled.
 	 */
-	virtual bool ReleaseActorOrCancelSpawning(const FMassEntityHandle MassAgent, FDataFragment_Actor& ActorInfo, const int16 TemplateActorIndex, FMassActorSpawnRequestHandle& SpawnRequestHandle, FMassCommandBuffer& CommandBuffer, const bool bCancelSpawningOnly = false);
+	virtual bool ReleaseActorOrCancelSpawning(UMassRepresentationSubsystem& RepresentationSubsystem, const FMassEntityHandle MassAgent, FDataFragment_Actor& ActorInfo, const int16 TemplateActorIndex, FMassActorSpawnRequestHandle& SpawnRequestHandle, FMassCommandBuffer& CommandBuffer, const bool bCancelSpawningOnly = false);
 
 	/** 
 	 * Enable/disable a spawned actor
@@ -130,11 +132,11 @@ public:
 	 * Release an actor or cancel its spawning (calls ReleaseAnyActorOrCancelAnySpawning)
 	 * WARNING: This method will destroy the associated actor in any and by the same fact might also move the entity into a new archetype.
 	 *          So any reference to fragment might become invalid.
-	 * @param RepresentationSubsystem to use to release any actors or cancel spawning requests
+	 * @param EntitySubsystem to use to retrieve the mass agent fragments
 	 * @param MassAgent is the handle to the associated mass agent
 	 * @return True if actor was release or spawning request was canceled
 	 */
-	static void ReleaseAnyActorOrCancelAnySpawning(UMassRepresentationSubsystem& RepresentationSubsystem, const FMassEntityHandle MassAgent);
+	static void ReleaseAnyActorOrCancelAnySpawning(UMassEntitySubsystem& EntitySubsystem, const FMassEntityHandle MassAgent);
 
 	/**
 	 * Release an actor or cancel its spawning
@@ -194,10 +196,6 @@ protected:
 	UPROPERTY(Transient)
 	UMassEntitySubsystem* CachedEntitySubsystem;
 
-	/** A cache pointer to the representation subsystem */
-	UPROPERTY(Transient)
-	UMassRepresentationSubsystem* RepresentationSubsystem;
-
 	/** Default representation when unable to spawn an actor */
 	ERepresentationType DefaultRepresentationType = ERepresentationType::None;
 
@@ -208,13 +206,6 @@ protected:
 	float NotVisibleUpdateRate = 0.5f;
 };
 
-
-USTRUCT()
-struct FMassRepresentationDefaultDestructorTag : public FMassTag
-{
-	GENERATED_BODY()
-};
-
 UCLASS()
 class MASSREPRESENTATION_API UMassRepresentationFragmentDestructor : public UMassFragmentDeinitializer
 {
@@ -224,13 +215,8 @@ public:
 	UMassRepresentationFragmentDestructor();
 
 protected:
-	virtual void Initialize(UObject& Owner) override;
 	virtual void ConfigureQueries() override;
 	virtual void Execute(UMassEntitySubsystem& EntitySubsystem, FMassExecutionContext& Context) override;
 
-	/** A cache pointer to the representation subsystem */
-	UPROPERTY(Transient)
-	UMassRepresentationSubsystem* RepresentationSubsystem;
-
 	FMassEntityQuery EntityQuery;
 };
\ No newline at end of file
diff --git a/Engine/Plugins/Runtime/MassGameplay/Source/MassRepresentation/Public/MassRepresentationSubsystem.h b/Engine/Plugins/Runtime/MassGameplay/Source/MassRepresentation/Public/MassRepresentationSubsystem.h
--- a/Engine/Plugins/Runtime/MassGameplay/Source/MassRepresentation/Public/MassRepresentationSubsystem.h
+++ b/Engine/Plugins/Runtime/MassGameplay/Source/MassRepresentation/Public/MassRepresentationSubsystem.h
@@ -36,17 +36,17 @@ public:
 	 * @param Desc is the information for the static mesh that will be instantiated later via AddStaticMeshInstance()
 	 * @return The index of the static mesh type 
 	 */
-	int16 FindOrAddStaticMeshDesc(const FStaticMeshInstanceVisualizationDesc& Desc) const;
+	int16 FindOrAddStaticMeshDesc(const FStaticMeshInstanceVisualizationDesc& Desc);
 
 	/** @todo: need to add a release API at some point for static mesh types */
 
 	/** 
 	 * @return the array of all the static mesh instance component information
 	 */
-	FMassInstancedStaticMeshInfoArrayView GetMutableInstancedStaticMeshInfos() const;
+	FMassInstancedStaticMeshInfoArrayView GetMutableInstancedStaticMeshInfos();
 
 	/** Mark render state of the static mesh instances dirty */
-	void DirtyStaticMeshInstances() const;
+	void DirtyStaticMeshInstances();
 
 	/** 
 	 * Store the template actor uniquely and return an index to it 
@@ -105,7 +105,7 @@ public:
 	 * @param TemplateActorIndex is the template type of the actor to compare against
 	 * @return True if actor matches the template
 	 */
-	bool DoesActorMatchTemplate(const AActor& Actor, const int16 TemplateActorIndex);
+	bool DoesActorMatchTemplate(const AActor& Actor, const int16 TemplateActorIndex) const;
 
 	TSubclassOf<AActor> GetTemplateActorClass(const int16 TemplateActorIndex);
 
diff --git a/Engine/Plugins/Runtime/MassGameplay/Source/MassRepresentation/Public/MassUpdateISMProcessor.h b/Engine/Plugins/Runtime/MassGameplay/Source/MassRepresentation/Public/MassUpdateISMProcessor.h
--- a/Engine/Plugins/Runtime/MassGameplay/Source/MassRepresentation/Public/MassUpdateISMProcessor.h
+++ b/Engine/Plugins/Runtime/MassGameplay/Source/MassRepresentation/Public/MassUpdateISMProcessor.h
@@ -10,26 +10,6 @@ class UMassRepresentationSubsystem;
 
 struct FMassInstancedStaticMeshInfo;
 
-USTRUCT()
-struct FMassUpdateISMConfig
-{
-	GENERATED_BODY()
-
-	/** Component Tag that will be used to associate Update ISM config */
-	UPROPERTY(EditAnywhere, Category = "Mass|ISM", config, meta = (BaseStruct = "MassTag"))
-	FInstancedStruct TagFilter;
-
-	/** To specify a different MassRepresentationSubsystem, otherwise leave it null to use the default one  */
-	UPROPERTY(EditAnywhere, Category = "Mass|ISM", config)
-	TSubclassOf<UMassRepresentationSubsystem> RepresentationSubsystemClass;
-
-	/** A cached pointer to the representation subsystem */
-	UPROPERTY(Transient)
-	UMassRepresentationSubsystem* RepresentationSubsystem = nullptr;
-
-	FMassEntityQuery EntityQuery;
-};
-
 UCLASS()
 class MASSREPRESENTATION_API UMassUpdateISMProcessor : public UMassProcessor
 {
@@ -41,10 +21,6 @@ public:
 	static void UpdateISMTransform(int32 EntityId, FMassInstancedStaticMeshInfo& ISMInfo, const FTransform& Transform, const FTransform& PrevTransform, const float LODSignificance, const float PrevLODSignificance = -1.0f);
 
 protected:
-	/**
-	 * Initialize the processor
-	 * @param Owner of the Processor */
-	virtual void Initialize(UObject& Owner) override;
 
 	/** Configure the owned FMassEntityQuery instances to express processor's requirements */
 	virtual void ConfigureQueries() override;
@@ -55,8 +31,5 @@ protected:
 	 * @param Context is the execution context to be passed when executing the lambdas */
 	virtual void Execute(UMassEntitySubsystem& EntitySubsystem, FMassExecutionContext& Context) override;
 
-	/** Config that will be used to update ISM */
-	UPROPERTY(EditAnywhere, Category = "Mass|LOD", config, meta = (BaseStruct = "MassTag"))
-	TArray<FMassUpdateISMConfig> ISMConfigs;
-
+	FMassEntityQuery EntityQuery;
 };
\ No newline at end of file
diff --git a/Engine/Plugins/Runtime/MassGameplay/Source/MassRepresentation/Public/MassVisualizationTrait.h b/Engine/Plugins/Runtime/MassGameplay/Source/MassRepresentation/Public/MassVisualizationTrait.h
--- a/Engine/Plugins/Runtime/MassGameplay/Source/MassRepresentation/Public/MassVisualizationTrait.h
+++ b/Engine/Plugins/Runtime/MassGameplay/Source/MassRepresentation/Public/MassVisualizationTrait.h
@@ -40,7 +40,4 @@ protected:
 	// @todo the following property will be cut once new de/initializers are in
 	/** Allow subclasses to override the representation fragment deinitializer */
 	TSubclassOf<UMassProcessor> RepresentationFragmentDeinitializerClass;
-
-	UPROPERTY()
-	UScriptStruct* RepresentationDestructorTag;
 };
