[SmartObjectZoneGraphAnnotations] added DebugDraw for lane entry points to object locations.
Fixed bug when component was not registered again after editing a property (PostEditChangeChainProperty was not called on the base class preventing the ComponentReregisterContext from being processed)

diff --git a/Engine/Plugins/Runtime/MassGameplay/Source/MassSmartObjects/Private/SmartObjectZoneAnnotations.cpp b/Engine/Plugins/Runtime/MassGameplay/Source/MassSmartObjects/Private/SmartObjectZoneAnnotations.cpp
--- a/Engine/Plugins/Runtime/MassGameplay/Source/MassSmartObjects/Private/SmartObjectZoneAnnotations.cpp
+++ b/Engine/Plugins/Runtime/MassGameplay/Source/MassSmartObjects/Private/SmartObjectZoneAnnotations.cpp
@@ -15,9 +15,9 @@
 
 void USmartObjectZoneAnnotations::PostSubsystemsInitialized()
 {
-#if WITH_EDITOR
 	SmartObjectSubsystem = UWorld::GetSubsystem<USmartObjectSubsystem>(GetWorld());
 
+#if WITH_EDITOR
 	// Monitor collection changes to rebuild our lookup data
 	if (SmartObjectSubsystem != nullptr)
 	{
@@ -144,6 +144,10 @@ void USmartObjectZoneAnnotations::TickAnnotation(const float DeltaTime, FZoneGra
 
 		Data.bInitialTaggingCompleted = true;
 	}
+
+#if !UE_BUILD_SHIPPING && !UE_BUILD_TEST
+	MarkRenderStateDirty();
+#endif
 }
 
 #if WITH_EDITOR
@@ -160,6 +164,8 @@ void USmartObjectZoneAnnotations::OnUnregister()
 
 void USmartObjectZoneAnnotations::PostEditChangeChainProperty(FPropertyChangedChainEvent& PropertyChangedEvent)
 {
+	Super::PostEditChangeChainProperty(PropertyChangedEvent);
+
 	FProperty* Property = PropertyChangedEvent.Property;
 	FProperty* MemberProperty = nullptr;
 	if (PropertyChangedEvent.PropertyChain.GetActiveMemberNode())
@@ -176,6 +182,51 @@ void USmartObjectZoneAnnotations::PostEditChangeChainProperty(FPropertyChangedCh
 	}
 }
 
+#if !UE_BUILD_SHIPPING && !UE_BUILD_TEST
+void USmartObjectZoneAnnotations::DebugDraw(FZoneGraphAnnotationSceneProxy* DebugProxy)
+{
+	UZoneGraphSubsystem* ZoneGraph = UWorld::GetSubsystem<UZoneGraphSubsystem>(GetWorld());
+	if (ZoneGraph == nullptr)
+	{
+		return;
+	}
+
+	for (FSmartObjectAnnotationData& Data : SmartObjectAnnotationDataArray)
+	{
+		const FZoneGraphStorage* ZoneStorage = Data.DataHandle.IsValid() ? ZoneGraph->GetZoneGraphStorage(Data.DataHandle) : nullptr;
+		if (ZoneStorage == nullptr)
+		{
+			continue;
+		}
+
+		const ASmartObjectCollection* Collection = SmartObjectSubsystem->GetMainCollection();
+		if (Collection == nullptr)
+		{
+			return;
+		}
+
+		for (const FSmartObjectCollectionEntry& Entry : Collection->GetEntries())
+		{
+			FSmartObjectHandle Handle = Entry.GetHandle();
+			const FSmartObjectLaneLocation* SOLaneLocation = Data.ObjectToEntryPointLookup.Find(Handle);
+			if (SOLaneLocation == nullptr)
+			{
+				continue;
+			}
+
+			const FVector& ObjectLocation = Entry.GetComponent()->GetComponentLocation();
+			FZoneGraphLaneLocation EntryPointLocation;
+            UE::ZoneGraph::Query::CalculateLocationAlongLane(*ZoneStorage, SOLaneLocation->LaneIndex, SOLaneLocation->DistanceAlongLane, EntryPointLocation);
+			const FColor Color = FColor::Silver;
+			constexpr float SphereRadius = 25.f;
+            DebugProxy->Spheres.Emplace(SphereRadius, EntryPointLocation.Position, Color);
+			DebugProxy->Spheres.Emplace(SphereRadius, ObjectLocation, Color);
+			DebugProxy->DashedLines.Emplace(ObjectLocation, EntryPointLocation.Position, Color, /*dash size*/10.f);
+		}
+	}
+}
+#endif // !UE_BUILD_SHIPPING && !UE_BUILD_TEST
+
 void USmartObjectZoneAnnotations::RebuildForSingleGraph(FSmartObjectAnnotationData& Data, const FZoneGraphStorage& Storage)
 {
 	TRACE_CPUPROFILER_EVENT_SCOPE_STR("ZoneGraphSmartObjectBehavior RebuildData")
diff --git a/Engine/Plugins/Runtime/MassGameplay/Source/MassSmartObjects/Public/SmartObjectZoneAnnotations.h b/Engine/Plugins/Runtime/MassGameplay/Source/MassSmartObjects/Public/SmartObjectZoneAnnotations.h
--- a/Engine/Plugins/Runtime/MassGameplay/Source/MassSmartObjects/Public/SmartObjectZoneAnnotations.h
+++ b/Engine/Plugins/Runtime/MassGameplay/Source/MassSmartObjects/Public/SmartObjectZoneAnnotations.h
@@ -86,6 +86,10 @@ protected:
 	virtual void PostZoneGraphDataAdded(const AZoneGraphData& ZoneGraphData) override;
 	virtual void PreZoneGraphDataRemoved(const AZoneGraphData& ZoneGraphData) override;
 
+#if !UE_BUILD_SHIPPING && !UE_BUILD_TEST
+	virtual void DebugDraw(FZoneGraphAnnotationSceneProxy* DebugProxy) override;
+#endif
+
 	/** Filter specifying which lanes the behavior is applied to. */
 	UPROPERTY(EditAnywhere, Category = SmartObject)
 	FZoneGraphTagFilter AffectedLaneTags;
@@ -110,9 +114,7 @@ protected:
 	FDelegateHandle OnMainCollectionChangedHandle;
 #endif // WITH_EDITOR
 
-#if WITH_EDITORONLY_DATA
 	/** Cached SmartObjectSubsystem */
 	UPROPERTY(Transient)
 	USmartObjectSubsystem* SmartObjectSubsystem = nullptr;
-#endif // WITH_EDITORONLY_DATA
 };
