CIS fix for issue #135452

diff --git a/Engine/Plugins/Runtime/MassGameplay/Source/MassSmartObjects/Private/SmartObjectZoneAnnotations.cpp b/Engine/Plugins/Runtime/MassGameplay/Source/MassSmartObjects/Private/SmartObjectZoneAnnotations.cpp
--- a/Engine/Plugins/Runtime/MassGameplay/Source/MassSmartObjects/Private/SmartObjectZoneAnnotations.cpp
+++ b/Engine/Plugins/Runtime/MassGameplay/Source/MassSmartObjects/Private/SmartObjectZoneAnnotations.cpp
@@ -150,38 +150,6 @@ void USmartObjectZoneAnnotations::TickAnnotation(const float DeltaTime, FZoneGra
 #endif
 }
 
-#if WITH_EDITOR
-void USmartObjectZoneAnnotations::OnUnregister()
-{
-	GetDefault<UMassSmartObjectSettings>()->OnAnnotationSettingsChanged.Remove(OnAnnotationSettingsChangedHandle);
-	OnAnnotationSettingsChangedHandle.Reset();
-
-	UE::ZoneGraphDelegates::OnZoneGraphDataBuildDone.Remove(OnGraphDataChangedHandle);
-	OnGraphDataChangedHandle.Reset();
-
-	Super::OnUnregister();
-}
-
-void USmartObjectZoneAnnotations::PostEditChangeChainProperty(FPropertyChangedChainEvent& PropertyChangedEvent)
-{
-	Super::PostEditChangeChainProperty(PropertyChangedEvent);
-
-	FProperty* Property = PropertyChangedEvent.Property;
-	FProperty* MemberProperty = nullptr;
-	if (PropertyChangedEvent.PropertyChain.GetActiveMemberNode())
-	{
-		MemberProperty = PropertyChangedEvent.PropertyChain.GetActiveMemberNode()->GetValue();
-	}
-
-	if (MemberProperty && Property)
-	{
-		if (MemberProperty->GetFName() == GET_MEMBER_NAME_CHECKED(USmartObjectZoneAnnotations, AffectedLaneTags))
-		{
-			RebuildForAllGraphs();
-		}
-	}
-}
-
 #if !UE_BUILD_SHIPPING && !UE_BUILD_TEST
 void USmartObjectZoneAnnotations::DebugDraw(FZoneGraphAnnotationSceneProxy* DebugProxy)
 {
@@ -216,10 +184,10 @@ void USmartObjectZoneAnnotations::DebugDraw(FZoneGraphAnnotationSceneProxy* Debu
 
 			const FVector& ObjectLocation = Entry.GetComponent()->GetComponentLocation();
 			FZoneGraphLaneLocation EntryPointLocation;
-            UE::ZoneGraph::Query::CalculateLocationAlongLane(*ZoneStorage, SOLaneLocation->LaneIndex, SOLaneLocation->DistanceAlongLane, EntryPointLocation);
+			UE::ZoneGraph::Query::CalculateLocationAlongLane(*ZoneStorage, SOLaneLocation->LaneIndex, SOLaneLocation->DistanceAlongLane, EntryPointLocation);
 			const FColor Color = FColor::Silver;
 			constexpr float SphereRadius = 25.f;
-            DebugProxy->Spheres.Emplace(SphereRadius, EntryPointLocation.Position, Color);
+			DebugProxy->Spheres.Emplace(SphereRadius, EntryPointLocation.Position, Color);
 			DebugProxy->Spheres.Emplace(SphereRadius, ObjectLocation, Color);
 			DebugProxy->DashedLines.Emplace(ObjectLocation, EntryPointLocation.Position, Color, /*dash size*/10.f);
 		}
@@ -227,6 +195,38 @@ void USmartObjectZoneAnnotations::DebugDraw(FZoneGraphAnnotationSceneProxy* Debu
 }
 #endif // !UE_BUILD_SHIPPING && !UE_BUILD_TEST
 
+#if WITH_EDITOR
+void USmartObjectZoneAnnotations::OnUnregister()
+{
+	GetDefault<UMassSmartObjectSettings>()->OnAnnotationSettingsChanged.Remove(OnAnnotationSettingsChangedHandle);
+	OnAnnotationSettingsChangedHandle.Reset();
+
+	UE::ZoneGraphDelegates::OnZoneGraphDataBuildDone.Remove(OnGraphDataChangedHandle);
+	OnGraphDataChangedHandle.Reset();
+
+	Super::OnUnregister();
+}
+
+void USmartObjectZoneAnnotations::PostEditChangeChainProperty(FPropertyChangedChainEvent& PropertyChangedEvent)
+{
+	Super::PostEditChangeChainProperty(PropertyChangedEvent);
+
+	FProperty* Property = PropertyChangedEvent.Property;
+	FProperty* MemberProperty = nullptr;
+	if (PropertyChangedEvent.PropertyChain.GetActiveMemberNode())
+	{
+		MemberProperty = PropertyChangedEvent.PropertyChain.GetActiveMemberNode()->GetValue();
+	}
+
+	if (MemberProperty && Property)
+	{
+		if (MemberProperty->GetFName() == GET_MEMBER_NAME_CHECKED(USmartObjectZoneAnnotations, AffectedLaneTags))
+		{
+			RebuildForAllGraphs();
+		}
+	}
+}
+
 void USmartObjectZoneAnnotations::RebuildForSingleGraph(FSmartObjectAnnotationData& Data, const FZoneGraphStorage& Storage)
 {
 	TRACE_CPUPROFILER_EVENT_SCOPE_STR("ZoneGraphSmartObjectBehavior RebuildData")
