[SmartObjectAnnotations] all object entry points are now sorted by distance for each lane so queries will be able to extract the best candidates

diff --git a/Engine/Plugins/AI/MassAI/Source/MassAIBehavior/Private/Tasks/MassZoneGraphFindSmartObjectTarget.cpp b/Engine/Plugins/AI/MassAI/Source/MassAIBehavior/Private/Tasks/MassZoneGraphFindSmartObjectTarget.cpp
--- a/Engine/Plugins/AI/MassAI/Source/MassAIBehavior/Private/Tasks/MassZoneGraphFindSmartObjectTarget.cpp
+++ b/Engine/Plugins/AI/MassAI/Source/MassAIBehavior/Private/Tasks/MassZoneGraphFindSmartObjectTarget.cpp
@@ -45,26 +45,27 @@ EStateTreeRunStatus FMassZoneGraphFindSmartObjectTarget::EnterState(FStateTreeEx
 	const FZoneGraphTag SmartObjectTag = GetDefault<UMassSmartObjectSettings>()->SmartObjectTag;
 	const USmartObjectZoneAnnotations* SOAnnotations = Cast<USmartObjectZoneAnnotations>(AnnotationSubsystem.GetFirstAnnotationForTag(SmartObjectTag));
 
-	
 	SmartObjectLocation.LaneHandle = LaneHandle;
 	SmartObjectLocation.NextExitLinkType = EZoneLaneLinkType::None;
 	SmartObjectLocation.NextLaneHandle.Reset();
 	SmartObjectLocation.bMoveReverse = false;
 	SmartObjectLocation.EndOfPathIntent = EMassMovementAction::Stand;
-	SmartObjectLocation.EndOfPathPosition = SOUser.GetTargetLocation();
+	SmartObjectLocation.EndOfPathPosition = SOUser.TargetLocation;
 	// Let's start moving toward the interaction a bit before the entry point.
 	SmartObjectLocation.AnticipationDistance.Set(100.f);
 
+	// Find entry point on lane for the claimed object
+	TOptional<FSmartObjectLaneLocation> SmartObjectLaneLocation;
 	if (SOAnnotations != nullptr)
 	{
-		// Find entry point on lanes for the claimed object
-		const FSmartObjectAnnotationData* AnnotationData = SOAnnotations->GetAnnotationData(LaneHandle.DataHandle);
-		checkf(AnnotationData, TEXT("FSmartObjectAnnotationData should have been created for each registered valid ZoneGraphData"));
-		const FSmartObjectLaneLocation EntryPoint = AnnotationData->ObjectToEntryPointLookup.FindChecked(SOUser.ClaimHandle.SmartObjectHandle);
+		SmartObjectLaneLocation = SOAnnotations->GetSmartObjectLaneLocation(LaneHandle.DataHandle, SOUser.ClaimHandle.SmartObjectHandle);
+	}
 
+	if (SmartObjectLaneLocation.IsSet())
+	{
 		// Request path along current lane to reach entry point on lane
 		MASSBEHAVIOR_LOG(Log, TEXT("Claim successful: create path along lane to reach interaction location."));
-		SmartObjectLocation.TargetDistance = EntryPoint.DistanceAlongLane;
+		SmartObjectLocation.TargetDistance = SmartObjectLaneLocation.GetValue().DistanceAlongLane;
 	}
 	else
 	{
diff --git a/Engine/Plugins/Runtime/MassGameplay/Source/MassSmartObjects/Private/MassSmartObjectProcessor.cpp b/Engine/Plugins/Runtime/MassGameplay/Source/MassSmartObjects/Private/MassSmartObjectProcessor.cpp
--- a/Engine/Plugins/Runtime/MassGameplay/Source/MassSmartObjects/Private/MassSmartObjectProcessor.cpp
+++ b/Engine/Plugins/Runtime/MassGameplay/Source/MassSmartObjects/Private/MassSmartObjectProcessor.cpp
@@ -214,17 +214,19 @@ void UMassProcessor_SmartObjectCandidatesFinder::Execute(UMassEntitySubsystem& E
 					continue;
 				}
 
-				const FSmartObjectList* SmartObjectList = GraphData->LaneToSmartObjectsLookup.Find(RequestLaneHandle.Index);
-				if (SmartObjectList == nullptr || !ensureMsgf(SmartObjectList->SmartObjects.Num() > 0, TEXT("Lookup table should only contains lanes with one or more associated object(s).")))
+				const FSmartObjectLaneLocationIndices* SmartObjectList = GraphData->LaneToLaneLocationIndicesLookup.Find(RequestLaneHandle.Index);
+				if (SmartObjectList == nullptr || !ensureMsgf(SmartObjectList->SmartObjectLaneLocationIndices.Num() > 0, TEXT("Lookup table should only contains lanes with one or more associated object(s).")))
 				{
 					continue;
 				}
 
-				for (const FSmartObjectHandle& Handle : SmartObjectList->SmartObjects)
+				for (const int32 Index : SmartObjectList->SmartObjectLaneLocationIndices)
 				{
 					// Find entry point using FindChecked since all smart objects added to LaneToSmartObjects lookup table
 					// were also added to the entry point lookup table
-					const FSmartObjectLaneLocation EntryPoint = GraphData->ObjectToEntryPointLookup.FindChecked(Handle);
+					check(GraphData->SmartObjectLaneLocations.IsValidIndex(Index));
+					const FSmartObjectLaneLocation& EntryPoint = GraphData->SmartObjectLaneLocations[Index];
+					const FSmartObjectHandle Handle = EntryPoint.ObjectHandle;
 
 					float Cost = 0.f;
 					if (ensureMsgf(EntryPoint.LaneIndex == RequestLocation.LaneHandle.Index, TEXT("Must be on same lane to be able to use distance along lane.")))
diff --git a/Engine/Plugins/Runtime/MassGameplay/Source/MassSmartObjects/Private/SmartObjectZoneAnnotations.cpp b/Engine/Plugins/Runtime/MassGameplay/Source/MassSmartObjects/Private/SmartObjectZoneAnnotations.cpp
--- a/Engine/Plugins/Runtime/MassGameplay/Source/MassSmartObjects/Private/SmartObjectZoneAnnotations.cpp
+++ b/Engine/Plugins/Runtime/MassGameplay/Source/MassSmartObjects/Private/SmartObjectZoneAnnotations.cpp
@@ -121,6 +121,20 @@ const FSmartObjectAnnotationData* USmartObjectZoneAnnotations::GetAnnotationData
 	return &SmartObjectAnnotationDataArray[Index];
 }
 
+TOptional<FSmartObjectLaneLocation> USmartObjectZoneAnnotations::GetSmartObjectLaneLocation(const FZoneGraphDataHandle& DataHandle, const FSmartObjectHandle& SmartObjectHandle) const
+{
+	TOptional<FSmartObjectLaneLocation> SmartObjectLaneLocation;
+	if (const FSmartObjectAnnotationData* AnnotationData = GetAnnotationData(DataHandle))
+	{
+		const int32 Index = AnnotationData->SmartObjectToLaneLocationIndexLookup.FindChecked(SmartObjectHandle);
+		if (AnnotationData->SmartObjectLaneLocations.IsValidIndex(Index))
+		{
+			SmartObjectLaneLocation = AnnotationData->SmartObjectLaneLocations[Index];
+		}
+	}
+	return SmartObjectLaneLocation;
+}
+
 void USmartObjectZoneAnnotations::TickAnnotation(const float DeltaTime, FZoneGraphAnnotationTagContainer& BehaviorTagContainer)
 {
 	if (!BehaviorTag.IsValid())
@@ -130,7 +144,7 @@ void USmartObjectZoneAnnotations::TickAnnotation(const float DeltaTime, FZoneGra
 
 	for (FSmartObjectAnnotationData& Data : SmartObjectAnnotationDataArray)
 	{
-		if (Data.bInitialTaggingCompleted || !Data.IsValid() || Data.ObjectToEntryPointLookup.IsEmpty())
+		if (Data.bInitialTaggingCompleted || !Data.IsValid() || Data.SmartObjectToLaneLocationIndexLookup.IsEmpty())
 		{
 			continue;
 		}
@@ -159,32 +173,32 @@ void USmartObjectZoneAnnotations::DebugDraw(FZoneGraphAnnotationSceneProxy* Debu
 		return;
 	}
 
-	for (FSmartObjectAnnotationData& Data : SmartObjectAnnotationDataArray)
+	const ASmartObjectCollection* Collection = SmartObjectSubsystem->GetMainCollection();
+	if (Collection == nullptr)
+	{
+		return;
+	}
+
+	for (FSmartObjectAnnotationData& AnnotationData : SmartObjectAnnotationDataArray)
 	{
-		const FZoneGraphStorage* ZoneStorage = Data.DataHandle.IsValid() ? ZoneGraph->GetZoneGraphStorage(Data.DataHandle) : nullptr;
+		const FZoneGraphStorage* ZoneStorage = AnnotationData.DataHandle.IsValid() ? ZoneGraph->GetZoneGraphStorage(AnnotationData.DataHandle) : nullptr;
 		if (ZoneStorage == nullptr)
 		{
 			continue;
 		}
 
-		const ASmartObjectCollection* Collection = SmartObjectSubsystem->GetMainCollection();
-		if (Collection == nullptr)
-		{
-			return;
-		}
-
 		for (const FSmartObjectCollectionEntry& Entry : Collection->GetEntries())
 		{
-			FSmartObjectHandle Handle = Entry.GetHandle();
-			const FSmartObjectLaneLocation* SOLaneLocation = Data.ObjectToEntryPointLookup.Find(Handle);
-			if (SOLaneLocation == nullptr)
+			int32* Index = AnnotationData.SmartObjectToLaneLocationIndexLookup.Find(Entry.GetHandle());
+			if (Index == nullptr)
 			{
 				continue;
 			}
+			const FSmartObjectLaneLocation& SOLaneLocation = AnnotationData.SmartObjectLaneLocations[*Index];
 
 			const FVector& ObjectLocation = Entry.GetComponent()->GetComponentLocation();
 			FZoneGraphLaneLocation EntryPointLocation;
-			UE::ZoneGraph::Query::CalculateLocationAlongLane(*ZoneStorage, SOLaneLocation->LaneIndex, SOLaneLocation->DistanceAlongLane, EntryPointLocation);
+			UE::ZoneGraph::Query::CalculateLocationAlongLane(*ZoneStorage, SOLaneLocation.LaneIndex, SOLaneLocation.DistanceAlongLane, EntryPointLocation);
 			const FColor Color = FColor::Silver;
 			constexpr float SphereRadius = 25.f;
 			DebugProxy->Spheres.Emplace(SphereRadius, EntryPointLocation.Position, Color);
@@ -256,53 +270,58 @@ void USmartObjectZoneAnnotations::RebuildForSingleGraph(FSmartObjectAnnotationDa
 
 	const int32 NumSO = Collection->GetEntries().Num();
 	const int32 NumLanes = Storage.Lanes.Num();
-	Data.ObjectToEntryPointLookup.Empty(NumSO);
-	Data.LaneToSmartObjectsLookup.Empty(NumLanes);
+	Data.SmartObjectLaneLocations.Empty(NumSO);
+	Data.SmartObjectToLaneLocationIndexLookup.Empty(NumSO);
+	Data.LaneToLaneLocationIndicesLookup.Empty(NumLanes);
 	Data.AffectedLanes.Empty(NumLanes);
 
 	for (const FSmartObjectCollectionEntry& Entry : Collection->GetEntries())
 	{
-		FZoneGraphLaneLocation EntryPointLocation;
 		FSmartObjectHandle Handle = Entry.GetHandle();
 		const FVector& ObjectLocation = Entry.GetComponent()->GetComponentLocation();
-
 		const FBox QueryBounds(ObjectLocation - SearchExtent, ObjectLocation + SearchExtent);
 
+		FZoneGraphLaneLocation LaneLocation;
 		float DistanceSqr = 0.f;
-		const bool bFound = UE::ZoneGraph::Query::FindNearestLane(Storage, QueryBounds, AffectedLaneTags, EntryPointLocation, DistanceSqr);
+		const bool bFound = UE::ZoneGraph::Query::FindNearestLane(Storage, QueryBounds, AffectedLaneTags, LaneLocation, DistanceSqr);
 		if (bFound)
 		{
 			NumAdded++;
 
-			const int32 LaneIndex = EntryPointLocation.LaneHandle.Index;
-
-			Data.ObjectToEntryPointLookup.Add(Handle, { LaneIndex, EntryPointLocation.DistanceAlongLane });
-			UE_VLOG_UELOG(this, LogSmartObject, Verbose, TEXT("Adding ZG annotation for SmartObject '%s' on lane '%s'"), *LexToString(Handle), *EntryPointLocation.LaneHandle.ToString());
+			const int32 LaneIndex = LaneLocation.LaneHandle.Index;
 
-			auto AddLookupEntries = [this, &Data, Handle, &ObjectLocation](const int32 InLaneIndex, const FVector& LocationOnLane)
-			{
-				Data.AffectedLanes.AddUnique(InLaneIndex);
-				Data.LaneToSmartObjectsLookup.FindOrAdd(InLaneIndex).SmartObjects.Add(Handle);
+			const int32 SOLaneLocationIndex = Data.SmartObjectLaneLocations.Add(FSmartObjectLaneLocation(Handle, LaneIndex, LaneLocation.DistanceAlongLane));
+			Data.SmartObjectToLaneLocationIndexLookup.Add(Handle, SOLaneLocationIndex);
+			Data.AffectedLanes.AddUnique(LaneIndex);
+			Data.LaneToLaneLocationIndicesLookup.FindOrAdd(LaneIndex).SmartObjectLaneLocationIndices.Add(SOLaneLocationIndex);
 
-				UE_VLOG_SEGMENT(this, LogZoneGraphAnnotations, Display, ObjectLocation, LocationOnLane, FColor::Green, TEXT(""));
-			};
-
-			AddLookupEntries(LaneIndex, EntryPointLocation.Position);
-			UE_VLOG_LOCATION(this, LogZoneGraphAnnotations, Display, ObjectLocation, 50.f /*radius*/, FColor::Green, TEXT("%s"), *LexToString(Handle));
+			UE_VLOG_UELOG(this, LogSmartObject, Verbose, TEXT("Adding ZG annotation for SmartObject '%s' on lane '%s'"), *LexToString(Handle), *LaneLocation.LaneHandle.ToString());
+			UE_VLOG_SEGMENT(this, LogSmartObject, Display, ObjectLocation, LaneLocation.Position, FColor::Green, TEXT(""));
+			UE_VLOG_LOCATION(this, LogSmartObject, Display, ObjectLocation, 50.f /*radius*/, FColor::Green, TEXT("%s"), *LexToString(Handle));
 		}
 		else
 		{
 			NumDiscarded++;
-			UE_VLOG_LOCATION(this, LogZoneGraphAnnotations, Display, ObjectLocation, 75.f /*radius*/, FColor::Red, TEXT("%s"), *LexToString(Handle));
+			UE_VLOG_LOCATION(this, LogSmartObject, Display, ObjectLocation, 75.f /*radius*/, FColor::Red, TEXT("%s"), *LexToString(Handle));
 		}
 	}
 
-	Data.ObjectToEntryPointLookup.Shrink();
+	// Sort all entry points per distance on lane
+	for (auto It(Data.LaneToLaneLocationIndicesLookup.CreateIterator()); It; ++It)
+	{
+		It.Value().SmartObjectLaneLocationIndices.Sort([Locations = Data.SmartObjectLaneLocations](const int32 FirstIndex, const int32 SecondIndex)
+		{
+			return Locations[FirstIndex].DistanceAlongLane < Locations[SecondIndex].DistanceAlongLane;
+		});
+	}
+
+	Data.SmartObjectLaneLocations.Shrink();
+	Data.SmartObjectToLaneLocationIndexLookup.Shrink();
 	Data.AffectedLanes.Shrink();
-	Data.LaneToSmartObjectsLookup.Shrink();
+	Data.LaneToLaneLocationIndicesLookup.Shrink();
 	Data.bInitialTaggingCompleted = false;
 
-	UE_VLOG_UELOG(this, LogZoneGraphAnnotations, Log, TEXT("Summary: %d entry points added, %d discarded%s."), NumAdded, NumDiscarded, NumDiscarded == 0 ? TEXT("") : TEXT(" (too far from any lane)"));
+	UE_VLOG_UELOG(this, LogSmartObject, Log, TEXT("Summary: %d entry points added, %d discarded%s."), NumAdded, NumDiscarded, NumDiscarded == 0 ? TEXT("") : TEXT(" (too far from any lane)"));
 }
 
 void USmartObjectZoneAnnotations::RebuildForAllGraphs()
diff --git a/Engine/Plugins/Runtime/MassGameplay/Source/MassSmartObjects/Public/SmartObjectZoneAnnotations.h b/Engine/Plugins/Runtime/MassGameplay/Source/MassSmartObjects/Public/SmartObjectZoneAnnotations.h
--- a/Engine/Plugins/Runtime/MassGameplay/Source/MassSmartObjects/Public/SmartObjectZoneAnnotations.h
+++ b/Engine/Plugins/Runtime/MassGameplay/Source/MassSmartObjects/Public/SmartObjectZoneAnnotations.h
@@ -10,12 +10,23 @@
 class AZoneGraphData;
 class USmartObjectSubsystem;
 
-/** Minimal amount of data to search and compare lane locations within the same graph. */
+/** Struct to keep track of a SmartObject entry point on a given lane. */
 USTRUCT()
 struct FSmartObjectLaneLocation
 {
 	GENERATED_BODY()
 
+	FSmartObjectLaneLocation() = default;
+	FSmartObjectLaneLocation(const FSmartObjectHandle InObjectHandle, const int32 InLaneIndex, const float InDistanceAlongLane)
+        : ObjectHandle(InObjectHandle)
+        , LaneIndex(InLaneIndex)
+        , DistanceAlongLane(InDistanceAlongLane)
+	{
+	}
+
+	UPROPERTY()
+	FSmartObjectHandle ObjectHandle;
+
 	UPROPERTY()
 	int32 LaneIndex = INDEX_NONE;
 
@@ -23,14 +34,17 @@ struct FSmartObjectLaneLocation
 	float DistanceAlongLane = 0.0f;
 };
 
-/** Simple struct to wrap the container to be able to use in a TMap */
+/**
+ * Struct to store indices to all entry points on a given lane.
+ * Used as a container wrapper to be able to use in a TMap.
+ */
 USTRUCT()
-struct FSmartObjectList
+struct FSmartObjectLaneLocationIndices
 {
 	GENERATED_BODY()
 
 	UPROPERTY(VisibleAnywhere, Category = SmartObject)
-	TArray<FSmartObjectHandle> SmartObjects;
+	TArray<int32> SmartObjectLaneLocationIndices;
 };
 
 /** Per ZoneGraphData smart object look up data. */
@@ -47,8 +61,9 @@ struct FSmartObjectAnnotationData
 	{
 		DataHandle = {};
 		AffectedLanes.Reset();
-		ObjectToEntryPointLookup.Reset();
-		LaneToSmartObjectsLookup.Reset();
+		SmartObjectLaneLocations.Reset();
+		SmartObjectToLaneLocationIndexLookup.Reset();
+		LaneToLaneLocationIndicesLookup.Reset();
 	}
 
 	/** Handle of the ZoneGraphData that this smart object annotation data is associated to */
@@ -59,10 +74,13 @@ struct FSmartObjectAnnotationData
 	TArray<int32> AffectedLanes;
 
 	UPROPERTY(VisibleAnywhere, Category = SmartObject)
-	TMap<FSmartObjectHandle, FSmartObjectLaneLocation> ObjectToEntryPointLookup;
+	TArray<FSmartObjectLaneLocation> SmartObjectLaneLocations;
+
+	UPROPERTY(VisibleAnywhere, Category = SmartObject)
+	TMap<FSmartObjectHandle, int32> SmartObjectToLaneLocationIndexLookup;
 
 	UPROPERTY(VisibleAnywhere, Category = SmartObject)
-	TMap<int32, FSmartObjectList> LaneToSmartObjectsLookup;
+	TMap<int32, FSmartObjectLaneLocationIndices> LaneToLaneLocationIndicesLookup;
 
 	bool bInitialTaggingCompleted = false;
 };
@@ -77,6 +95,7 @@ class MASSSMARTOBJECTS_API USmartObjectZoneAnnotations : public UZoneGraphAnnota
 
 public:
 	const FSmartObjectAnnotationData* GetAnnotationData(FZoneGraphDataHandle DataHandle) const;
+	TOptional<FSmartObjectLaneLocation> GetSmartObjectLaneLocation(const FZoneGraphDataHandle& DataHandle, const FSmartObjectHandle& SmartObjectHandle) const;
 
 protected:
 	virtual void PostSubsystemsInitialized() override;
