[MassSmartObject] misc fixes and cleanup
+ renamed processors
+ moved fragments to dedicated header file
+ moved UseTime from User fragment to new TimedBehavior fragment that replaces Tag
+ processor no longer release and invalidate the user's ClaimHandle, it simply mark the interaction as completed and the StateTree task will take care of the rest. This fixes timing issues between evaluator and task within the tree.

diff --git a/Engine/Plugins/AI/MassAI/Source/MassAIBehavior/Private/Evaluators/MassStateTreeSmartObjectEvaluator.cpp b/Engine/Plugins/AI/MassAI/Source/MassAIBehavior/Private/Evaluators/MassStateTreeSmartObjectEvaluator.cpp
--- a/Engine/Plugins/AI/MassAI/Source/MassAIBehavior/Private/Evaluators/MassStateTreeSmartObjectEvaluator.cpp
+++ b/Engine/Plugins/AI/MassAI/Source/MassAIBehavior/Private/Evaluators/MassStateTreeSmartObjectEvaluator.cpp
@@ -4,12 +4,10 @@
 
 #include "MassAIBehaviorTypes.h"
 #include "MassCommonFragments.h"
-#include "MassAIMovementFragments.h"
 #include "MassSignalSubsystem.h"
-#include "MassStateTreeExecutionContext.h"
+#include "MassSmartObjectFragments.h"
 #include "MassSmartObjectHandler.h"
-#include "MassSmartObjectProcessor.h"
-#include "MassStateTreeProcessors.h"
+#include "MassStateTreeExecutionContext.h"
 #include "MassZoneGraphMovementFragments.h"
 #include "SmartObjectSubsystem.h"
 #include "Engine/World.h"
@@ -73,17 +71,19 @@ void FMassStateTreeSmartObjectEvaluator::Evaluate(FStateTreeExecutionContext& Co
 	FMassSmartObjectRequestID& SearchRequestID = Context.GetInstanceData(SearchRequestIDHandle);
 
 	bCandidatesFound = false;
-	bClaimed = SOUser.GetClaimHandle().IsValid();
+	bClaimed = SOUser.ClaimHandle.IsValid();
 
 	// Already claimed, nothing to do
 	if (bClaimed)
 	{
+		MASSBEHAVIOR_LOG(Verbose, TEXT("Skipped: currently claimed"));
 		return;
 	}
 
 	const UWorld* World = Context.GetWorld();
-	if (SOUser.GetCooldown() > World->GetTimeSeconds())
+	if (SOUser.CooldownEndTime > World->GetTimeSeconds())
 	{
+		MASSBEHAVIOR_LOG(Verbose, TEXT("Skipped: cooldown active (%s)"), *LexToString(SOUser.CooldownEndTime - World->GetTimeSeconds()));
 		return;
 	}
 
@@ -91,6 +91,7 @@ void FMassStateTreeSmartObjectEvaluator::Evaluate(FStateTreeExecutionContext& Co
 	float& NextUpdate = Context.GetInstanceData(NextUpdateHandle);
 	if (NextUpdate > World->GetTimeSeconds())
 	{
+		MASSBEHAVIOR_LOG(Verbose, TEXT("Skipped: waiting next update time (%s)"), *LexToString(NextUpdate-World->GetTimeSeconds()));
 		return;
 	}
 	NextUpdate = 0.f;
@@ -112,6 +113,10 @@ void FMassStateTreeSmartObjectEvaluator::Evaluate(FStateTreeExecutionContext& Co
 			MASSBEHAVIOR_CLOG(!LaneLocation->LaneHandle.IsValid(), Error, TEXT("Always expecting a valid lane from the ZoneGraph movement"));
 			if (LaneLocation->LaneHandle.IsValid())
 			{
+				MASSBEHAVIOR_LOG(Log, TEXT("Requesting search candidates from lane %s (%s/%s)"),
+					*LaneLocation->LaneHandle.ToString(),
+					*LexToString(LaneLocation->DistanceAlongLane),
+					*LexToString(LaneLocation->LaneLength));
 				SearchRequestID = MassSmartObjectHandler.FindCandidatesAsync(RequestingEntity, { LaneLocation->LaneHandle, LaneLocation->DistanceAlongLane });
 			}
 		}
@@ -137,8 +142,7 @@ void FMassStateTreeSmartObjectEvaluator::Evaluate(FStateTreeExecutionContext& Co
 			// Update bindable flag to indicate to tasks and conditions if some candidates were found
 			bCandidatesFound = SearchRequestResult.NumCandidates > 0;
 
-			const FMassEntityHandle RequestingEntity = MassContext.GetEntity();
-			MASSBEHAVIOR_CLOG(bCandidatesFound, Log, TEXT("Found %d smart object candidates for %s"), SearchRequestResult.NumCandidates, *RequestingEntity.DebugGetDescription());
+			MASSBEHAVIOR_CLOG(bCandidatesFound, Log, TEXT("Found %d smart object candidates"), SearchRequestResult.NumCandidates);
 
 			// When using ZoneGraph annotations we don't need to schedule a new update since we only need the CurrentLaneChanged signal.
 			// Otherwise we reschedule with default interval on success or retry interval on failed attempt
@@ -147,7 +151,7 @@ void FMassStateTreeSmartObjectEvaluator::Evaluate(FStateTreeExecutionContext& Co
 				const float DelayInSeconds = bCandidatesFound ? TickInterval : RetryCooldown;
 				NextUpdate = World->GetTimeSeconds() + DelayInSeconds;
 				UMassSignalSubsystem& MassSignalSubsystem = Context.GetExternalData(MassSignalSubsystemHandle);
-				MassSignalSubsystem.DelaySignalEntity(UE::Mass::Signals::SmartObjectRequestCandidates, RequestingEntity, DelayInSeconds);
+				MassSignalSubsystem.DelaySignalEntity(UE::Mass::Signals::SmartObjectRequestCandidates, MassContext.GetEntity(), DelayInSeconds);
 			}
 		}
 		// else wait for the Evaluation that will be triggered by the "candidates ready" signal
diff --git a/Engine/Plugins/AI/MassAI/Source/MassAIBehavior/Private/Tasks/MassStateTreeSmartObjectTask.cpp b/Engine/Plugins/AI/MassAI/Source/MassAIBehavior/Private/Tasks/MassStateTreeSmartObjectTask.cpp
--- a/Engine/Plugins/AI/MassAI/Source/MassAIBehavior/Private/Tasks/MassStateTreeSmartObjectTask.cpp
+++ b/Engine/Plugins/AI/MassAI/Source/MassAIBehavior/Private/Tasks/MassStateTreeSmartObjectTask.cpp
@@ -4,7 +4,7 @@
 #include "MassCommonFragments.h"
 #include "MassAIBehaviorTypes.h"
 #include "MassSmartObjectHandler.h"
-#include "MassSmartObjectProcessor.h"
+#include "MassSmartObjectFragments.h"
 #include "MassStateTreeExecutionContext.h"
 #include "MassAIMovementFragments.h"
 #include "SmartObjectSubsystem.h"
@@ -211,7 +211,7 @@ EStateTreeRunStatus FMassUseSmartObjectTask::Tick(FStateTreeExecutionContext& Co
 
 	case EMassSmartObjectInteractionStatus::Completed:
 		MASSBEHAVIOR_LOG(Log, TEXT("Interaction completed"));
-		SOUser.SetCooldown(Context.GetWorld()->GetTimeSeconds() + Cooldown);
+		SOUser.CooldownEndTime = Context.GetWorld()->GetTimeSeconds() + Cooldown;
 		Status = EStateTreeRunStatus::Succeeded;
 		break;
 
diff --git a/Engine/Plugins/AI/MassAI/Source/MassAIBehavior/Private/Tasks/MassZoneGraphFindSmartObjectTarget.cpp b/Engine/Plugins/AI/MassAI/Source/MassAIBehavior/Private/Tasks/MassZoneGraphFindSmartObjectTarget.cpp
--- a/Engine/Plugins/AI/MassAI/Source/MassAIBehavior/Private/Tasks/MassZoneGraphFindSmartObjectTarget.cpp
+++ b/Engine/Plugins/AI/MassAI/Source/MassAIBehavior/Private/Tasks/MassZoneGraphFindSmartObjectTarget.cpp
@@ -3,7 +3,7 @@
 #include "Tasks/MassZoneGraphFindSmartObjectTarget.h"
 #include "ZoneGraphAnnotationSubsystem.h"
 #include "MassAIBehaviorTypes.h"
-#include "MassSmartObjectProcessor.h"
+#include "MassSmartObjectFragments.h"
 #include "MassSmartObjectSettings.h"
 #include "MassStateTreeExecutionContext.h"
 #include "MassZoneGraphMovementFragments.h"
diff --git a/Engine/Plugins/AI/MassAI/Source/MassAIBehavior/Public/Evaluators/MassStateTreeSmartObjectEvaluator.h b/Engine/Plugins/AI/MassAI/Source/MassAIBehavior/Public/Evaluators/MassStateTreeSmartObjectEvaluator.h
--- a/Engine/Plugins/AI/MassAI/Source/MassAIBehavior/Public/Evaluators/MassStateTreeSmartObjectEvaluator.h
+++ b/Engine/Plugins/AI/MassAI/Source/MassAIBehavior/Public/Evaluators/MassStateTreeSmartObjectEvaluator.h
@@ -2,14 +2,16 @@
 
 #pragma once
 
-#include "MassSmartObjectProcessor.h"
+#include "MassSignalSubsystem.h"
+#include "MassSmartObjectRequest.h"
 #include "MassStateTreeTypes.h"
-#include "MassSmartObjectTypes.h"
+#include "SmartObjectSubsystem.h"
 #include "MassStateTreeSmartObjectEvaluator.generated.h"
 
 struct FStateTreeExecutionContext;
 struct FDataFragment_Transform;
 struct FMassZoneGraphLaneLocationFragment;
+struct FMassSmartObjectUserFragment;
 
 /**
  * Evaluator that will keep track if there is one or more smart object(s) that can be used.
diff --git a/Engine/Plugins/Runtime/MassGameplay/Source/MassSmartObjects/Private/MassSmartObjectBehaviorDefinition.cpp b/Engine/Plugins/Runtime/MassGameplay/Source/MassSmartObjects/Private/MassSmartObjectBehaviorDefinition.cpp
--- a/Engine/Plugins/Runtime/MassGameplay/Source/MassSmartObjects/Private/MassSmartObjectBehaviorDefinition.cpp
+++ b/Engine/Plugins/Runtime/MassGameplay/Source/MassSmartObjects/Private/MassSmartObjectBehaviorDefinition.cpp
@@ -1,11 +1,15 @@
 // Copyright Epic Games, Inc. All Rights Reserved.
 
 #include "MassSmartObjectBehaviorDefinition.h"
-#include "MassSmartObjectProcessor.h"
+
+#include "MassCommandBuffer.h"
+#include "MassSmartObjectFragments.h"
 
 void USmartObjectMassBehaviorDefinition::Activate(UMassEntitySubsystem& EntitySubsystem,
-                                              FMassExecutionContext& Context,
-                                              const FMassBehaviorEntityContext& EntityContext) const
+												  FMassExecutionContext& Context,
+												  const FMassBehaviorEntityContext& EntityContext) const
 {
-	EntityContext.SOUser.SetUseTime(UseTime);
+	FMassSmartObjectTimedBehaviorFragment TimedBehaviorFragment;
+	TimedBehaviorFragment.UseTime = UseTime;
+	Context.Defer().PushCommand(FCommandAddFragmentInstance(EntityContext.Entity, FConstStructView::Make(TimedBehaviorFragment)));
 }
diff --git a/Engine/Plugins/Runtime/MassGameplay/Source/MassSmartObjects/Private/MassSmartObjectHandler.cpp b/Engine/Plugins/Runtime/MassGameplay/Source/MassSmartObjects/Private/MassSmartObjectHandler.cpp
--- a/Engine/Plugins/Runtime/MassGameplay/Source/MassSmartObjects/Private/MassSmartObjectHandler.cpp
+++ b/Engine/Plugins/Runtime/MassGameplay/Source/MassSmartObjects/Private/MassSmartObjectHandler.cpp
@@ -1,11 +1,12 @@
 // Copyright Epic Games, Inc. All Rights Reserved.
 
 #include "MassSmartObjectHandler.h"
-#include "MassSmartObjectProcessor.h"
 #include "MassSmartObjectBehaviorDefinition.h"
+#include "MassCommonTypes.h"
 #include "SmartObjectSubsystem.h"
 #include "VisualLogger/VisualLogger.h"
 #include "MassCommandBuffer.h"
+#include "MassSmartObjectFragments.h"
 
 namespace UE::Mass::SmartObject
 {
@@ -202,7 +203,6 @@ bool FMassSmartObjectHandler::UseSmartObject(
 	BehaviorDefinition->Activate(EntitySubsystem, ExecutionContext, FMassBehaviorEntityContext(Entity, Transform, User, SmartObjectSubsystem));
 
 	User.InteractionStatus = EMassSmartObjectInteractionStatus::InProgress;
-	ExecutionContext.Defer().AddTag<FMassSmartObjectTimedBehaviorTag>(Entity);
 
 	return true;
 }
@@ -238,7 +238,9 @@ void FMassSmartObjectHandler::ReleaseSmartObject(const FMassEntityHandle Entity,
 
 	case EMassSmartObjectInteractionStatus::Completed:
 	case EMassSmartObjectInteractionStatus::Aborted:
-		ensureMsgf(false, TEXT("Not expecting status changes for'Completed' or 'Aborted' interaction. Received %s"), *UEnum::GetValueAsString(NewStatus));
+		ensureMsgf(CurrentStatus == NewStatus, TEXT("Not expecting status changes for'Completed' or 'Aborted' interaction. Current %s - Received %s"),
+			*UEnum::GetValueAsString(CurrentStatus),
+			*UEnum::GetValueAsString(NewStatus));
 		break;
 
 	default:
diff --git a/Engine/Plugins/Runtime/MassGameplay/Source/MassSmartObjects/Private/MassSmartObjectProcessor.cpp b/Engine/Plugins/Runtime/MassGameplay/Source/MassSmartObjects/Private/MassSmartObjectProcessor.cpp
--- a/Engine/Plugins/Runtime/MassGameplay/Source/MassSmartObjects/Private/MassSmartObjectProcessor.cpp
+++ b/Engine/Plugins/Runtime/MassGameplay/Source/MassSmartObjects/Private/MassSmartObjectProcessor.cpp
@@ -5,7 +5,10 @@
 #include "MassCommonTypes.h"
 #include "MassSignalSubsystem.h"
 #include "MassSmartObjectBehaviorDefinition.h"
+#include "MassSmartObjectFragments.h"
+#include "MassSmartObjectRequest.h"
 #include "MassSmartObjectSettings.h"
+#include "MassSmartObjectTypes.h"
 #include "SmartObjectZoneAnnotations.h"
 #include "Misc/ScopeExit.h"
 #include "SmartObjectOctree.h"
@@ -15,28 +18,17 @@
 #include "ZoneGraphSubsystem.h"
 
 //----------------------------------------------------------------------//
-// UMassProcessor_SmartObjectBase
+// UMassSmartObjectCandidatesFinderProcessor
 //----------------------------------------------------------------------//
-void UMassProcessor_SmartObjectBase::Initialize(UObject& Owner)
+void UMassSmartObjectCandidatesFinderProcessor::Initialize(UObject& Owner)
 {
-	Super::Initialize(Owner);
-
 	SmartObjectSubsystem = UWorld::GetSubsystem<USmartObjectSubsystem>(GetWorld());
 	SignalSubsystem = UWorld::GetSubsystem<UMassSignalSubsystem>(Owner.GetWorld());
-}
-
-//----------------------------------------------------------------------//
-// UMassProcessor_SmartObjectCandidatesFinder
-//----------------------------------------------------------------------//
-void UMassProcessor_SmartObjectCandidatesFinder::Initialize(UObject& Owner)
-{
-	Super::Initialize(Owner);
-
 	AnnotationSubsystem = UWorld::GetSubsystem<UZoneGraphAnnotationSubsystem>(GetWorld());
 	ZoneGraphSubsystem = UWorld::GetSubsystem<UZoneGraphSubsystem>(GetWorld());
 }
 
-void UMassProcessor_SmartObjectCandidatesFinder::ConfigureQueries()
+void UMassSmartObjectCandidatesFinderProcessor::ConfigureQueries()
 {
 	WorldRequestQuery.AddRequirement<FMassSmartObjectWorldLocationRequestFragment>(EMassFragmentAccess::ReadOnly);
 	WorldRequestQuery.AddRequirement<FMassSmartObjectRequestResultFragment>(EMassFragmentAccess::ReadWrite);
@@ -47,7 +39,7 @@ void UMassProcessor_SmartObjectCandidatesFinder::ConfigureQueries()
 	LaneRequestQuery.AddTagRequirement<FMassSmartObjectCompletedRequestTag>(EMassFragmentPresence::None);
 }
 
-UMassProcessor_SmartObjectCandidatesFinder::UMassProcessor_SmartObjectCandidatesFinder()
+UMassSmartObjectCandidatesFinderProcessor::UMassSmartObjectCandidatesFinderProcessor()
 {
 	// 1. Frame T Behavior create a request(deferred entity creation)
 	// 2. Frame T+1: Processor execute the request might mark it as done(deferred add tag flushed at the end of the frame)
@@ -56,7 +48,7 @@ UMassProcessor_SmartObjectCandidatesFinder::UMassProcessor_SmartObjectCandidates
 	ExecutionOrder.ExecuteBefore.Add(UE::Mass::ProcessorGroupNames::Behavior);
 }
 
-void UMassProcessor_SmartObjectCandidatesFinder::Execute(UMassEntitySubsystem& EntitySubsystem, FMassExecutionContext& Context)
+void UMassSmartObjectCandidatesFinderProcessor::Execute(UMassEntitySubsystem& EntitySubsystem, FMassExecutionContext& Context)
 {
 	checkf(SmartObjectSubsystem != nullptr, TEXT("SmartObjectSubsystem should exist when executing processors."));
 	checkf(SignalSubsystem != nullptr, TEXT("MassSignalSubsystem should exist when executing processors."));
@@ -280,26 +272,32 @@ void UMassProcessor_SmartObjectCandidatesFinder::Execute(UMassEntitySubsystem& E
 }
 
 //----------------------------------------------------------------------//
-// UMassProcessor_SmartObjectTimedBehavior
+// UMassSmartObjectTimedBehaviorProcessor
 //----------------------------------------------------------------------//
-void UMassProcessor_SmartObjectTimedBehavior::ConfigureQueries()
+void UMassSmartObjectTimedBehaviorProcessor::Initialize(UObject& Owner)
+{
+	SmartObjectSubsystem = UWorld::GetSubsystem<USmartObjectSubsystem>(GetWorld());
+	SignalSubsystem = UWorld::GetSubsystem<UMassSignalSubsystem>(Owner.GetWorld());
+}
+
+void UMassSmartObjectTimedBehaviorProcessor::ConfigureQueries()
 {
 	EntityQuery.AddRequirement<FMassSmartObjectUserFragment>(EMassFragmentAccess::ReadWrite);
-	EntityQuery.AddTagRequirement<FMassSmartObjectTimedBehaviorTag>(EMassFragmentPresence::All);
+	EntityQuery.AddRequirement<FMassSmartObjectTimedBehaviorFragment>(EMassFragmentAccess::ReadWrite);
 }
 
-UMassProcessor_SmartObjectTimedBehavior::UMassProcessor_SmartObjectTimedBehavior()
+UMassSmartObjectTimedBehaviorProcessor::UMassSmartObjectTimedBehaviorProcessor()
 {
 	ExecutionOrder.ExecuteInGroup = UE::Mass::ProcessorGroupNames::Tasks;
 	ExecutionOrder.ExecuteAfter.Add(UE::Mass::ProcessorGroupNames::Behavior);
 }
 
-void UMassProcessor_SmartObjectTimedBehavior::Execute(UMassEntitySubsystem& EntitySubsystem, FMassExecutionContext& Context)
+void UMassSmartObjectTimedBehaviorProcessor::Execute(UMassEntitySubsystem& EntitySubsystem, FMassExecutionContext& Context)
 {
 	checkf(SmartObjectSubsystem != nullptr, TEXT("SmartObjectSubsystem should exist when executing processors."));
 	checkf(SignalSubsystem != nullptr, TEXT("MassSignalSubsystem should exist when executing processors."));
 
-	TArray<TTuple<FMassEntityHandle, FMassSmartObjectUserFragment*>> ToRelease;
+	TArray<FMassEntityHandle> ToRelease;
 
 	QUICK_SCOPE_CYCLE_COUNTER(UMassProcessor_SmartObjectTestBehavior_Run);
 
@@ -307,19 +305,18 @@ void UMassProcessor_SmartObjectTimedBehavior::Execute(UMassEntitySubsystem& Enti
 	{
 		const int32 NumEntities = Context.GetNumEntities();
 		const TArrayView<FMassSmartObjectUserFragment> UserList = Context.GetMutableFragmentView<FMassSmartObjectUserFragment>();
+		const TArrayView<FMassSmartObjectTimedBehaviorFragment> TimedBehaviorFragments = Context.GetMutableFragmentView<FMassSmartObjectTimedBehaviorFragment>();
 
 		for (int32 i = 0; i < NumEntities; ++i)
 		{
 			FMassSmartObjectUserFragment& SOUser = UserList[i];
-			if (SOUser.InteractionStatus != EMassSmartObjectInteractionStatus::InProgress)
-			{
-				// Only expecting in progress interaction status
-				continue;
-			}
+			FMassSmartObjectTimedBehaviorFragment& TimedBehaviorFragment = TimedBehaviorFragments[i];
+			ensureMsgf(SOUser.InteractionStatus == EMassSmartObjectInteractionStatus::InProgress, TEXT("Fragment should only be present for in-progress interactions"));
 
 			const float DT = Context.GetDeltaTimeSeconds();
-			SOUser.SetUseTime(FMath::Max(SOUser.GetUseTime() - DT, 0.0f));
-			const bool bMustRelease = SOUser.GetUseTime() <= 0.f;
+			float& UseTime = TimedBehaviorFragment.UseTime;
+			UseTime = FMath::Max(UseTime - DT, 0.0f);
+			const bool bMustRelease = UseTime <= 0.f;
 
 #if WITH_MASSGAMEPLAY_DEBUG
 			const FMassEntityHandle Entity = Context.GetEntity(i);
@@ -327,8 +324,8 @@ void UMassProcessor_SmartObjectTimedBehavior::Execute(UMassEntitySubsystem& Enti
 			const bool bIsDebuggingEntity = UE::Mass::Debug::IsDebuggingEntity(Entity, &DebugColor);
 			if (bIsDebuggingEntity)
 			{
-				UE_CVLOG(bMustRelease, SmartObjectSubsystem, LogSmartObject, Log, TEXT("[%s] stops using [%s]"), *Entity.DebugGetDescription(), *LexToString(SOUser.GetClaimHandle()));
-				UE_CVLOG(!bMustRelease, SmartObjectSubsystem, LogSmartObject, Verbose, TEXT("[%s] using [%s] for %.1f"), *Entity.DebugGetDescription(), *LexToString(SOUser.GetClaimHandle()), SOUser.GetUseTime());
+				UE_CVLOG(bMustRelease, SmartObjectSubsystem, LogSmartObject, Log, TEXT("[%s] stops using [%s]"), *Entity.DebugGetDescription(), *LexToString(SOUser.ClaimHandle));
+				UE_CVLOG(!bMustRelease, SmartObjectSubsystem, LogSmartObject, Verbose, TEXT("[%s] using [%s] for %.1f"), *Entity.DebugGetDescription(), *LexToString(SOUser.ClaimHandle), UseTime);
 
 				const TOptional<FTransform> Transform = SmartObjectSubsystem->GetSlotTransform(SOUser.ClaimHandle);
 				if (Transform.IsSet())
@@ -345,20 +342,15 @@ void UMassProcessor_SmartObjectTimedBehavior::Execute(UMassEntitySubsystem& Enti
 
 			if (bMustRelease)
 			{
-				ToRelease.Add(MakeTuple(Context.GetEntity(i), &SOUser));
+				SOUser.InteractionStatus = EMassSmartObjectInteractionStatus::Completed;
+				ToRelease.Add(Context.GetEntity(i));
 			}
 		}
 	});
 
-	for (const TTuple<FMassEntityHandle, FMassSmartObjectUserFragment*> EntryToRelease : ToRelease)
+	for (const FMassEntityHandle EntityToRelease : ToRelease)
 	{
-		const FMassEntityHandle Entity = EntryToRelease.Get<0>();
-		SignalSubsystem->SignalEntity(UE::Mass::Signals::SmartObjectInteractionDone, Entity);
-
-		FMassSmartObjectUserFragment* User = EntryToRelease.Get<1>();
-		SmartObjectSubsystem->Release(User->ClaimHandle);
-		User->InteractionStatus = EMassSmartObjectInteractionStatus::Completed;
-		Context.Defer().RemoveTag<FMassSmartObjectTimedBehaviorTag>(Entity);
-		User->ClaimHandle.Invalidate();
+		SignalSubsystem->SignalEntity(UE::Mass::Signals::SmartObjectInteractionDone, EntityToRelease);
+		Context.Defer().RemoveFragment<FMassSmartObjectTimedBehaviorFragment>(EntityToRelease);
 	}
 }
diff --git a/Engine/Plugins/Runtime/MassGameplay/Source/MassSmartObjects/Private/MassSmartObjectUserTrait.cpp b/Engine/Plugins/Runtime/MassGameplay/Source/MassSmartObjects/Private/MassSmartObjectUserTrait.cpp
--- a/Engine/Plugins/Runtime/MassGameplay/Source/MassSmartObjects/Private/MassSmartObjectUserTrait.cpp
+++ b/Engine/Plugins/Runtime/MassGameplay/Source/MassSmartObjects/Private/MassSmartObjectUserTrait.cpp
@@ -2,7 +2,7 @@
 
 #include "MassSmartObjectUserTrait.h"
 
-#include "MassSmartObjectProcessor.h"
+#include "MassSmartObjectFragments.h"
 #include "MassEntityTemplateRegistry.h"
 
 void UMassSmartObjectUserTrait::BuildTemplate(FMassEntityTemplateBuildContext& BuildContext, UWorld& World) const
diff --git a/Engine/Plugins/Runtime/MassGameplay/Source/MassSmartObjects/Private/MassSmartObjectsModule.cpp b/Engine/Plugins/Runtime/MassGameplay/Source/MassSmartObjects/Private/MassSmartObjectsModule.cpp
--- a/Engine/Plugins/Runtime/MassGameplay/Source/MassSmartObjects/Private/MassSmartObjectsModule.cpp
+++ b/Engine/Plugins/Runtime/MassGameplay/Source/MassSmartObjects/Private/MassSmartObjectsModule.cpp
@@ -21,6 +21,8 @@ void FMassSmartObjectsModule::StartupModule()
 
 	TArray<FCoreRedirect> Redirects;
 	Redirects.Emplace(ECoreRedirectFlags::Type_Class, TEXT("SmartObjectMassBehaviorConfig"), TEXT("SmartObjectMassBehaviorDefinition"));
+	Redirects.Emplace(ECoreRedirectFlags::Type_Class, TEXT("MassProcessor_SmartObjectTimedBehavior"), TEXT("MassSmartObjectTimedBehaviorProcessor"));
+	Redirects.Emplace(ECoreRedirectFlags::Type_Class, TEXT("MassProcessor_SmartObjectCandidatesFinder"), TEXT("MassSmartObjectCandidatesFinderProcessor"));
 
 	FCoreRedirects::AddRedirectList(Redirects, TEXT("MassSmartObject"));
 }
diff --git a/Engine/Plugins/Runtime/MassGameplay/Source/MassSmartObjects/Public/MassSmartObjectFragments.h b/Engine/Plugins/Runtime/MassGameplay/Source/MassSmartObjects/Public/MassSmartObjectFragments.h
new file mode 100644
--- /dev/null
+++ b/Engine/Plugins/Runtime/MassGameplay/Source/MassSmartObjects/Public/MassSmartObjectFragments.h
@@ -0,0 +1,41 @@
+﻿// Copyright Epic Games, Inc. All Rights Reserved.
+
+#pragma once
+
+#include "MassEntityTypes.h"
+#include "MassSmartObjectTypes.h"
+#include "SmartObjectRuntime.h"
+#include "MassSmartObjectFragments.generated.h"
+
+/** Fragment used by an entity to be able to interact with smart objects */
+USTRUCT()
+struct MASSSMARTOBJECTS_API FMassSmartObjectUserFragment : public FMassFragment
+{
+	GENERATED_BODY()
+
+	UPROPERTY(Transient)
+	FSmartObjectClaimHandle ClaimHandle;
+
+	UPROPERTY(Transient)
+	FVector TargetLocation = FVector::ZeroVector;
+
+	UPROPERTY(Transient)
+	EMassSmartObjectInteractionStatus InteractionStatus = EMassSmartObjectInteractionStatus::Unset;
+
+	/**
+	 * World time in seconds before which the user is considered in cooldown and
+	 * won't look for new interactions (value of 0 indicates no cooldown).
+	 */
+	UPROPERTY(Transient)
+	float CooldownEndTime = 0.f;
+};
+
+/** Fragment used to process time based smartobject interactions */
+USTRUCT()
+struct MASSSMARTOBJECTS_API FMassSmartObjectTimedBehaviorFragment : public FMassFragment
+{
+	GENERATED_BODY()
+
+	UPROPERTY(Transient)
+	float UseTime = 0.f;
+};
\ No newline at end of file
diff --git a/Engine/Plugins/Runtime/MassGameplay/Source/MassSmartObjects/Public/MassSmartObjectProcessor.h b/Engine/Plugins/Runtime/MassGameplay/Source/MassSmartObjects/Public/MassSmartObjectProcessor.h
--- a/Engine/Plugins/Runtime/MassGameplay/Source/MassSmartObjects/Public/MassSmartObjectProcessor.h
+++ b/Engine/Plugins/Runtime/MassGameplay/Source/MassSmartObjects/Public/MassSmartObjectProcessor.h
@@ -3,10 +3,6 @@
 #pragma once
 
 #include "MassProcessor.h"
-#include "SmartObjectRuntime.h"
-#include "MassSmartObjectRequest.h"
-#include "MassSimulationLOD.h"
-#include "MassSmartObjectTypes.h"
 #include "MassSmartObjectProcessor.generated.h"
 
 class USmartObjectSubsystem;
@@ -14,65 +10,14 @@ class UMassSignalSubsystem;
 class UZoneGraphSubsystem;
 class UZoneGraphAnnotationSubsystem;
 
-USTRUCT()
-struct MASSSMARTOBJECTS_API FMassSmartObjectUserFragment : public FMassFragment
-{
-	GENERATED_BODY()
-
-public:
-	FVector GetTargetLocation() const { return TargetLocation; }
-
-	const FSmartObjectClaimHandle& GetClaimHandle() const { return ClaimHandle; }
-
-	float GetCooldown() const { return Cooldown; }
-	void SetCooldown(const float InCooldown) { Cooldown = InCooldown; }
-
-	float GetUseTime() const { return UseTime; }
-	void SetUseTime(const float InUseTime) { UseTime = InUseTime; }
-
-	UPROPERTY(Transient)
-	FSmartObjectClaimHandle ClaimHandle;
-
-	UPROPERTY(Transient)
-	FVector TargetLocation = FVector::ZeroVector;
-
-	UPROPERTY(Transient)
-	EMassSmartObjectInteractionStatus InteractionStatus = EMassSmartObjectInteractionStatus::Unset;
-
-private:
-	UPROPERTY(Transient)
-	float Cooldown = 0.f;
-
-	UPROPERTY(Transient)
-	float UseTime = 0.f; 
-};
-
-/**
- * Base class for smart object processors that takes care of caching the smart object manager.
- */
-UCLASS(Abstract)
-class MASSSMARTOBJECTS_API UMassProcessor_SmartObjectBase : public UMassProcessor
-{
-	GENERATED_BODY()
-
-protected:
-	virtual void Initialize(UObject& Owner) override;
-
-	UPROPERTY(Transient)
-	USmartObjectSubsystem* SmartObjectSubsystem;
-
-	UPROPERTY(Transient)
-	UMassSignalSubsystem* SignalSubsystem;
-};
-
 /** Processor that builds a list of candidates objects for each users. */
 UCLASS()
-class MASSSMARTOBJECTS_API UMassProcessor_SmartObjectCandidatesFinder : public UMassProcessor_SmartObjectBase
+class MASSSMARTOBJECTS_API UMassSmartObjectCandidatesFinderProcessor : public UMassProcessor
 {
 	GENERATED_BODY()
 
 public:
-	UMassProcessor_SmartObjectCandidatesFinder();
+	UMassSmartObjectCandidatesFinderProcessor();
 
 protected:
 	virtual void Initialize(UObject& Owner) override;
@@ -83,6 +28,12 @@ protected:
 	UPROPERTY(EditDefaultsOnly, Category = SmartObject, config)
 	float SearchExtents = 5000.f;
 
+	UPROPERTY(Transient)
+	USmartObjectSubsystem* SmartObjectSubsystem;
+
+	UPROPERTY(Transient)
+	UMassSignalSubsystem* SignalSubsystem;
+
 	UPROPERTY(Transient)
 	UZoneGraphSubsystem* ZoneGraphSubsystem;
 
@@ -96,24 +47,24 @@ protected:
 	FMassEntityQuery LaneRequestQuery;
 };
 
-USTRUCT()
-struct MASSSMARTOBJECTS_API FMassSmartObjectTimedBehaviorTag : public FMassTag
-{
-	GENERATED_BODY()
-};
-
-
 /** Processor for time based user's behavior that waits x seconds then releases its claim on the object */
 UCLASS()
-class MASSSMARTOBJECTS_API UMassProcessor_SmartObjectTimedBehavior : public UMassProcessor_SmartObjectBase
+class MASSSMARTOBJECTS_API UMassSmartObjectTimedBehaviorProcessor : public UMassProcessor
 {
 	GENERATED_BODY()
 public:
-	UMassProcessor_SmartObjectTimedBehavior();
+	UMassSmartObjectTimedBehaviorProcessor();
 
 protected:
+	virtual void Initialize(UObject& Owner) override;
 	virtual void Execute(UMassEntitySubsystem& EntitySubsystem, FMassExecutionContext& Context) override;
 	virtual void ConfigureQueries() override;
 
+	UPROPERTY(Transient)
+	USmartObjectSubsystem* SmartObjectSubsystem;
+
+	UPROPERTY(Transient)
+	UMassSignalSubsystem* SignalSubsystem;
+
 	FMassEntityQuery EntityQuery;
 };
\ No newline at end of file
diff --git a/Engine/Plugins/Runtime/MassGameplay/Source/MassSmartObjects/Public/MassSmartObjectRequest.h b/Engine/Plugins/Runtime/MassGameplay/Source/MassSmartObjects/Public/MassSmartObjectRequest.h
--- a/Engine/Plugins/Runtime/MassGameplay/Source/MassSmartObjects/Public/MassSmartObjectRequest.h
+++ b/Engine/Plugins/Runtime/MassGameplay/Source/MassSmartObjects/Public/MassSmartObjectRequest.h
@@ -26,6 +26,28 @@ struct MASSSMARTOBJECTS_API FSmartObjectCandidate
 	float Cost = 0.f;
 };
 
+/**
+ * Identifier associated to a request for smart object candidates. We use a 1:1 match
+ * with an FMassEntityHandle since all requests are batched together using the EntitySubsystem.
+ */
+USTRUCT()
+struct MASSSMARTOBJECTS_API FMassSmartObjectRequestID
+{
+	GENERATED_BODY()
+
+	FMassSmartObjectRequestID() = default;
+	FMassSmartObjectRequestID(const FMassEntityHandle InEntity) : Entity(InEntity) {}
+
+	bool IsSet() const { return Entity.IsSet(); }
+	void Reset() { Entity.Reset(); }
+
+	explicit operator FMassEntityHandle() const { return Entity; }
+
+private:
+	UPROPERTY(Transient)
+	FMassEntityHandle Entity;
+};
+
 /**
  * Struct that holds status and results of a candidate finder request
  */
diff --git a/Engine/Plugins/Runtime/MassGameplay/Source/MassSmartObjects/Public/MassSmartObjectTypes.h b/Engine/Plugins/Runtime/MassGameplay/Source/MassSmartObjects/Public/MassSmartObjectTypes.h
--- a/Engine/Plugins/Runtime/MassGameplay/Source/MassSmartObjects/Public/MassSmartObjectTypes.h
+++ b/Engine/Plugins/Runtime/MassGameplay/Source/MassSmartObjects/Public/MassSmartObjectTypes.h
@@ -2,7 +2,7 @@
 
 #pragma once
 
-#include "MassEntityTypes.h"
+#include "SmartObjectTypes.h"
 #include "MassSmartObjectTypes.generated.h"
 
 namespace UE::Mass::Signals
@@ -12,28 +12,6 @@ namespace UE::Mass::Signals
 	const FName SmartObjectInteractionDone = FName(TEXT("SmartObjectInteractionDone"));
 }
 
-/**
- * Identifier associated to a request for smart object candidates. We use a 1:1 match
- * with an FMassEntityHandle since all requests are batched together using the EntitySubsystem.
- */
-USTRUCT()
-struct MASSSMARTOBJECTS_API FMassSmartObjectRequestID
-{
-	GENERATED_BODY()
-
-	FMassSmartObjectRequestID() = default;
-	FMassSmartObjectRequestID(const FMassEntityHandle InEntity) : Entity(InEntity) {}
-	
-	bool IsSet() const { return Entity.IsSet(); }
-	void Reset() { Entity.Reset(); }
-	
-	explicit operator FMassEntityHandle() const { return Entity; }
-
-private:
-	UPROPERTY(Transient)
-	FMassEntityHandle Entity;
-};
-
 UENUM()
 enum class EMassSmartObjectClaimResult: uint8
 {
@@ -51,4 +29,4 @@ enum class EMassSmartObjectInteractionStatus: uint8
     InProgress,
     Completed,
     Aborted
-};
\ No newline at end of file
+};
