Transformed the MassRepresentationProcessor overridable methods into a more data driven approach.

#rnx

diff --git a/Engine/Plugins/AI/MassCrowd/Source/MassCrowd/Private/MassCrowdRepresentationActorManagement.cpp b/Engine/Plugins/AI/MassCrowd/Source/MassCrowd/Private/MassCrowdRepresentationActorManagement.cpp
new file mode 100644
--- /dev/null
+++ b/Engine/Plugins/AI/MassCrowd/Source/MassCrowd/Private/MassCrowdRepresentationActorManagement.cpp
@@ -0,0 +1,103 @@
+// Copyright Epic Games, Inc. All Rights Reserved.
+
+#include "MassCrowdRepresentationActorManagement.h"
+ #include "Components/SkeletalMeshComponent.h"
+#include "Components/CapsuleComponent.h"
+#include "GameFramework/CharacterMovementComponent.h"
+#include "GameFramework/Character.h"
+#include "MassAgentComponent.h"
+#include "MassRepresentationSubsystem.h"
+
+void UMassCrowdRepresentationActorManagement::SetActorEnabled(const EActorEnabledType EnabledType, AActor& Actor, const int32 EntityIdx, FMassCommandBuffer& CommandBuffer) const
+{
+	Super::SetActorEnabled(EnabledType, Actor, EntityIdx, CommandBuffer);
+
+	const bool bEnabled = EnabledType != EActorEnabledType::Disabled;
+
+	USkeletalMeshComponent* SkeletalMeshComponent = Actor.FindComponentByClass<USkeletalMeshComponent>();
+	if (SkeletalMeshComponent)
+	{
+		// Enable/disable the ticking and visibility of SkeletalMesh and its children
+		SkeletalMeshComponent->SetVisibility(bEnabled);
+		SkeletalMeshComponent->SetComponentTickEnabled(bEnabled);
+		const TArray<USceneComponent*>& AttachedChildren = SkeletalMeshComponent->GetAttachChildren();
+		if (AttachedChildren.Num() > 0)
+		{
+			TInlineComponentArray<USceneComponent*, NumInlinedActorComponents> ComponentStack;
+
+			ComponentStack.Append(AttachedChildren);
+			while (ComponentStack.Num() > 0)
+			{
+				USceneComponent* const CurrentComp = ComponentStack.Pop(/*bAllowShrinking=*/false);
+				if (CurrentComp)
+				{
+					ComponentStack.Append(CurrentComp->GetAttachChildren());
+					CurrentComp->SetVisibility(bEnabled);
+					if (bEnabled)
+					{
+						// Re-enable only if it was enabled at startup
+						CurrentComp->SetComponentTickEnabled(CurrentComp->PrimaryComponentTick.bStartWithTickEnabled);
+					}
+					else
+					{
+						CurrentComp->SetComponentTickEnabled(false);
+					}
+				}
+			}
+		}
+	}
+
+	// Enable/disable the ticking of CharacterMovementComponent as well
+	ACharacter* Character = Cast<ACharacter>(&Actor);
+	UCharacterMovementComponent* MovementComp = Character != nullptr ? Character->GetCharacterMovement() : nullptr;
+	if (MovementComp != nullptr)
+	{
+		MovementComp->SetComponentTickEnabled(bEnabled);
+	}
+
+	// when we "suspend" the puppet actor we need to let the agent subsystem know by unregistering the agent component
+	// associated with the actor. This will result in removing all the puppet-actor-specific fragments which in turn
+	// will exclude the owner entity from being processed by puppet-specific processors (usually translators).
+	if (UMassAgentComponent* AgentComp = Actor.FindComponentByClass<UMassAgentComponent>())
+	{
+		AgentComp->PausePuppet(!bEnabled);
+	}
+}
+
+AActor* UMassCrowdRepresentationActorManagement::GetOrSpawnActor(UMassRepresentationSubsystem& RepresentationSubsystem, UMassEntitySubsystem& EntitySubsystem, const FMassEntityHandle MassAgent, FDataFragment_Actor& ActorInfo, const FTransform& Transform, const int16 TemplateActorIndex, FMassActorSpawnRequestHandle& SpawnRequestHandle, const float Priority) const
+{
+	FTransform RootTransform = Transform;
+	
+	if (const AActor* DefaultActor = RepresentationSubsystem.GetTemplateActorClass(TemplateActorIndex).GetDefaultObject())
+	{
+		if (const UCapsuleComponent* CapsuleComp = DefaultActor->FindComponentByClass<UCapsuleComponent>())
+		{
+			RootTransform.AddToTranslation(FVector(0.0f, 0.0f, CapsuleComp->GetScaledCapsuleHalfHeight()));
+		}
+	}
+
+	return Super::GetOrSpawnActor(RepresentationSubsystem, EntitySubsystem, MassAgent, ActorInfo, RootTransform, TemplateActorIndex, SpawnRequestHandle, Priority);
+}
+
+void UMassCrowdRepresentationActorManagement::TeleportActor(const FTransform& Transform, AActor& Actor, FMassCommandBuffer& CommandBuffer) const
+{
+	FTransform RootTransform = Transform;
+
+	if (const UCapsuleComponent* CapsuleComp = Actor.FindComponentByClass<UCapsuleComponent>())
+	{
+		const FVector HalfHeight(0.0f, 0.0f, CapsuleComp->GetScaledCapsuleHalfHeight());
+		RootTransform.AddToTranslation(HalfHeight);
+		const FVector RootLocation = RootTransform.GetLocation();
+		const FVector SweepOffset(0.0f, 0.0f, 20.0f);
+		const FVector Start = RootLocation + SweepOffset;
+		const FVector End = RootLocation - SweepOffset;
+		FCollisionQueryParams Params;
+		Params.AddIgnoredActor(&Actor);
+		FHitResult OutHit;
+		if (Actor.GetWorld()->SweepSingleByChannel(OutHit, Start, End, Transform.GetRotation(), CapsuleComp->GetCollisionObjectType(), CapsuleComp->GetCollisionShape(), Params))
+		{
+			RootTransform.SetLocation(OutHit.Location);
+		}
+	}
+	Super::TeleportActor(RootTransform, Actor, CommandBuffer);
+}
diff --git a/Engine/Plugins/AI/MassCrowd/Source/MassCrowd/Private/MassCrowdRepresentationProcessor.cpp b/Engine/Plugins/AI/MassCrowd/Source/MassCrowd/Private/MassCrowdRepresentationProcessor.cpp
--- a/Engine/Plugins/AI/MassCrowd/Source/MassCrowd/Private/MassCrowdRepresentationProcessor.cpp
+++ b/Engine/Plugins/AI/MassCrowd/Source/MassCrowd/Private/MassCrowdRepresentationProcessor.cpp
@@ -28,103 +28,4 @@ void UMassCrowdRepresentationProcessor::ConfigureQueries()
 {
 	Super::ConfigureQueries();
 	EntityQuery.AddTagRequirement<FTagFragment_MassCrowd>(EMassFragmentPresence::All);
-
-	CharacterMovementEntitiesQuery_Conditional = EntityQuery;
-	CharacterMovementEntitiesQuery_Conditional.AddRequirement<FDataFragment_CharacterMovementComponentWrapper>(EMassFragmentAccess::ReadWrite);
-	CharacterMovementEntitiesQuery_Conditional.AddRequirement<FMassVelocityFragment>(EMassFragmentAccess::ReadOnly);
-	CharacterMovementEntitiesQuery_Conditional.SetChunkFilter(&FMassVisualizationChunkFragment::AreAnyEntitiesVisibleInChunk);
-}
-
-void UMassCrowdRepresentationProcessor::SetActorEnabled(const EActorEnabledType EnabledType, AActor& Actor, const int32 EntityIdx, FMassCommandBuffer& CommandBuffer)
-{
-	Super::SetActorEnabled(EnabledType, Actor, EntityIdx, CommandBuffer);
-
-	const bool bEnabled = EnabledType != EActorEnabledType::Disabled;
-
-	USkeletalMeshComponent* SkeletalMeshComponent = Actor.FindComponentByClass<USkeletalMeshComponent>();
-	if (SkeletalMeshComponent)
-	{
-		// Enable/disable the ticking and visibility of SkeletalMesh and its children
-		SkeletalMeshComponent->SetVisibility(bEnabled);
-		SkeletalMeshComponent->SetComponentTickEnabled(bEnabled);
-		const TArray<USceneComponent*>& AttachedChildren = SkeletalMeshComponent->GetAttachChildren();
-		if (AttachedChildren.Num() > 0)
-		{
-			TInlineComponentArray<USceneComponent*, NumInlinedActorComponents> ComponentStack;
-
-			ComponentStack.Append(AttachedChildren);
-			while (ComponentStack.Num() > 0)
-			{
-				USceneComponent* const CurrentComp = ComponentStack.Pop(/*bAllowShrinking=*/false);
-				if (CurrentComp)
-				{
-					ComponentStack.Append(CurrentComp->GetAttachChildren());
-					CurrentComp->SetVisibility(bEnabled);
-					if (bEnabled)
-					{
-						// Re-enable only if it was enabled at startup
-						CurrentComp->SetComponentTickEnabled(CurrentComp->PrimaryComponentTick.bStartWithTickEnabled);
-					}
-					else
-					{
-						CurrentComp->SetComponentTickEnabled(false);
-					}
-				}
-			}
-		}
-	}
-
-	// Enable/disable the ticking of CharacterMovementComponent as well
-	ACharacter* Character = Cast<ACharacter>(&Actor);
-	UCharacterMovementComponent* MovementComp = Character != nullptr ? Character->GetCharacterMovement() : nullptr;
-	if (MovementComp != nullptr)
-	{
-		MovementComp->SetComponentTickEnabled(bEnabled);
-	}
-
-	// when we "suspend" the puppet actor we need to let the agent subsystem know by unregistering the agent component
-	// associated with the actor. This will result in removing all the puppet-actor-specific fragments which in turn
-	// will exclude the owner entity from being processed by puppet-specific processors (usually translators).
-	if (UMassAgentComponent* AgentComp = Actor.FindComponentByClass<UMassAgentComponent>())
-	{
-		AgentComp->PausePuppet(!bEnabled);
-	}
-}
-
-AActor* UMassCrowdRepresentationProcessor::GetOrSpawnActor(UMassRepresentationSubsystem& RepresentationSubsystem, const FMassEntityHandle MassAgent, FDataFragment_Actor& ActorInfo, const FTransform& Transform, const int16 TemplateActorIndex, FMassActorSpawnRequestHandle& SpawnRequestHandle, const float Priority)
-{
-	FTransform RootTransform = Transform;
-	
-	if (const AActor* DefaultActor = RepresentationSubsystem.GetTemplateActorClass(TemplateActorIndex).GetDefaultObject())
-	{
-		if (const UCapsuleComponent* CapsuleComp = DefaultActor->FindComponentByClass<UCapsuleComponent>())
-		{
-			RootTransform.AddToTranslation(FVector(0.0f, 0.0f, CapsuleComp->GetScaledCapsuleHalfHeight()));
-		}
-	}
-
-	return Super::GetOrSpawnActor(RepresentationSubsystem, MassAgent, ActorInfo, RootTransform, TemplateActorIndex, SpawnRequestHandle, Priority);
-}
-
-void UMassCrowdRepresentationProcessor::TeleportActor(const FTransform& Transform, AActor& Actor, FMassCommandBuffer& CommandBuffer)
-{
-	FTransform RootTransform = Transform;
-
-	if (const UCapsuleComponent* CapsuleComp = Actor.FindComponentByClass<UCapsuleComponent>())
-	{
-		const FVector HalfHeight(0.0f, 0.0f, CapsuleComp->GetScaledCapsuleHalfHeight());
-		RootTransform.AddToTranslation(HalfHeight);
-		const FVector RootLocation = RootTransform.GetLocation();
-		const FVector SweepOffset(0.0f, 0.0f, 20.0f);
-		const FVector Start = RootLocation + SweepOffset;
-		const FVector End = RootLocation - SweepOffset;
-		FCollisionQueryParams Params;
-		Params.AddIgnoredActor(&Actor);
-		FHitResult OutHit;
-		if (Actor.GetWorld()->SweepSingleByChannel(OutHit, Start, End, Transform.GetRotation(), CapsuleComp->GetCollisionObjectType(), CapsuleComp->GetCollisionShape(), Params))
-		{
-			RootTransform.SetLocation(OutHit.Location);
-		}
-	}
-	Super::TeleportActor(RootTransform, Actor, CommandBuffer);
 }
diff --git a/Engine/Plugins/AI/MassCrowd/Source/MassCrowd/Private/MassCrowdServerRepresentationTrait.cpp b/Engine/Plugins/AI/MassCrowd/Source/MassCrowd/Private/MassCrowdServerRepresentationTrait.cpp
--- a/Engine/Plugins/AI/MassCrowd/Source/MassCrowd/Private/MassCrowdServerRepresentationTrait.cpp
+++ b/Engine/Plugins/AI/MassCrowd/Source/MassCrowd/Private/MassCrowdServerRepresentationTrait.cpp
@@ -6,6 +6,7 @@
 #include "MassCommonFragments.h"
 #include "MassRepresentationFragments.h"
 #include "Engine/World.h"
+#include "MassCrowdRepresentationActorManagement.h"
 
 void UMassCrowdServerRepresentationTrait::BuildTemplate(FMassEntityTemplateBuildContext& BuildContext, UWorld& World) const
 {
@@ -32,6 +33,12 @@ void UMassCrowdServerRepresentationTrait::BuildTemplate(FMassEntityTemplateBuild
 	FSharedStruct SubsystemFragment = EntitySubsystem->GetOrCreateSharedFragment<FMassRepresentationSubsystemFragment>(SubsystemHash, Subsystem);
 	BuildContext.AddSharedFragment(SubsystemFragment);
 
+	FMassRepresentationConfig Config;
+	Config.RepresentationActorManagement = UMassCrowdRepresentationActorManagement::StaticClass()->GetDefaultObject<UMassCrowdRepresentationActorManagement>();
+	uint32 ConfigHash = UE::StructUtils::GetStructCrc32(FConstStructView::Make(Config));
+	FConstSharedStruct ConfigFragment = EntitySubsystem->GetOrCreateConstSharedFragment<FMassRepresentationConfig>(ConfigHash, Config);
+	BuildContext.AddConstSharedFragment(ConfigFragment);
+
 	FMassRepresentationFragment& RepresentationFragment = BuildContext.AddFragment_GetRef<FMassRepresentationFragment>();
 	RepresentationFragment.StaticMeshDescIndex = INDEX_NONE;
 	RepresentationFragment.HighResTemplateActorIndex = TemplateActor.Get() ? RepresentationSubsystem->FindOrAddTemplateActor(TemplateActor.Get()) : INDEX_NONE;
diff --git a/Engine/Plugins/AI/MassCrowd/Source/MassCrowd/Private/MassCrowdVisualizationTrait.cpp b/Engine/Plugins/AI/MassCrowd/Source/MassCrowd/Private/MassCrowdVisualizationTrait.cpp
--- a/Engine/Plugins/AI/MassCrowd/Source/MassCrowd/Private/MassCrowdVisualizationTrait.cpp
+++ b/Engine/Plugins/AI/MassCrowd/Source/MassCrowd/Private/MassCrowdVisualizationTrait.cpp
@@ -3,9 +3,11 @@
 #include "MassCrowdVisualizationTrait.h"
 #include "MassCrowdRepresentationSubsystem.h"
 #include "MassCrowdVisualizationProcessor.h"
+#include "MassCrowdRepresentationActorManagement.h"
 
 UMassCrowdVisualizationTrait::UMassCrowdVisualizationTrait()
 {
 	// Override the subsystem to support parallelization of the crowd
 	RepresentationSubsystemClass = UMassCrowdRepresentationSubsystem::StaticClass();
+	RepresentationActorManagementClass = UMassCrowdRepresentationActorManagement::StaticClass();
 }
diff --git a/Engine/Plugins/AI/MassCrowd/Source/MassCrowd/Public/MassCrowdRepresentationActorManagement.h b/Engine/Plugins/AI/MassCrowd/Source/MassCrowd/Public/MassCrowdRepresentationActorManagement.h
new file mode 100644
--- /dev/null
+++ b/Engine/Plugins/AI/MassCrowd/Source/MassCrowd/Public/MassCrowdRepresentationActorManagement.h
@@ -0,0 +1,51 @@
+// Copyright Epic Games, Inc. All Rights Reserved.
+
+#pragma once
+
+#include "MassRepresentationActorManagement.h"
+
+#include "MassCrowdRepresentationActorManagement.generated.h"
+
+/**
+ * Overridden representation processor to make it tied to the crowd via the requirements.
+ * It is also the base class for all the different type of crowd representation (Visualization & ServerSideRepresentation)
+ */
+UCLASS(abstract)
+class MASSCROWD_API UMassCrowdRepresentationActorManagement : public UMassRepresentationActorManagement
+{
+	GENERATED_BODY()
+
+protected:
+
+	/**
+	 * Enable/disable a spawned actor
+	 * @param EnabledType is the type of enabling to do on this actor
+	 * @param Actor is the actual actor to perform enabling type on
+	 * @param EntityIdx is the entity index currently processing
+	 * @param Context is the current Mass execution context
+	 */
+	virtual void SetActorEnabled(const EActorEnabledType EnabledType, AActor& Actor, const int32 EntityIdx, FMassCommandBuffer& CommandBuffer) const override;
+
+	/**
+	 * Returns an actor of the template type and setup fragments values from it
+	 * @param RepresentationSubsystem to use to get or spawn the actor
+	 * @param EntitySubsystem associated to the mass agent
+	 * @param MassAgent is the handle to the associated mass agent
+	 * @param ActorInfo is the fragment where we are going to store the actor pointer
+	 * @param Transform is the spatial information about where to spawn the actor
+	 * @param TemplateActorIndex is the index of the type fetched with UMassRepresentationSubsystem::FindOrAddTemplateActor()
+	 * @param SpawnRequestHandle (in/out) In: previously requested spawn Out: newly requested spawn
+	 * @param Priority of this spawn request in comparison with the others, lower value means higher priority
+	 * @return the actor spawned
+	 */
+	virtual AActor* GetOrSpawnActor(UMassRepresentationSubsystem& RepresentationSubsystem, UMassEntitySubsystem& EntitySubsystem, const FMassEntityHandle MassAgent, FDataFragment_Actor& ActorInfo, const FTransform& Transform, const int16 TemplateActorIndex, FMassActorSpawnRequestHandle& SpawnRequestHandle, const float Priority) const override;
+
+	/**
+	 * Teleports the actor at the specified transform by preserving its velocity and without collision.
+	 * The destination will be adjusted to fit an existing capsule.
+	 * @param Transform is the new actor's transform 
+	 * @param Actor is the actual actor to teleport
+	 * @param Context is the current Mass execution context
+	 */
+	virtual void TeleportActor(const FTransform& Transform, AActor& Actor, FMassCommandBuffer& CommandBuffer) const override;
+};
\ No newline at end of file
diff --git a/Engine/Plugins/AI/MassCrowd/Source/MassCrowd/Public/MassCrowdRepresentationProcessor.h b/Engine/Plugins/AI/MassCrowd/Source/MassCrowd/Public/MassCrowdRepresentationProcessor.h
--- a/Engine/Plugins/AI/MassCrowd/Source/MassCrowd/Public/MassCrowdRepresentationProcessor.h
+++ b/Engine/Plugins/AI/MassCrowd/Source/MassCrowd/Public/MassCrowdRepresentationProcessor.h
@@ -28,41 +28,4 @@ protected:
 
 	/** Configure the owned FMassEntityQuery instances to express processor's requirements */
 	virtual void ConfigureQueries() override;
-
-	/**
-	 * Enable/disable a spawned actor
-	 * @param EnabledType is the type of enabling to do on this actor
-	 * @param Actor is the actual actor to perform enabling type on
-	 * @param EntityIdx is the entity index currently processing
-	 * @param Context is the current Mass execution context
-	 */
-	virtual void SetActorEnabled(const EActorEnabledType EnabledType, AActor& Actor, const int32 EntityIdx, FMassCommandBuffer& CommandBuffer) override;
-
-	/** 
-	 * Returns an actor of the template type and setup fragments values from it
-	 * @param RepresentationSubsystem to use to get or spawn the actor
-	 * @param MassAgent is the handle to the associated mass agent
-	 * @param ActorInfo is the fragment where we are going to store the actor pointer
-	 * @param Transform is the spatial information about where to spawn the actor 
-	 * @param TemplateActorIndex is the index of the type fetched with UMassRepresentationSubsystem::FindOrAddTemplateActor()
-	 * @param SpawnRequestHandle (in/out) In: previously requested spawn Out: newly requested spawn
-	 * @param Priority of this spawn request in comparison with the others, lower value means higher priority
-	 * @param Context of the execution from the entity sub system
-	 * @return the actor spawned
-	 */
-	virtual AActor* GetOrSpawnActor(UMassRepresentationSubsystem& RepresentationSubsystem, const FMassEntityHandle MassAgent, FDataFragment_Actor& ActorInfo, const FTransform& Transform, const int16 TemplateActorIndex, FMassActorSpawnRequestHandle& SpawnRequestHandle, const float Priority);
-
-	/**
-	 * Teleports the actor at the specified transform by preserving its velocity and without collision.
-	 * The destination will be adjusted to fit an existing capsule.
-	 * @param Transform is the new actor's transform 
-	 * @param Actor is the actual actor to teleport
-	 * @param Context is the current Mass execution context
-	 */
-	virtual void TeleportActor(const FTransform& Transform, AActor& Actor, FMassCommandBuffer& CommandBuffer) override;
-
-	/** 
-	 * A dedicated query for processing entities owning a FDataFragment_CharacterMovementComponentWrapper
-	 */
-	FMassEntityQuery CharacterMovementEntitiesQuery_Conditional;
 };
\ No newline at end of file
diff --git a/Engine/Plugins/Runtime/MassGameplay/Source/MassActors/Public/MassActorSpawnerSubsystem.h b/Engine/Plugins/Runtime/MassGameplay/Source/MassActors/Public/MassActorSpawnerSubsystem.h
--- a/Engine/Plugins/Runtime/MassGameplay/Source/MassActors/Public/MassActorSpawnerSubsystem.h
+++ b/Engine/Plugins/Runtime/MassGameplay/Source/MassActors/Public/MassActorSpawnerSubsystem.h
@@ -31,7 +31,7 @@ typedef FIndexedHandleManager<FMassActorSpawnRequestHandle, true/*bOptimizeHandl
 
 
 DECLARE_DELEGATE_TwoParams(FMassActorPreSpawnDelegate, const FMassActorSpawnRequestHandle& SpawnRequestHandle, const FStructView& SpawnRequest);
-enum EMassActorSpawnRequestAction
+enum class EMassActorSpawnRequestAction : uint8
 {
 	Keep, // Will leave spawning request in the queue and it will be users job to call RemoveActorSpawnRequest
 	Remove, // Will remove the spawning request from the queue once the callback ends
diff --git a/Engine/Plugins/Runtime/MassGameplay/Source/MassRepresentation/Private/MassRepresentationActorManagement.cpp b/Engine/Plugins/Runtime/MassGameplay/Source/MassRepresentation/Private/MassRepresentationActorManagement.cpp
new file mode 100644
--- /dev/null
+++ b/Engine/Plugins/Runtime/MassGameplay/Source/MassRepresentation/Private/MassRepresentationActorManagement.cpp
@@ -0,0 +1,126 @@
+// Copyright Epic Games, Inc. All Rights Reserved.
+
+#include "MassRepresentationActorManagement.h"
+#include "MassCommandBuffer.h"
+#include "MassEntityView.h"
+#include "MassActorSubsystem.h"
+#include "MassRepresentationSubsystem.h"
+
+AActor* UMassRepresentationActorManagement::GetOrSpawnActor(UMassRepresentationSubsystem& RepresentationSubsystem, UMassEntitySubsystem& EntitySubsystem, const FMassEntityHandle MassAgent, FDataFragment_Actor& ActorInfo, const FTransform& Transform, const int16 TemplateActorIndex, FMassActorSpawnRequestHandle& SpawnRequestHandle, const float Priority) const
+{
+	return RepresentationSubsystem.GetOrSpawnActorFromTemplate(MassAgent, Transform, TemplateActorIndex, SpawnRequestHandle, Priority,
+		FMassActorPreSpawnDelegate::CreateUObject(this, &UMassRepresentationActorManagement::OnPreActorSpawn, &EntitySubsystem),
+		FMassActorPostSpawnDelegate::CreateUObject(this, &UMassRepresentationActorManagement::OnPostActorSpawn, &EntitySubsystem));
+}
+
+
+void UMassRepresentationActorManagement::SetActorEnabled(const EActorEnabledType EnabledType, AActor& Actor, const int32 EntityIdx, FMassCommandBuffer& CommandBuffer) const
+{
+	const bool bEnabled = EnabledType != EActorEnabledType::Disabled;
+	if (Actor.IsActorTickEnabled() != bEnabled)
+	{
+		Actor.SetActorTickEnabled(bEnabled);
+	}
+	if (Actor.GetActorEnableCollision() != bEnabled)
+	{
+		// Deferring this as there is a callback internally that could end up doing things outside of the game thread and will fire checks(Chaos mostly)
+		CommandBuffer.EmplaceCommand<FDeferredCommand>([&Actor, bEnabled](UMassEntitySubsystem&)
+		{
+			Actor.SetActorEnableCollision(bEnabled);
+		});
+	}
+}
+
+void UMassRepresentationActorManagement::TeleportActor(const FTransform& Transform, AActor& Actor, FMassCommandBuffer& CommandBuffer) const
+{
+	if (!Actor.GetTransform().Equals(Transform))
+	{
+		CommandBuffer.EmplaceCommand<FDeferredCommand>([&Actor, Transform](UMassEntitySubsystem&)
+		{
+			Actor.SetActorTransform(Transform, /*bSweep*/false, /*OutSweepHitResult*/nullptr, ETeleportType::TeleportPhysics);
+		});
+	}
+}
+
+void UMassRepresentationActorManagement::OnPreActorSpawn(const FMassActorSpawnRequestHandle& SpawnRequestHandle, const FStructView& SpawnRequest, UMassEntitySubsystem* EntitySubsystem) const
+{
+	check(EntitySubsystem);
+
+	const FMassActorSpawnRequest& MassActorSpawnRequest = SpawnRequest.Get<FMassActorSpawnRequest>();
+	const FMassEntityView EntityView(*EntitySubsystem, MassActorSpawnRequest.MassAgent);
+	FDataFragment_Actor& ActorInfo = EntityView.GetFragmentData<FDataFragment_Actor>();
+	FMassRepresentationFragment& Representation = EntityView.GetFragmentData<FMassRepresentationFragment>();
+	UMassRepresentationSubsystem* RepresentationSubsystem = EntityView.GetSharedFragmentData<FMassRepresentationSubsystemFragment>().RepresentationSubsystem;
+	check(RepresentationSubsystem);
+
+	// Release any existing actor
+	if (AActor* Actor = ActorInfo.GetMutable())
+	{
+		checkf(ActorInfo.IsOwnedByMass(), TEXT("If we reach here, we expect the actor to be owned by mass, otherwise we should not be spawning a new one one top of this one."));
+
+		// WARNING!
+		// Need to reset before ReleaseTemplateActor as this action might move the entity to a new archetype and
+		// so the Fragment passed in parameters would not be valid anymore.
+		ActorInfo.ResetAndUpdateHandleMap();
+
+		if (!RepresentationSubsystem->ReleaseTemplateActor(MassActorSpawnRequest.MassAgent, Representation.HighResTemplateActorIndex, Actor, /*bImmediate*/ true))
+		{
+			if (!RepresentationSubsystem->ReleaseTemplateActor(MassActorSpawnRequest.MassAgent, Representation.LowResTemplateActorIndex, Actor, /*bImmediate*/ true))
+			{
+				checkf(false, TEXT("Expecting to be able to release spawned actor either the high res or low res one"));
+			}
+		}
+	}
+}
+
+EMassActorSpawnRequestAction UMassRepresentationActorManagement::OnPostActorSpawn(const FMassActorSpawnRequestHandle& SpawnRequestHandle, const FStructView& SpawnRequest, UMassEntitySubsystem* EntitySubsystem) const
+{
+	check(EntitySubsystem);
+
+	const FMassActorSpawnRequest& MassActorSpawnRequest = SpawnRequest.Get<FMassActorSpawnRequest>();
+	checkf(MassActorSpawnRequest.SpawnedActor, TEXT("Expecting valid spawned actor"));
+
+	// Might be already done if the actor has a MassAgentComponent via the callback OnMassAgentComponentEntityAssociated on the MassRepresentationSubsystem
+	FDataFragment_Actor& ActorInfo = EntitySubsystem->GetFragmentDataChecked<FDataFragment_Actor>(MassActorSpawnRequest.MassAgent);
+	if (ActorInfo.IsValid())
+	{
+		// If already set, make sure it is pointing to the same actor.
+		checkf(ActorInfo.Get() == MassActorSpawnRequest.SpawnedActor, TEXT("Expecting the pointer to the spawned actor in the actor fragment"));
+	}
+	else
+	{
+		ActorInfo.SetAndUpdateHandleMap(MassActorSpawnRequest.MassAgent, MassActorSpawnRequest.SpawnedActor, true/*bIsOwnedByMass*/);
+	}
+
+	return EMassActorSpawnRequestAction::Keep;
+}
+
+void UMassRepresentationActorManagement::ReleaseAnyActorOrCancelAnySpawning(UMassEntitySubsystem& EntitySubsystem, const FMassEntityHandle MassAgent)
+{
+	FMassEntityView EntityView(EntitySubsystem, MassAgent);
+	FDataFragment_Actor& ActorInfo = EntityView.GetFragmentData<FDataFragment_Actor>();
+	FMassRepresentationFragment& Representation = EntityView.GetFragmentData<FMassRepresentationFragment>();
+	UMassRepresentationSubsystem* RepresentationSubsystem = EntityView.GetSharedFragmentData<FMassRepresentationSubsystemFragment>().RepresentationSubsystem;
+	check(RepresentationSubsystem);
+	ReleaseAnyActorOrCancelAnySpawning(*RepresentationSubsystem, MassAgent, ActorInfo, Representation);
+}
+
+void UMassRepresentationActorManagement::ReleaseAnyActorOrCancelAnySpawning(UMassRepresentationSubsystem& RepresentationSubsystem, const FMassEntityHandle MassAgent, FDataFragment_Actor& ActorInfo, FMassRepresentationFragment& Representation)
+{
+	// This method can only release owned by mass actors
+	AActor* Actor = ActorInfo.GetOwnedByMassMutable();
+	if (Actor)
+	{
+		// WARNING!
+		// Need to reset before ReleaseTemplateActorOrCancelSpawning as this action might move the entity to a new archetype and
+		// so the Fragment passed in parameters would not be valid anymore.
+		ActorInfo.ResetAndUpdateHandleMap();
+	}
+	// Try releasing both as we can have a low res actor and a high res spawning request
+	RepresentationSubsystem.ReleaseTemplateActorOrCancelSpawning(MassAgent, Representation.HighResTemplateActorIndex, Actor, Representation.ActorSpawnRequestHandle);
+	if (Representation.LowResTemplateActorIndex != Representation.HighResTemplateActorIndex)
+	{
+		RepresentationSubsystem.ReleaseTemplateActorOrCancelSpawning(MassAgent, Representation.LowResTemplateActorIndex, Actor, Representation.ActorSpawnRequestHandle);
+	}
+	check(!Representation.ActorSpawnRequestHandle.IsValid());
+}
diff --git a/Engine/Plugins/Runtime/MassGameplay/Source/MassRepresentation/Private/MassRepresentationProcessor.cpp b/Engine/Plugins/Runtime/MassGameplay/Source/MassRepresentation/Private/MassRepresentationProcessor.cpp
--- a/Engine/Plugins/Runtime/MassGameplay/Source/MassRepresentation/Private/MassRepresentationProcessor.cpp
+++ b/Engine/Plugins/Runtime/MassGameplay/Source/MassRepresentation/Private/MassRepresentationProcessor.cpp
@@ -9,6 +9,7 @@
 #include "MassEntitySubsystem.h"
 #include "MassEntityView.h"
 #include "Engine/World.h"
+#include "MassRepresentationActorManagement.h"
 
 namespace UE::MassRepresentation
 {
@@ -38,6 +39,7 @@ void UMassRepresentationProcessor::ConfigureQueries()
 	EntityQuery.AddRequirement<FMassRepresentationLODFragment>(EMassFragmentAccess::ReadOnly);
 	EntityQuery.AddRequirement<FDataFragment_Actor>(EMassFragmentAccess::ReadWrite);
 	EntityQuery.AddChunkRequirement<FMassVisualizationChunkFragment>(EMassFragmentAccess::ReadWrite);
+	EntityQuery.AddConstSharedRequirement<FMassRepresentationConfig>();
 	EntityQuery.AddSharedRequirement<FMassRepresentationSubsystemFragment>(EMassFragmentAccess::ReadWrite);
 }
 
@@ -67,6 +69,9 @@ void UMassRepresentationProcessor::UpdateRepresentation(FMassExecutionContext& C
 {
 	UMassRepresentationSubsystem* RepresentationSubsystem = Context.GetMutableSharedFragment<FMassRepresentationSubsystemFragment>().RepresentationSubsystem;
 	check(RepresentationSubsystem);
+	const FMassRepresentationConfig& RepresentationConfig = Context.GetConstSharedFragment<FMassRepresentationConfig>();
+	UMassRepresentationActorManagement* RepresentationActorManagement = RepresentationConfig.RepresentationActorManagement;
+	check(RepresentationActorManagement);
 
 	const TConstArrayView<FDataFragment_Transform> TransformList = Context.GetFragmentView<FDataFragment_Transform>();
 	const TArrayView<FMassRepresentationFragment> RepresentationList = Context.GetMutableFragmentView<FMassRepresentationFragment>();
@@ -119,12 +124,12 @@ void UMassRepresentationProcessor::UpdateRepresentation(FMassExecutionContext& C
 				// If we already queued spawn request but have changed our mind, continue with it but once we get the actor back, disable it immediately
 				if (Representation.ActorSpawnRequestHandle.IsValid())
 				{
-					Actor = GetOrSpawnActor(*RepresentationSubsystem, MassAgent, ActorInfo, TransformFragment.GetTransform(), Representation.LowResTemplateActorIndex, Representation.ActorSpawnRequestHandle, GetSpawnPriority(RepresentationLOD));
+					Actor = RepresentationActorManagement->GetOrSpawnActor(*RepresentationSubsystem, *CachedEntitySubsystem, MassAgent, ActorInfo, TransformFragment.GetTransform(), Representation.LowResTemplateActorIndex, Representation.ActorSpawnRequestHandle, RepresentationActorManagement->GetSpawnPriority(RepresentationLOD));
 				}
 			}
 			if (Actor != nullptr)
 			{
-				SetActorEnabled(EActorEnabledType::Disabled, *Actor, EntityIdx, Context.Defer());
+				RepresentationActorManagement->SetActorEnabled(EActorEnabledType::Disabled, *Actor, EntityIdx, Context.Defer());
 			}
 		};
 
@@ -166,7 +171,7 @@ void UMassRepresentationProcessor::UpdateRepresentation(FMassExecutionContext& C
 							!RepresentationSubsystem->DoesActorMatchTemplate(*Actor, WantedTemplateActorIndex) ||
 							Representation.ActorSpawnRequestHandle.IsValid())
 						{
-							NewActor = GetOrSpawnActor(*RepresentationSubsystem, MassAgent, ActorInfo, TransformFragment.GetTransform(), WantedTemplateActorIndex, Representation.ActorSpawnRequestHandle, GetSpawnPriority(RepresentationLOD));
+							NewActor = RepresentationActorManagement->GetOrSpawnActor(*RepresentationSubsystem, *CachedEntitySubsystem, MassAgent, ActorInfo, TransformFragment.GetTransform(), WantedTemplateActorIndex, Representation.ActorSpawnRequestHandle, RepresentationActorManagement->GetSpawnPriority(RepresentationLOD));
 						}
 						else
 						{
@@ -184,10 +189,10 @@ void UMassRepresentationProcessor::UpdateRepresentation(FMassExecutionContext& C
 						// Needs to be done before enabling the actor so the animation initialization can use the new values
 						if (Representation.CurrentRepresentation == ERepresentationType::StaticMeshInstance)
 						{
-							TeleportActor(Representation.PrevTransform, *NewActor, Context.Defer());
+							RepresentationActorManagement->TeleportActor(Representation.PrevTransform, *NewActor, Context.Defer());
 						}
 
-						SetActorEnabled(bHighResActor ? EActorEnabledType::HighRes : EActorEnabledType::LowRes, *NewActor, EntityIdx, Context.Defer());
+						RepresentationActorManagement->SetActorEnabled(bHighResActor ? EActorEnabledType::HighRes : EActorEnabledType::LowRes, *NewActor, EntityIdx, Context.Defer());
 						Representation.CurrentRepresentation = WantedRepresentationType;
 					}
 					else if (!Actor)
@@ -214,7 +219,7 @@ void UMassRepresentationProcessor::UpdateRepresentation(FMassExecutionContext& C
 					}
 					else
 					{
-						SetActorEnabled(EActorEnabledType::Disabled, *Actor, EntityIdx, Context.Defer());
+						RepresentationActorManagement->SetActorEnabled(EActorEnabledType::Disabled, *Actor, EntityIdx, Context.Defer());
 					}
 					Representation.CurrentRepresentation = ERepresentationType::None;
 					break;
@@ -242,13 +247,6 @@ void UMassRepresentationProcessor::Execute(UMassEntitySubsystem& InEntitySubsyst
 	});
 }
 
-AActor* UMassRepresentationProcessor::GetOrSpawnActor(UMassRepresentationSubsystem& RepresentationSubsystem, const FMassEntityHandle MassAgent, FDataFragment_Actor& ActorInfo, const FTransform& Transform, const int16 TemplateActorIndex, FMassActorSpawnRequestHandle& SpawnRequestHandle, const float Priority)
-{
-	return RepresentationSubsystem.GetOrSpawnActorFromTemplate(MassAgent, Transform, TemplateActorIndex, SpawnRequestHandle, Priority, 
-		FMassActorPreSpawnDelegate::CreateUObject(this, &UMassRepresentationProcessor::OnActorPreSpawn), 
-		FMassActorPostSpawnDelegate::CreateUObject(this, &UMassRepresentationProcessor::OnActorPostSpawn));
-}
-
 bool UMassRepresentationProcessor::ReleaseActorOrCancelSpawning(UMassRepresentationSubsystem& RepresentationSubsystem, const FMassEntityHandle MassAgent, FDataFragment_Actor& ActorInfo, const int16 TemplateActorIndex, FMassActorSpawnRequestHandle& SpawnRequestHandle, FMassCommandBuffer& CommandBuffer, bool bCancelSpawningOnly /*= false*/)
 {
 	check(!ActorInfo.IsValid() || ActorInfo.IsOwnedByMass());
@@ -279,64 +277,6 @@ bool UMassRepresentationProcessor::ReleaseActorOrCancelSpawning(UMassRepresentat
 	return false;
 }
 
-void UMassRepresentationProcessor::SetActorEnabled(const EActorEnabledType EnabledType, AActor& Actor, const int32 EntityIdx, FMassCommandBuffer& CommandBuffer)
-{
-	const bool bEnabled = EnabledType != EActorEnabledType::Disabled;
-	if (Actor.IsActorTickEnabled() != bEnabled)
-	{
-		Actor.SetActorTickEnabled(bEnabled);
-	}
-	if (Actor.GetActorEnableCollision() != bEnabled)
-	{
-		// Deferring this as there is a callback internally that could end up doing things outside of the game thread and will fire checks(Chaos mostly)
-		CommandBuffer.EmplaceCommand<FDeferredCommand>( [&Actor,bEnabled](UMassEntitySubsystem&)
-		{
-			Actor.SetActorEnableCollision(bEnabled);
-		});
-	}
-}
-
-void UMassRepresentationProcessor::TeleportActor(const FTransform& Transform, AActor& Actor, FMassCommandBuffer& CommandBuffer)
-{
-	if (!Actor.GetTransform().Equals(Transform))
-	{
-		CommandBuffer.EmplaceCommand<FDeferredCommand>([&Actor, Transform](UMassEntitySubsystem&)
-		{
-			Actor.SetActorTransform(Transform, /*bSweep*/false, /*OutSweepHitResult*/nullptr, ETeleportType::TeleportPhysics);
-		});
-	}
-}
-
-void UMassRepresentationProcessor::ReleaseAnyActorOrCancelAnySpawning(UMassEntitySubsystem& EntitySubsystem, const FMassEntityHandle MassAgent)
-{
-	FMassEntityView EntityView(EntitySubsystem,MassAgent);
-	FDataFragment_Actor& ActorInfo = EntityView.GetFragmentData<FDataFragment_Actor>();
-	FMassRepresentationFragment& Representation = EntityView.GetFragmentData<FMassRepresentationFragment>();
-	UMassRepresentationSubsystem* RepresentationSubsystem = EntityView.GetSharedFragmentData<FMassRepresentationSubsystemFragment>().RepresentationSubsystem;
-	check(RepresentationSubsystem);
-	ReleaseAnyActorOrCancelAnySpawning(*RepresentationSubsystem, MassAgent, ActorInfo, Representation);
-}
-
-void UMassRepresentationProcessor::ReleaseAnyActorOrCancelAnySpawning(UMassRepresentationSubsystem& RepresentationSubsystem, const FMassEntityHandle MassAgent, FDataFragment_Actor& ActorInfo, FMassRepresentationFragment& Representation)
-{
-	// This method can only release owned by mass actors
-	AActor* Actor = ActorInfo.GetOwnedByMassMutable();
-	if (Actor)
-	{
-		// WARNING!
-		// Need to reset before ReleaseTemplateActorOrCancelSpawning as this action might move the entity to a new archetype and
-		// so the Fragment passed in parameters would not be valid anymore.
-		ActorInfo.ResetAndUpdateHandleMap();
-	}
-	// Try releasing both as we can have a low res actor and a high res spawning request
-	RepresentationSubsystem.ReleaseTemplateActorOrCancelSpawning(MassAgent, Representation.HighResTemplateActorIndex, Actor, Representation.ActorSpawnRequestHandle);
-	if (Representation.LowResTemplateActorIndex != Representation.HighResTemplateActorIndex)
-	{
-		RepresentationSubsystem.ReleaseTemplateActorOrCancelSpawning(MassAgent, Representation.LowResTemplateActorIndex, Actor, Representation.ActorSpawnRequestHandle);
-	}
-	check(!Representation.ActorSpawnRequestHandle.IsValid());
-}
-
 void UMassRepresentationProcessor::UpdateVisualization(FMassExecutionContext& Context)
 {
 	FMassVisualizationChunkFragment& ChunkData = UpdateChunkVisibility(Context);
@@ -416,60 +356,6 @@ void UMassRepresentationProcessor::UpdateEntityVisibility(const FMassEntityHandl
 	}
 }
 
-void UMassRepresentationProcessor::OnActorPreSpawn(const FMassActorSpawnRequestHandle& SpawnRequestHandle, const FStructView& SpawnRequest)
-{
-	check(CachedEntitySubsystem);
-
-	const FMassActorSpawnRequest& MassActorSpawnRequest = SpawnRequest.Get<FMassActorSpawnRequest>();
-	const FMassEntityView EntityView(*CachedEntitySubsystem, MassActorSpawnRequest.MassAgent);
-	FDataFragment_Actor& ActorInfo = EntityView.GetFragmentData<FDataFragment_Actor>();
-	FMassRepresentationFragment& Representation = EntityView.GetFragmentData<FMassRepresentationFragment>();
-	UMassRepresentationSubsystem* RepresentationSubsystem = EntityView.GetSharedFragmentData<FMassRepresentationSubsystemFragment>().RepresentationSubsystem;
-	check(RepresentationSubsystem);
-
-	// Release any existing actor
-	if (AActor* Actor = ActorInfo.GetMutable())
-	{
-		checkf(ActorInfo.IsOwnedByMass(), TEXT("If we reach here, we expect the actor to be owned by mass, otherwise we should not be spawning a new one one top of this one."));
-
-		// WARNING!
-		// Need to reset before ReleaseTemplateActor as this action might move the entity to a new archetype and
-		// so the Fragment passed in parameters would not be valid anymore.
-		ActorInfo.ResetAndUpdateHandleMap();
-
-		if (!RepresentationSubsystem->ReleaseTemplateActor(MassActorSpawnRequest.MassAgent, Representation.HighResTemplateActorIndex, Actor, /*bImmediate*/ true))
-		{
-			if (!RepresentationSubsystem->ReleaseTemplateActor(MassActorSpawnRequest.MassAgent, Representation.LowResTemplateActorIndex, Actor, /*bImmediate*/ true))
-			{
-				checkf(false, TEXT("Expecting to be able to release spawned actor either the high res or low res one"));
-			}
-		}
-	}
-}
-
-EMassActorSpawnRequestAction UMassRepresentationProcessor::OnActorPostSpawn(const FMassActorSpawnRequestHandle& SpawnRequestHandle, const FStructView& SpawnRequest)
-{
-	check(CachedEntitySubsystem);
-
-	const FMassActorSpawnRequest& MassActorSpawnRequest = SpawnRequest.Get<FMassActorSpawnRequest>();
-	checkf(MassActorSpawnRequest.SpawnedActor, TEXT("Expecting valid spawned actor"));
-
-	// Might be already done if the actor has a MassAgentComponent via the callback OnMassAgentComponentEntityAssociated on the MassRepresentationSubsystem
-	FDataFragment_Actor& ActorInfo = CachedEntitySubsystem->GetFragmentDataChecked<FDataFragment_Actor>(MassActorSpawnRequest.MassAgent);
-	if (ActorInfo.IsValid())
-	{
-		// If already set, make sure it is pointing to the same actor.
-		checkf(ActorInfo.Get() == MassActorSpawnRequest.SpawnedActor, TEXT("Expecting the pointer to the spawned actor in the actor fragment"));
-	}
-	else
-	{
-		ActorInfo.SetAndUpdateHandleMap(MassActorSpawnRequest.MassAgent, MassActorSpawnRequest.SpawnedActor, true/*bIsOwnedByMass*/);
-	}
-
-	PostActorSpawned(MassActorSpawnRequest.MassAgent, *MassActorSpawnRequest.SpawnedActor);
-
-	return EMassActorSpawnRequestAction::Keep;
-}
 
 //----------------------------------------------------------------------//
 // UMassRepresentationFragmentDestructor 
@@ -486,6 +372,7 @@ void UMassRepresentationFragmentDestructor::ConfigureQueries()
 {
 	EntityQuery.AddRequirement<FMassRepresentationFragment>(EMassFragmentAccess::ReadOnly);
 	EntityQuery.AddRequirement<FDataFragment_Actor>(EMassFragmentAccess::ReadWrite);
+	EntityQuery.AddConstSharedRequirement<FMassRepresentationConfig>();
 	EntityQuery.AddSharedRequirement<FMassRepresentationSubsystemFragment>(EMassFragmentAccess::ReadWrite);
 }
 
@@ -507,7 +394,7 @@ void UMassRepresentationFragmentDestructor::Execute(UMassEntitySubsystem& Entity
 
 			const FMassEntityHandle MassAgent = Context.GetEntity(i);
 
-			UMassRepresentationProcessor::ReleaseAnyActorOrCancelAnySpawning(*RepresentationSubsystem, MassAgent, ActorInfo, Representation);
+			UMassRepresentationActorManagement::ReleaseAnyActorOrCancelAnySpawning(*RepresentationSubsystem, MassAgent, ActorInfo, Representation);
 		}
 	});
 }
\ No newline at end of file
diff --git a/Engine/Plugins/Runtime/MassGameplay/Source/MassRepresentation/Private/MassRepresentationSubsystem.cpp b/Engine/Plugins/Runtime/MassGameplay/Source/MassRepresentation/Private/MassRepresentationSubsystem.cpp
--- a/Engine/Plugins/Runtime/MassGameplay/Source/MassRepresentation/Private/MassRepresentationSubsystem.cpp
+++ b/Engine/Plugins/Runtime/MassGameplay/Source/MassRepresentation/Private/MassRepresentationSubsystem.cpp
@@ -13,7 +13,7 @@
 #include "MassAgentSubsystem.h"
 #include "MassEntitySubsystem.h"
 #include "MassRepresentationFragments.h"
-#include "MassRepresentationProcessor.h"
+#include "MassRepresentationActorManagement.h"
 #include "MassEntityView.h"
 #include "WorldPartition/WorldPartitionSubsystem.h"
 #include "WorldPartition/WorldPartitionRuntimeCell.h"
@@ -379,7 +379,7 @@ void UMassRepresentationSubsystem::OnMassAgentComponentEntityAssociated(const UM
 		// Check if this mass agent already handled by this sub system, if yes than release any local spawned actor or cancel any spawn requests
 		if (HandledMassAgents.Find(MassAgent))
 		{
-			UMassRepresentationProcessor::ReleaseAnyActorOrCancelAnySpawning(*EntitySubsystem, MassAgent);
+			UMassRepresentationActorManagement::ReleaseAnyActorOrCancelAnySpawning(*EntitySubsystem, MassAgent);
 		}
 	}
 }
diff --git a/Engine/Plugins/Runtime/MassGameplay/Source/MassRepresentation/Private/MassVisualizationTrait.cpp b/Engine/Plugins/Runtime/MassGameplay/Source/MassRepresentation/Private/MassVisualizationTrait.cpp
--- a/Engine/Plugins/Runtime/MassGameplay/Source/MassRepresentation/Private/MassVisualizationTrait.cpp
+++ b/Engine/Plugins/Runtime/MassGameplay/Source/MassRepresentation/Private/MassVisualizationTrait.cpp
@@ -5,12 +5,13 @@
 #include "MassRepresentationSubsystem.h"
 #include "MassCommonFragments.h"
 #include "MassRepresentationFragments.h"
-#include "MassRepresentationProcessor.h"
+#include "MassRepresentationActorManagement.h"
 #include "Engine/World.h"
 
 UMassVisualizationTrait::UMassVisualizationTrait()
 {
 	RepresentationSubsystemClass = UMassRepresentationSubsystem::StaticClass();
+	RepresentationActorManagementClass = UMassRepresentationActorManagement::StaticClass();
 }
 
 void UMassVisualizationTrait::BuildTemplate(FMassEntityTemplateBuildContext& BuildContext, UWorld& World) const
@@ -31,7 +32,12 @@ void UMassVisualizationTrait::BuildTemplate(FMassEntityTemplateBuildContext& Bui
 	check(EntitySubsystem);
 
 	UMassRepresentationSubsystem* RepresentationSubsystem = Cast<UMassRepresentationSubsystem>(World.GetSubsystemBase(RepresentationSubsystemClass));
-	check(RepresentationSubsystem);
+	if (RepresentationSubsystem == nullptr)
+	{
+		UE_LOG(LogMassRepresentation, Error, TEXT("Expecting a valid class for the representation subsystem"));
+		RepresentationSubsystem = UWorld::GetSubsystem<UMassRepresentationSubsystem>(&World);
+		check(RepresentationSubsystem);
+	}
 
 	FMassRepresentationSubsystemFragment Subsystem;
 	Subsystem.RepresentationSubsystem = RepresentationSubsystem;
@@ -39,6 +45,17 @@ void UMassVisualizationTrait::BuildTemplate(FMassEntityTemplateBuildContext& Bui
 	FSharedStruct SubsystemFragment = EntitySubsystem->GetOrCreateSharedFragment<FMassRepresentationSubsystemFragment>(SubsystemHash, Subsystem);
 	BuildContext.AddSharedFragment(SubsystemFragment);
 
+	FMassRepresentationConfig Config;
+	Config.RepresentationActorManagement = RepresentationActorManagementClass.GetDefaultObject();
+	if (Config.RepresentationActorManagement == nullptr)
+	{
+		UE_LOG(LogMassRepresentation, Error, TEXT("Expecting a valid class for the representation actor management"));
+		Config.RepresentationActorManagement = UMassRepresentationActorManagement::StaticClass()->GetDefaultObject<UMassRepresentationActorManagement>();
+	}
+	uint32 ConfigHash = UE::StructUtils::GetStructCrc32(FConstStructView::Make(Config));
+	FConstSharedStruct ConfigFragment = EntitySubsystem->GetOrCreateConstSharedFragment<FMassRepresentationConfig>(ConfigHash, Config);
+	BuildContext.AddConstSharedFragment(ConfigFragment);
+
 	FMassRepresentationFragment& RepresentationFragment = BuildContext.AddFragment_GetRef<FMassRepresentationFragment>();
 	RepresentationFragment.StaticMeshDescIndex = RepresentationSubsystem->FindOrAddStaticMeshDesc(StaticMeshInstanceDesc);
 	RepresentationFragment.HighResTemplateActorIndex = HighResTemplateActor.Get() ? RepresentationSubsystem->FindOrAddTemplateActor(HighResTemplateActor.Get()) : INDEX_NONE;
diff --git a/Engine/Plugins/Runtime/MassGameplay/Source/MassRepresentation/Public/MassRepresentationActorManagement.h b/Engine/Plugins/Runtime/MassGameplay/Source/MassRepresentation/Public/MassRepresentationActorManagement.h
new file mode 100644
--- /dev/null
+++ b/Engine/Plugins/Runtime/MassGameplay/Source/MassRepresentation/Public/MassRepresentationActorManagement.h
@@ -0,0 +1,105 @@
+// Copyright Epic Games, Inc. All Rights Reserved.
+
+#pragma once
+
+#include "MassRepresentationFragments.h"
+#include "MassRepresentationActorManagement.generated.h"
+
+struct FMassActorSpawnRequestHandle;
+struct FMassCommandBuffer;
+class UMassEntitySubsystem;
+enum class EMassActorSpawnRequestAction : uint8;
+enum class EActorEnabledType : uint8;
+struct FMassActorSpawnRequestHandle;
+struct FStructView;
+struct FMassEntityHandle;
+struct FDataFragment_Actor;
+
+UCLASS()
+class MASSREPRESENTATION_API UMassRepresentationActorManagement : public UObject
+{
+	GENERATED_BODY()
+
+public:
+
+	/**
+	 * Returns the spawn priority from the LOD
+	 * @param Representation is the type of enabling to do on this actor
+	 */
+	virtual float GetSpawnPriority(const FMassRepresentationLODFragment& Representation) const
+	{
+		// Bump up the spawning priority on the visible entities
+		return Representation.LODSignificance - (Representation.Visibility == EMassVisibility::CanBeSeen ? 1.0f : 0.0f);
+	}
+
+	/**
+	 * Returns an actor of the template type and setup fragments values from it
+	 * @param RepresentationSubsystem to use to get or spawn the actor
+	 * @param EntitySubsystem associated to the mass agent
+	 * @param MassAgent is the handle to the associated mass agent
+	 * @param ActorInfo is the fragment where we are going to store the actor pointer
+	 * @param Transform is the spatial information about where to spawn the actor
+	 * @param TemplateActorIndex is the index of the type fetched with UMassRepresentationSubsystem::FindOrAddTemplateActor()
+	 * @param SpawnRequestHandle (in/out) In: previously requested spawn Out: newly requested spawn
+	 * @param Priority of this spawn request in comparison with the others, lower value means higher priority
+	 * @return the actor spawned
+	 */
+	virtual AActor* GetOrSpawnActor(UMassRepresentationSubsystem& RepresentationSubsystem, UMassEntitySubsystem& EntitySubsystem, const FMassEntityHandle MassAgent, FDataFragment_Actor& ActorInfo, const FTransform& Transform, const int16 TemplateActorIndex, FMassActorSpawnRequestHandle& SpawnRequestHandle, const float Priority) const;
+
+	/**
+	 * Enable/disable a spawned actor for a mass entity
+	 * @param EnabledType is the type of enabling to do on this actor
+	 * @param Actor is the actual actor to perform enabling type on
+	 * @param EntityIdx is the entity index currently processing
+	 * @param CommandBuffer to queue up anything that is thread sensitive
+	 */
+	virtual void SetActorEnabled(const EActorEnabledType EnabledType, AActor& Actor, const int32 EntityIdx, FMassCommandBuffer& CommandBuffer) const;
+
+	/**
+	 * Teleports the actor at the specified transform by preserving its velocity and without collision.
+	 * The destination will be adjusted to fit an existing capsule.
+	 * @param Transform is the new actor's transform
+	 * @param Actor is the actual actor to teleport
+	 * @param CommandBuffer to queue up anything that is thread sensitive
+	 */
+	virtual void TeleportActor(const FTransform& Transform, AActor& Actor, FMassCommandBuffer& CommandBuffer) const;
+
+
+	/**
+	 * Method that will be bound to a delegate called before the spawning of an actor to let the requester prepare it
+	 * @param SpawnRequestHandle the handle of the spawn request that is about to spawn
+	 * @param SpawnRequest of the actor that is about to spawn
+	 * @param EntitySubsystem to use to retrieve the mass agent fragments
+	 */
+	virtual void OnPreActorSpawn(const FMassActorSpawnRequestHandle& SpawnRequestHandle, const FStructView& SpawnRequest, UMassEntitySubsystem* EntitySubsystem) const;
+
+	/**
+	 * Method that will be bound to a delegate used post-spawn to notify and let the requester configure the actor
+	 * @param SpawnRequestHandle the handle of the spawn request that was just spawned
+	 * @param SpawnRequest of the actor that just spawned
+	 * @param EntitySubsystem to use to retrieve the mass agent fragments
+	 * @return The action to take on the spawn request, either keep it there or remove it.
+	 */
+	virtual EMassActorSpawnRequestAction OnPostActorSpawn(const FMassActorSpawnRequestHandle& SpawnRequestHandle, const FStructView& SpawnRequest, UMassEntitySubsystem* EntitySubsystem) const;
+
+	/**
+	 * Static methods to Release an actor or cancel its spawning (calls ReleaseAnyActorOrCancelAnySpawning)
+	 * WARNING: This method will destroy the associated actor in any and by the same fact might also move the entity into a new archetype.
+	 *          So any reference to fragment might become invalid.
+	 * @param EntitySubsystem to use to retrieve the mass agent fragments
+	 * @param MassAgent is the handle to the associated mass agent
+	 * @return True if actor was release or spawning request was canceled
+	 */
+	static void ReleaseAnyActorOrCancelAnySpawning(UMassEntitySubsystem& EntitySubsystem, const FMassEntityHandle MassAgent);
+
+	/**
+	 * static Release an actor or cancel its spawning
+	 * WARNING: This method will destroy the associated actor in any and by the same fact might also move the entity into a new archetype.
+	 *          So any reference to fragment might become invalid if you are not within the pipe execution
+	 * @param RepresentationSubsystem to use to release any actors or cancel spawning requests
+	 * @param MassAgent is the handle to the associated mass agent
+	 * @param ActorInfo is the fragment where we are going to store the actor pointer
+	 * @param Representation fragment containing the current and previous visual state
+	 */
+	static void ReleaseAnyActorOrCancelAnySpawning(UMassRepresentationSubsystem& RepresentationSubsystem, const FMassEntityHandle MassAgent, FDataFragment_Actor& ActorInfo, FMassRepresentationFragment& Representation);
+};
\ No newline at end of file
diff --git a/Engine/Plugins/Runtime/MassGameplay/Source/MassRepresentation/Public/MassRepresentationFragments.h b/Engine/Plugins/Runtime/MassGameplay/Source/MassRepresentation/Public/MassRepresentationFragments.h
--- a/Engine/Plugins/Runtime/MassGameplay/Source/MassRepresentation/Public/MassRepresentationFragments.h
+++ b/Engine/Plugins/Runtime/MassGameplay/Source/MassRepresentation/Public/MassRepresentationFragments.h
@@ -11,6 +11,7 @@
 #include "MassRepresentationFragments.generated.h"
 
 class UMassRepresentationSubsystem;
+class UMassRepresentationActorManagement;
 
 USTRUCT()
 struct MASSREPRESENTATION_API FMassRepresentationLODFragment : public FMassFragment
@@ -58,5 +59,14 @@ struct FMassRepresentationSubsystemFragment : public FMassSharedFragment
 	GENERATED_BODY()
 
 	UPROPERTY(Transient)
-	UMassRepresentationSubsystem* RepresentationSubsystem = nullptr;
+	UMassRepresentationSubsystem* RepresentationSubsystem;
+};
+
+USTRUCT()
+struct FMassRepresentationConfig : public FMassSharedFragment
+{
+	GENERATED_BODY()
+
+	UPROPERTY(Transient)
+	UMassRepresentationActorManagement* RepresentationActorManagement;
 };
diff --git a/Engine/Plugins/Runtime/MassGameplay/Source/MassRepresentation/Public/MassRepresentationProcessor.h b/Engine/Plugins/Runtime/MassGameplay/Source/MassRepresentation/Public/MassRepresentationProcessor.h
--- a/Engine/Plugins/Runtime/MassGameplay/Source/MassRepresentation/Public/MassRepresentationProcessor.h
+++ b/Engine/Plugins/Runtime/MassGameplay/Source/MassRepresentation/Public/MassRepresentationProcessor.h
@@ -11,13 +11,6 @@
 class UMassRepresentationSubsystem;
 struct FDataFragment_Actor;
 
-enum EActorEnabledType
-{
-	Disabled,
-	LowRes,
-	HighRes,
-};
-
 UCLASS()
 class MASSREPRESENTATION_API UMassRepresentationProcessor : public UMassProcessor
 {
@@ -44,19 +37,6 @@ protected:
 	 */
 	virtual void Execute(UMassEntitySubsystem& EntitySubsystem, FMassExecutionContext& Context) override;
 
-	/** 
-	 * Returns an actor of the template type and setup fragments values from it
-	 * @param RepresentationSubsystem to use to get or spawn the actor
-	 * @param MassAgent is the handle to the associated mass agent
-	 * @param ActorInfo is the fragment where we are going to store the actor pointer
-	 * @param Transform is the spatial information about where to spawn the actor 
-	 * @param TemplateActorIndex is the index of the type fetched with UMassRepresentationSubsystem::FindOrAddTemplateActor()
-	 * @param SpawnRequestHandle (in/out) In: previously requested spawn Out: newly requested spawn
-	 * @param Priority of this spawn request in comparison with the others, lower value means higher priority
-	 * @return the actor spawned
-	 */
-	virtual AActor* GetOrSpawnActor(UMassRepresentationSubsystem& RepresentationSubsystem, const FMassEntityHandle MassAgent, FDataFragment_Actor& ActorInfo, const FTransform& Transform, const int16 TemplateActorIndex, FMassActorSpawnRequestHandle& SpawnRequestHandle, const float Priority);
-
 	/**
 	 * Release the actor to the subsystem, will only release it the actor or spawn request matches the template actor
 	 * @param RepresentationSubsystem to use to release the actor or cancel the spawning
@@ -68,47 +48,7 @@ protected:
 	 * @param bCancelSpawningOnly tell to only cancel the existing spawning request and to not release the associated actor it any.
 	 * @return if the actor was release or the spawning was canceled.
 	 */
-	virtual bool ReleaseActorOrCancelSpawning(UMassRepresentationSubsystem& RepresentationSubsystem, const FMassEntityHandle MassAgent, FDataFragment_Actor& ActorInfo, const int16 TemplateActorIndex, FMassActorSpawnRequestHandle& SpawnRequestHandle, FMassCommandBuffer& CommandBuffer, const bool bCancelSpawningOnly = false);
-
-	/** 
-	 * Enable/disable a spawned actor
-	 * @param EnabledType is the type of enabling to do on this actor
-	 * @param Actor is the actual actor to perform enabling type on
-	 * @param EntityIdx is the entity index currently processing
-	 * @param CommandBuffer to queue up anything that is thread sensitive
-	 */
-	virtual void SetActorEnabled(const EActorEnabledType EnabledType, AActor& Actor, const int32 EntityIdx, FMassCommandBuffer& CommandBuffer);
-
-	/**
-	 * Teleports the actor at the specified transform by preserving its velocity and without collision.
-	 * The destination will be adjusted to fit an existing capsule.
-	 * @param Transform is the new actor's transform 
-	 * @param Actor is the actual actor to teleport
-	 * @param CommandBuffer to queue up anything that is thread sensitive
-	 */
-	virtual void TeleportActor(const FTransform& Transform, AActor& Actor, FMassCommandBuffer& CommandBuffer);
-
-	/**
-	 * Method that will be bound to a delegate called before the spawning of an actor to let the requester prepare it
-	 * @param SpawnRequestHandle the handle of the spawn request that is about to spawn
-	 * @param SpawnRequest of the actor that is about to spawn
-	 */
-	void OnActorPreSpawn(const FMassActorSpawnRequestHandle& SpawnRequestHandle, const FStructView& SpawnRequest);
-
-	/**
-	 * Method that will be bound to a delegate used post-spawn to notify and let the requester configure the actor
-	 * @param SpawnRequestHandle the handle of the spawn request that was just spawned
-	 * @param SpawnRequest of the actor that just spawned
-	 * @return The action to take on the spawn request, either keep it there or remove it.
-	 */
-	EMassActorSpawnRequestAction OnActorPostSpawn(const FMassActorSpawnRequestHandle& SpawnRequestHandle, const FStructView& SpawnRequest);
-
-	/**
-	 * Method called once the spawning of an actor just occurred, give the processors a chance to configure the actor during the spawn time slicing
-	 * @param MassAgent mass entity to the spawn actor
-	 * @param SpawnedActor that was just spawned
-	 */
-	virtual void PostActorSpawned(const FMassEntityHandle MassAgent, AActor& SpawnedActor) {}
+	bool ReleaseActorOrCancelSpawning(UMassRepresentationSubsystem& RepresentationSubsystem, const FMassEntityHandle MassAgent, FDataFragment_Actor& ActorInfo, const int16 TemplateActorIndex, FMassActorSpawnRequestHandle& SpawnRequestHandle, FMassCommandBuffer& CommandBuffer, const bool bCancelSpawningOnly = false);
 
 	/**
 	 * Updates chunk visibility info for later chunk logic optimization
@@ -127,36 +67,8 @@ protected:
 	 */
 	static void UpdateEntityVisibility(const FMassEntityHandle Entity, const FMassRepresentationFragment& Representation, const FMassRepresentationLODFragment& RepresentationLOD, FMassVisualizationChunkFragment& ChunkData, FMassCommandBuffer& CommandBuffer);
 
-public:
-	/**
-	 * Release an actor or cancel its spawning (calls ReleaseAnyActorOrCancelAnySpawning)
-	 * WARNING: This method will destroy the associated actor in any and by the same fact might also move the entity into a new archetype.
-	 *          So any reference to fragment might become invalid.
-	 * @param EntitySubsystem to use to retrieve the mass agent fragments
-	 * @param MassAgent is the handle to the associated mass agent
-	 * @return True if actor was release or spawning request was canceled
-	 */
-	static void ReleaseAnyActorOrCancelAnySpawning(UMassEntitySubsystem& EntitySubsystem, const FMassEntityHandle MassAgent);
-
-	/**
-	 * Release an actor or cancel its spawning
-	 * WARNING: This method will destroy the associated actor in any and by the same fact might also move the entity into a new archetype.
-	 *          So any reference to fragment might become invalid if you are not within the pipe execution
-	 * @param RepresentationSubsystem to use to release any actors or cancel spawning requests
-	 * @param MassAgent is the handle to the associated mass agent
-	 * @param ActorInfo is the fragment where we are going to store the actor pointer
-	 * @param Representation fragment containing the current and previous visual state
-	 */
-	static void ReleaseAnyActorOrCancelAnySpawning(UMassRepresentationSubsystem& RepresentationSubsystem, const FMassEntityHandle MassAgent, FDataFragment_Actor& ActorInfo, FMassRepresentationFragment& Representation);
-
 protected:
 
-	virtual float GetSpawnPriority(const FMassRepresentationLODFragment& Representation)
-	{
-		// Bump up the spawning priority on the visible entities
-		return Representation.LODSignificance - (Representation.Visibility == EMassVisibility::CanBeSeen ? 1.0f : 0.0f);
-	}
-
 	/*
 	 * Update representation type for each entity, must be called within a ForEachEntityChunk
 	 * @param Context of the execution from the entity sub system
diff --git a/Engine/Plugins/Runtime/MassGameplay/Source/MassRepresentation/Public/MassRepresentationTypes.h b/Engine/Plugins/Runtime/MassGameplay/Source/MassRepresentation/Public/MassRepresentationTypes.h
--- a/Engine/Plugins/Runtime/MassGameplay/Source/MassRepresentation/Public/MassRepresentationTypes.h
+++ b/Engine/Plugins/Runtime/MassGameplay/Source/MassRepresentation/Public/MassRepresentationTypes.h
@@ -27,6 +27,13 @@ enum class ERepresentationType : uint8
 	None,
 };
 
+enum class EActorEnabledType : uint8
+{
+	Disabled,
+	LowRes,
+	HighRes,
+};
+
 USTRUCT()
 struct FStaticMeshInstanceVisualizationMeshDesc
 {
diff --git a/Engine/Plugins/Runtime/MassGameplay/Source/MassRepresentation/Public/MassVisualizationTrait.h b/Engine/Plugins/Runtime/MassGameplay/Source/MassRepresentation/Public/MassVisualizationTrait.h
--- a/Engine/Plugins/Runtime/MassGameplay/Source/MassRepresentation/Public/MassVisualizationTrait.h
+++ b/Engine/Plugins/Runtime/MassGameplay/Source/MassRepresentation/Public/MassVisualizationTrait.h
@@ -8,6 +8,7 @@
 #include "MassVisualizationTrait.generated.h"
 
 class UMassRepresentationSubsystem;
+class UMassRepresentationActorManagement;
 class UMassProcessor;
 
 UCLASS(meta=(DisplayName="Visualization"))
@@ -34,9 +35,13 @@ protected:
 	TSubclassOf<AActor> LowResTemplateActor;
 
 	/** Allow subclasses to override the representation subsystem to use */
-	UPROPERTY()
+	UPROPERTY(EditAnywhere, Category = "Mass|Visual")
 	TSubclassOf<UMassRepresentationSubsystem> RepresentationSubsystemClass;
 
+	/** Allow subclasses to override the representation actor management behavior */
+	UPROPERTY(EditAnywhere, Category = "Mass|Visual")
+	TSubclassOf<UMassRepresentationActorManagement> RepresentationActorManagementClass;
+
 	// @todo the following property will be cut once new de/initializers are in
 	/** Allow subclasses to override the representation fragment deinitializer */
 	TSubclassOf<UMassProcessor> RepresentationFragmentDeinitializerClass;
