MassVisualizationTrait - Made API public to allow for programatic construction and configuration

diff --git a/Engine/Plugins/Runtime/MassGameplay/Source/MassRepresentation/Public/MassVisualizationTrait.h b/Engine/Plugins/Runtime/MassGameplay/Source/MassRepresentation/Public/MassVisualizationTrait.h
--- a/Engine/Plugins/Runtime/MassGameplay/Source/MassRepresentation/Public/MassVisualizationTrait.h
+++ b/Engine/Plugins/Runtime/MassGameplay/Source/MassRepresentation/Public/MassVisualizationTrait.h
@@ -18,8 +18,6 @@ class MASSREPRESENTATION_API UMassVisualizationTrait : public UMassEntityTraitBa
 public:
 	UMassVisualizationTrait();
 
-protected:
-
 	virtual void BuildTemplate(FMassEntityTemplateBuildContext& BuildContext, UWorld& World) const override;
 
 	/** Instanced static mesh information for this agent */
