Checking if the mass processing is active before destroying non-puppet entity in UMassAgentSubsystem::UnregisterAgentComponent and if so issuing a command.

diff --git a/Engine/Plugins/Runtime/MassGameplay/Source/MassActors/Private/MassAgentSubsystem.cpp b/Engine/Plugins/Runtime/MassGameplay/Source/MassActors/Private/MassAgentSubsystem.cpp
--- a/Engine/Plugins/Runtime/MassGameplay/Source/MassActors/Private/MassAgentSubsystem.cpp
+++ b/Engine/Plugins/Runtime/MassGameplay/Source/MassActors/Private/MassAgentSubsystem.cpp
@@ -181,7 +181,7 @@ void UMassAgentSubsystem::UnregisterAgentComponent(UMassAgentComponent& AgentCom
 				const FMassEntityHandle Entity = AgentComp.GetEntityHandle();
 
 			    // remove fragments that have been added for the puppet agent
-			    if (EntitySystem->IsProcessing())
+			    if (SimulationSystem->GetPhaseManager().IsDuringMassProcessing())
 			    {
 				    // need to request via command buffer since we can't move entities while processing is happening
 				    EntitySystem->Defer().PushCommand(FCommandRemoveComposition(Entity, AgentComp.GetPuppetSpecificAddition()));
@@ -212,7 +212,12 @@ void UMassAgentSubsystem::UnregisterAgentComponent(UMassAgentComponent& AgentCom
 			AgentComp.ClearEntityHandle();
 
 			// Destroy the entity
-			if (ensure(EntityTemplate))
+			if (SimulationSystem->GetPhaseManager().IsDuringMassProcessing())
+			{
+				// need to request via command buffer since we can't move entities while processing is happening
+				EntitySystem->Defer().DestroyEntity(Entity);
+			}
+			else if (ensure(EntityTemplate))
 			{
 				SpawnerSystem->DestroyEntities(EntityTemplate->GetTemplateID(), TArrayView<FMassEntityHandle>(&Entity, 1));
 			}
