MassSpawner: Added more generic spawning
- Removed AdjustTickSchematics and remove tick schematics from the spawner
- Renamed "spawn point generation" to "spawn data generation"
- Moved spawned entity selection into the generator (spawn count is more like a hint now)
- Spawn data consists of: entity config, count, spawn data, spawn data processor, and post spawn processors
- Update EQS and ZG generators to work with the new setup

diff --git a/Engine/Plugins/Runtime/MassGameplay/Config/BaseMassGameplay.ini b/Engine/Plugins/Runtime/MassGameplay/Config/BaseMassGameplay.ini
--- a/Engine/Plugins/Runtime/MassGameplay/Config/BaseMassGameplay.ini
+++ b/Engine/Plugins/Runtime/MassGameplay/Config/BaseMassGameplay.ini
@@ -1,2 +1,8 @@
 [CoreRedirects]
-+PackageRedirects=(OldName="/Script/Mass",NewName="/Script/MassGameplay")
\ No newline at end of file
++PackageRedirects=(OldName="/Script/Mass",NewName="/Script/MassGameplay")
+
++StructRedirects=(OldName="MassSpawnPointGenerator", NewName="MassSpawnDataGenerator")
++StructRedirects=(OldName="/Script/MassAIMovement.MassMovementStyleRef", NewName="/Script/MassMovement.MassMovementStyleRef")
++ClassRedirects=(OldName="/Script/MassAIMovement.MassMovementSettings", NewName="/Script/MassMovement.MassMovementSettings")
+
++PropertyRedirects=(OldName="/Script/MassSpawner.SpawnPointsGenerators", NewName="SpawnDataGenerators")
diff --git a/Engine/Plugins/Runtime/MassGameplay/Source/MassSimulation/Public/MassSimulationSubsystem.h b/Engine/Plugins/Runtime/MassGameplay/Source/MassSimulation/Public/MassSimulationSubsystem.h
--- a/Engine/Plugins/Runtime/MassGameplay/Source/MassSimulation/Public/MassSimulationSubsystem.h
+++ b/Engine/Plugins/Runtime/MassGameplay/Source/MassSimulation/Public/MassSimulationSubsystem.h
@@ -20,7 +20,6 @@ class MASSSIMULATION_API UMassSimulationSubsystem : public UWorldSubsystem
 {
 	GENERATED_BODY()
 public:
-	DECLARE_MULTICAST_DELEGATE_TwoParams(FOnAdjustTickSchematics, UWorld* /*World*/, TArray<TSoftObjectPtr<UMassSchematic>>& /*InOutTickSchematics*/);
 	DECLARE_MULTICAST_DELEGATE_OneParam(FOnSimulationStarted, UWorld* /*World*/);
 	
 	UMassSimulationSubsystem(const FObjectInitializer& ObjectInitializer = FObjectInitializer::Get());
@@ -30,7 +29,6 @@ public:
 
 	FMassProcessingPhase::FOnPhaseEvent& GetOnProcessingPhaseStarted(const EMassProcessingPhase Phase) const;
 	FMassProcessingPhase::FOnPhaseEvent& GetOnProcessingPhaseFinished(const EMassProcessingPhase Phase) const;
-	static FOnAdjustTickSchematics& GetOnAdjustTickSchematics() { return OnAdjustTickSchematics; }
 	static FOnSimulationStarted& GetOnSimulationStarted() { return OnSimulationStarted; }
 
 	bool IsSimulationStarted() const { return bSimulationStarted; }
@@ -63,7 +61,6 @@ protected:
 	UPROPERTY()
 	UMassProcessingPhaseManager* PhaseManager;
 
-	inline static FOnAdjustTickSchematics OnAdjustTickSchematics={};
 	inline static FOnSimulationStarted OnSimulationStarted={};
 
 	UPROPERTY()
diff --git a/Engine/Plugins/Runtime/MassGameplay/Source/MassSpawner/Private/MassEntityEQSSpawnPointsGenerator.cpp b/Engine/Plugins/Runtime/MassGameplay/Source/MassSpawner/Private/MassEntityEQSSpawnPointsGenerator.cpp
--- a/Engine/Plugins/Runtime/MassGameplay/Source/MassSpawner/Private/MassEntityEQSSpawnPointsGenerator.cpp
+++ b/Engine/Plugins/Runtime/MassGameplay/Source/MassSpawner/Private/MassEntityEQSSpawnPointsGenerator.cpp
@@ -2,7 +2,9 @@
 
 #include "MassEntityEQSSpawnPointsGenerator.h"
 #include "MassSpawnerTypes.h"
+#include "MassSpawnLocationProcessor.h"
 #include "VisualLogger/VisualLogger.h"
+#include "MassGameplaySettings.h"
 
 
 UMassEntityEQSSpawnPointsGenerator::UMassEntityEQSSpawnPointsGenerator()
@@ -10,7 +12,7 @@ UMassEntityEQSSpawnPointsGenerator::UMassEntityEQSSpawnPointsGenerator()
 	EQSRequest.RunMode = EEnvQueryRunMode::AllMatching; 
 }
 
-void UMassEntityEQSSpawnPointsGenerator::GenerateSpawnPoints(UObject& QueryOwner, int32 Count, FFinishedGeneratingSpawnPointsSignature& FinishedGeneratingSpawnPointsDelegate) const
+void UMassEntityEQSSpawnPointsGenerator::Generate(UObject& QueryOwner, TConstArrayView<FMassSpawnedEntityType> EntityTypes, const int32 Count, FFinishedGeneratingSpawnDataSignature& FinishedGeneratingSpawnPointsDelegate) const
 {
 	// Need to copy the request as it is called inside a CDO and CDO states cannot be changed.
 	FEQSParametrizedQueryExecutionRequest EQSRequestInstanced = EQSRequest;
@@ -23,23 +25,29 @@ void UMassEntityEQSSpawnPointsGenerator::GenerateSpawnPoints(UObject& QueryOwner
 		}
 	}
 
-	FQueryFinishedSignature Delegate = FQueryFinishedSignature::CreateUObject(this, &UMassEntityEQSSpawnPointsGenerator::OnEQSQueryFinished, Count, FinishedGeneratingSpawnPointsDelegate);
+	// Build array of entity types to spawn.
+	// @todo: I dont like that this get's passed by value to OnEQSQueryFinished, but seemed like the cleanest solution.
+	TArray<FMassEntitySpawnDataGeneratorResult> Results;
+	BuildResultsFromEntityTypes(Count, EntityTypes, Results);
+
+	FQueryFinishedSignature Delegate = FQueryFinishedSignature::CreateUObject(this, &UMassEntityEQSSpawnPointsGenerator::OnEQSQueryFinished, Results, FinishedGeneratingSpawnPointsDelegate);
 	EQSRequestInstanced.Execute(QueryOwner, /*BlackboardComponent=*/nullptr, Delegate);
 }
 
-void UMassEntityEQSSpawnPointsGenerator::OnEQSQueryFinished(TSharedPtr<FEnvQueryResult> Result, int32 Count, FFinishedGeneratingSpawnPointsSignature FinishedGeneratingSpawnPointsDelegate) const
+void UMassEntityEQSSpawnPointsGenerator::OnEQSQueryFinished(TSharedPtr<FEnvQueryResult> EQSResult, TArray<FMassEntitySpawnDataGeneratorResult> Results,
+															FFinishedGeneratingSpawnDataSignature FinishedGeneratingSpawnPointsDelegate) const
 {
-	TArray<FVector> Locations;
-
-	if (Result.IsValid() == false || Result->IsSuccessful() == false)
+	if (EQSResult.IsValid() == false || EQSResult->IsSuccessful() == false)
 	{
 		UE_VLOG_UELOG(this, LogMassSpawner, Error, TEXT("EQS query failed or result is invalid"));
-	}
-	else
-	{
-		Result->GetAllAsLocations(Locations);
+		// Return empty result.
+		FinishedGeneratingSpawnPointsDelegate.Execute(Results);
+		return;
 	}
 
+	TArray<FVector> Locations;
+	EQSResult->GetAllAsLocations(Locations);
+
 	// Randomize them
 	FRandomStream RandomStream(GFrameNumber);
 	for (int32 I = 0; I < Locations.Num(); ++I)
@@ -48,11 +56,49 @@ void UMassEntityEQSSpawnPointsGenerator::OnEQSQueryFinished(TSharedPtr<FEnvQuery
 		Locations.Swap(I, J);
 	}
 
-	// If we generated too many, shrink it.
-	if (Locations.Num() > Count)
+	const int32 LocationCount = Locations.Num();
+	int32 LocationIndex = 0;
+
+	// Distribute points amongst the entities to spawn.
+	for (FMassEntitySpawnDataGeneratorResult& Result : Results)
 	{
-		Locations.SetNum(Count);
+		Result.SpawnDataProcessor = UMassSpawnLocationProcessor::StaticClass();
+		Result.SpawnData.InitializeAs<FMassTransformsSpawnData>();
+		FMassTransformsSpawnData& Transforms = Result.SpawnData.GetMutable<FMassTransformsSpawnData>();
+
+		Transforms.Transforms.Reserve(Result.NumEntities);
+		for (int i = 0; i < Result.NumEntities; i++)
+		{
+			FTransform& Transform = Transforms.Transforms.AddDefaulted_GetRef();
+			Transform.SetLocation(Locations[LocationIndex % LocationCount]);
+			LocationIndex++;
+		}
+	}
+
+#if ENABLE_VISUAL_LOG
+	UE_VLOG(this, LogMassSpawner, Log, TEXT("Spawning at %d locations"), LocationIndex);
+	if (GetDefault<UMassGameplaySettings>()->bLogSpawnLocations)
+	{
+		if (FVisualLogEntry* LogEntry = FVisualLogger::Get().GetLastEntryForObject(this))
+		{
+			FVisualLogShapeElement Element(TEXT(""), FColor::Orange, /*Thickness*/20, LogMassSpawner.GetCategoryName());
+
+			Element.Points.Reserve(LocationIndex);
+			for (const FMassEntitySpawnDataGeneratorResult& Result : Results)
+			{
+				const FMassTransformsSpawnData& Transforms = Result.SpawnData.Get<FMassTransformsSpawnData>();
+				for (int i = 0; i < Result.NumEntities; i++)
+				{
+					Element.Points.Add(Transforms.Transforms[i].GetLocation());
+				}
+			}
+			
+			Element.Type = EVisualLoggerShapeElement::SinglePoint;
+			Element.Verbosity = ELogVerbosity::Display;
+			LogEntry->AddElement(Element);
+		}
 	}
+#endif // ENABLE_VISUAL_LOG
 
-	FinishedGeneratingSpawnPointsDelegate.Execute(Locations);
+	FinishedGeneratingSpawnPointsDelegate.Execute(Results);
 }
\ No newline at end of file
diff --git a/Engine/Plugins/Runtime/MassGameplay/Source/MassSpawner/Private/MassEntitySpawnDataGeneratorBase.cpp b/Engine/Plugins/Runtime/MassGameplay/Source/MassSpawner/Private/MassEntitySpawnDataGeneratorBase.cpp
new file mode 100644
--- /dev/null
+++ b/Engine/Plugins/Runtime/MassGameplay/Source/MassSpawner/Private/MassEntitySpawnDataGeneratorBase.cpp
@@ -0,0 +1,31 @@
+﻿// Copyright Epic Games, Inc. All Rights Reserved.
+
+#include "MassEntitySpawnDataGeneratorBase.h"
+#include "MassSpawnerTypes.h"
+
+void UMassEntitySpawnDataGeneratorBase::BuildResultsFromEntityTypes(const int32 SpawnCount, TConstArrayView<FMassSpawnedEntityType> EntityTypes, TArray<FMassEntitySpawnDataGeneratorResult>& OutResults) const
+{
+	float TotalProportion = 0.0f;
+	for (const FMassSpawnedEntityType& EntityType : EntityTypes)
+	{
+		TotalProportion += EntityType.Proportion;
+	}
+
+	if (TotalProportion <= 0)
+	{
+		UE_VLOG_UELOG(this, LogMassSpawner, Error, TEXT("The total combined porportion of all the entity types needs to be greater than 0.0f."));
+		return;
+	}
+
+	for (int32 i = 0; i < EntityTypes.Num(); i++)
+	{
+		const FMassSpawnedEntityType& EntityType = EntityTypes[i];
+		const int32 EntityCount = int32(SpawnCount * EntityType.Proportion / TotalProportion);
+		if (EntityCount > 0 && EntityType.GetEntityConfig() != nullptr)
+		{
+			FMassEntitySpawnDataGeneratorResult& Res = OutResults.AddDefaulted_GetRef();
+			Res.NumEntities = EntityCount;
+			Res.EntityConfigIndex = i;
+		}
+	}
+}
\ No newline at end of file
diff --git a/Engine/Plugins/Runtime/MassGameplay/Source/MassSpawner/Private/MassEntityZoneGraphSpawnPointsGenerator.cpp b/Engine/Plugins/Runtime/MassGameplay/Source/MassSpawner/Private/MassEntityZoneGraphSpawnPointsGenerator.cpp
--- a/Engine/Plugins/Runtime/MassGameplay/Source/MassSpawner/Private/MassEntityZoneGraphSpawnPointsGenerator.cpp
+++ b/Engine/Plugins/Runtime/MassGameplay/Source/MassSpawner/Private/MassEntityZoneGraphSpawnPointsGenerator.cpp
@@ -6,49 +6,93 @@
 #include "ZoneGraphQuery.h"
 #include "MassSpawnerTypes.h"
 #include "VisualLogger/VisualLogger.h"
-#include "Engine/World.h"
+#include "MassGameplaySettings.h"
+#include "MassSpawnLocationProcessor.h"
 
-void UMassEntityZoneGraphSpawnPointsGenerator::GenerateSpawnPoints(UObject& QueryOwner, int32 Count, FFinishedGeneratingSpawnPointsSignature& FinishedGeneratingSpawnPointsDelegate) const
+void UMassEntityZoneGraphSpawnPointsGenerator::Generate(UObject& QueryOwner, TConstArrayView<FMassSpawnedEntityType> EntityTypes, int32 Count, FFinishedGeneratingSpawnDataSignature& FinishedGeneratingSpawnPointsDelegate) const
 {
+	const UZoneGraphSubsystem* ZoneGraph = UWorld::GetSubsystem<UZoneGraphSubsystem>(QueryOwner.GetWorld());
+	if (ZoneGraph == nullptr)
+	{
+		UE_VLOG_UELOG(&QueryOwner, LogMassSpawner, Error, TEXT("No zone graph found in world"));
+		return;
+	}
+
 	TArray<FVector> Locations;
-	if (const UWorld* World = QueryOwner.GetWorld())
+	
+	const FRandomStream RandomStream(GFrameNumber);
+	const TConstArrayView<FRegisteredZoneGraphData> RegisteredZoneGraphs = ZoneGraph->GetRegisteredZoneGraphData();
+	for (const FRegisteredZoneGraphData& Registered : RegisteredZoneGraphs)
 	{
-		if (const UZoneGraphSubsystem* ZoneGraph = UWorld::GetSubsystem<UZoneGraphSubsystem>(World))
+		if (Registered.bInUse && Registered.ZoneGraphData)
 		{
-			const FRandomStream RandomStream(GFrameNumber);
-			const TConstArrayView<FRegisteredZoneGraphData> RegisteredZoneGraphs = ZoneGraph->GetRegisteredZoneGraphData();
-			for (const FRegisteredZoneGraphData& Registered : RegisteredZoneGraphs)
-			{
-				if (Registered.bInUse && Registered.ZoneGraphData)
-				{
-					GeneratePointsForZoneGraphData(*Registered.ZoneGraphData, Locations, RandomStream);
-				}
-			}
+			GeneratePointsForZoneGraphData(*Registered.ZoneGraphData, Locations, RandomStream);
+		}
+	}
 
-			// Randomize them
-			for (int32 I = 0; I < Locations.Num(); ++I)
-			{
-				const int32 J = RandomStream.RandHelper(Locations.Num());
-				Locations.Swap(I, J);
-			}
+	// Randomize them
+	for (int32 I = 0; I < Locations.Num(); ++I)
+	{
+		const int32 J = RandomStream.RandHelper(Locations.Num());
+		Locations.Swap(I, J);
+	}
 
-			// If we generated too many, shrink it.
-			if (Locations.Num() > Count)
-			{
-				Locations.SetNum(Count);
-			}
-		}
-		else
+	// If we generated too many, shrink it.
+	if (Locations.Num() > Count)
+	{
+		Locations.SetNum(Count);
+	}
+
+	// Build array of entity types to spawn.
+	TArray<FMassEntitySpawnDataGeneratorResult> Results;
+	BuildResultsFromEntityTypes(Count, EntityTypes, Results);
+
+	const int32 LocationCount = Locations.Num();
+	int32 LocationIndex = 0;
+
+	// Distribute points amongst the entities to spawn.
+	for (FMassEntitySpawnDataGeneratorResult& Result : Results)
+	{
+		// @todo: Make separate processors and pass the ZoneGraph locations directly.
+		Result.SpawnDataProcessor = UMassSpawnLocationProcessor::StaticClass();
+		Result.SpawnData.InitializeAs<FMassTransformsSpawnData>();
+		FMassTransformsSpawnData& Transforms = Result.SpawnData.GetMutable<FMassTransformsSpawnData>();
+
+		Transforms.Transforms.Reserve(Result.NumEntities);
+		for (int i = 0; i < Result.NumEntities; i++)
 		{
-			UE_VLOG_UELOG(&QueryOwner, LogMassSpawner, Error, TEXT("No zone graph found in world"));
+			FTransform& Transform = Transforms.Transforms.AddDefaulted_GetRef();
+			Transform.SetLocation(Locations[LocationIndex % LocationCount]);
+			LocationIndex++;
 		}
 	}
-	else
+
+#if ENABLE_VISUAL_LOG
+	UE_VLOG(this, LogMassSpawner, Log, TEXT("Spawning at %d locations"), LocationIndex);
+	if (GetDefault<UMassGameplaySettings>()->bLogSpawnLocations)
 	{
-		UE_VLOG_UELOG(&QueryOwner, LogMassSpawner, Error, TEXT("Unable to retrieve world from QueryOwner"));
+		if (FVisualLogEntry* LogEntry = FVisualLogger::Get().GetLastEntryForObject(this))
+		{
+			FVisualLogShapeElement Element(TEXT(""), FColor::Orange, /*Thickness*/20, LogMassSpawner.GetCategoryName());
+
+			Element.Points.Reserve(LocationIndex);
+			for (const FMassEntitySpawnDataGeneratorResult& Result : Results)
+			{
+				const FMassTransformsSpawnData& Transforms = Result.SpawnData.Get<FMassTransformsSpawnData>();
+				for (int i = 0; i < Result.NumEntities; i++)
+				{
+					Element.Points.Add(Transforms.Transforms[i].GetLocation());
+				}
+			}
+			
+			Element.Type = EVisualLoggerShapeElement::SinglePoint;
+			Element.Verbosity = ELogVerbosity::Display;
+			LogEntry->AddElement(Element);
+		}
 	}
+#endif // ENABLE_VISUAL_LOG
 
-	FinishedGeneratingSpawnPointsDelegate.Execute(Locations);
+	FinishedGeneratingSpawnPointsDelegate.Execute(Results);
 }
 
 void UMassEntityZoneGraphSpawnPointsGenerator::GeneratePointsForZoneGraphData(const ::AZoneGraphData& ZoneGraphData, TArray<FVector>& Locations, const FRandomStream& RandomStream) const
diff --git a/Engine/Plugins/Runtime/MassGameplay/Source/MassSpawner/Private/MassSpawnLocationProcessor.cpp b/Engine/Plugins/Runtime/MassGameplay/Source/MassSpawner/Private/MassSpawnLocationProcessor.cpp
--- a/Engine/Plugins/Runtime/MassGameplay/Source/MassSpawner/Private/MassSpawnLocationProcessor.cpp
+++ b/Engine/Plugins/Runtime/MassGameplay/Source/MassSpawner/Private/MassSpawnLocationProcessor.cpp
@@ -22,7 +22,7 @@ void UMassSpawnLocationProcessor::ConfigureQueries()
 
 void UMassSpawnLocationProcessor::Execute(UMassEntitySubsystem& EntitySubsystem, FMassExecutionContext& Context)
 {
-	if (!ensure(Context.ValidateAuxDataType<FMassSpawnAuxData>()))
+	if (!ensure(Context.ValidateAuxDataType<FMassTransformsSpawnData>()))
 	{
 		UE_VLOG_UELOG(this, LogMass, Log, TEXT("Execution context has invalid AuxData or it's not FMassSpawnAuxData. Entity transforms won't be initialized."));
 		return;
@@ -36,7 +36,7 @@ void UMassSpawnLocationProcessor::Execute(UMassEntitySubsystem& EntitySubsystem,
 
 	if (NetMode != NM_Client)
 	{
-		FMassSpawnAuxData& AuxData = Context.GetMutableAuxData().GetMutable<FMassSpawnAuxData>();
+		FMassTransformsSpawnData& AuxData = Context.GetMutableAuxData().GetMutable<FMassTransformsSpawnData>();
 		TArray<FTransform>& Transforms = AuxData.Transforms;
 
 		const int32 NumSpawnTransforms = Transforms.Num();
diff --git a/Engine/Plugins/Runtime/MassGameplay/Source/MassSpawner/Private/MassSpawner.cpp b/Engine/Plugins/Runtime/MassGameplay/Source/MassSpawner/Private/MassSpawner.cpp
--- a/Engine/Plugins/Runtime/MassGameplay/Source/MassSpawner/Private/MassSpawner.cpp
+++ b/Engine/Plugins/Runtime/MassGameplay/Source/MassSpawner/Private/MassSpawner.cpp
@@ -11,11 +11,12 @@
 #include "EnvironmentQuery/EnvQuery.h"
 #include "VisualLogger/VisualLogger.h"
 #include "MassEntityConfigAsset.h"
+#include "MassEntitySubsystem.h"
 #include "EngineUtils.h"
 #include "Engine/StreamableManager.h"
 #include "Engine/AssetManager.h"
-#include "MassGameplaySettings.h"
 #include "MassSpawnLocationProcessor.h"
+#include "MassExecutor.h"
 
 
 namespace UE::MassSpawner
@@ -111,25 +112,15 @@ AMassSpawner::AMassSpawner()
 	bOverrideSchematics = false;
 }
 
-void AMassSpawner::PostInitProperties()
-{
-	Super::PostInitProperties();
-
-	if (HasAnyFlags(RF_ClassDefaultObject) == false)
-	{
-		TickSchematicsAdjustHandle = UMassSimulationSubsystem::GetOnAdjustTickSchematics().AddUObject(this, &AMassSpawner::OnAdjustTickSchematics);
-	}
-}
-
 void AMassSpawner::PostLoad()
 {
 	Super::PostLoad();
 
-	for (FMassSpawnPointGenerator& SpawnPointsGenerator : SpawnPointsGenerators)
+	for (FMassSpawnDataGenerator& SpawnPointsGenerator : SpawnDataGenerators)
 	{
 		if (SpawnPointsGenerator.GeneratorClass)
 		{
-			SpawnPointsGenerator.GeneratorInstance = NewObject<UMassEntitySpawnPointsGeneratorBase>(this, SpawnPointsGenerator.GeneratorClass);
+			SpawnPointsGenerator.GeneratorInstance = NewObject<UMassEntitySpawnDataGeneratorBase>(this, SpawnPointsGenerator.GeneratorClass);
 			SpawnPointsGenerator.GeneratorClass = nullptr;
 			MarkPackageDirty();
 		}
@@ -180,7 +171,6 @@ void AMassSpawner::OnPostWorldInit(UWorld* World, const UWorld::InitializationVa
 void AMassSpawner::BeginDestroy()
 {
 	FWorldDelegates::OnPostWorldInitialization.Remove(OnPostWorldInitDelegateHandle);
-	UMassSimulationSubsystem::GetOnAdjustTickSchematics().Remove(TickSchematicsAdjustHandle);
 
 	DoDespawning();
 	
@@ -273,117 +263,119 @@ void AMassSpawner::RegisterEntityTemplates()
 
 void AMassSpawner::DoSpawning()
 {
-	const int32 SpawnCount = GetSpawnCount();
-	if (SpawnCount > 0)
+	const int32 TotalSpawnCount = GetSpawnCount();
+	if (TotalSpawnCount <= 0)
 	{
-		// no spawn point generators configured. Let user know and fall back to the spawner's location
-		if (SpawnPointsGenerators.Num() == 0)
+		return;
+	}
+	
+	// no spawn point generators configured. Let user know and fall back to the spawner's location
+	if (SpawnDataGenerators.Num() == 0)
+	{
+		UE_VLOG_UELOG(this, LogMassSpawner, Warning, TEXT("No Spawn Data Generators configured."));
+		return;
+	}
+
+	AllGeneratedResults.Reset();
+	
+	float TotalProportion = 0.0f;
+	for (FMassSpawnDataGenerator& Generator : SpawnDataGenerators)
+	{
+		if (Generator.GeneratorInstance)
 		{
-			UE_VLOG_UELOG(this, LogMassSpawner, Warning, TEXT("No Spawn Points Generators configured. Falling back to the MassSpawner\'s location."));
-			const FVector SpawnerLocation = GetActorLocation();
-			SpawnAtLocations(MakeArrayView(&SpawnerLocation, 1));
-			return;
+			Generator.bDataGenerated = false;
+			TotalProportion += Generator.Proportion;
 		}
+	}
 
-		float TotalProportion = 0.0f;
-		for (FMassSpawnPointGenerator& SpawnPointsGenerator : SpawnPointsGenerators)
+	if (TotalProportion <= 0.0f)
+	{
+		UE_VLOG_UELOG(this, LogMassSpawner, Error, TEXT("The total combined porportion of all the generator needs to be greater than 0.0f."));
+		return;
+	}
+	
+	// Check if it needs loading
+	if (StreamingHandle.IsValid() && StreamingHandle->IsActive())
+	{
+		// @todo, instead of blindly canceling, we should remember what was asked to load with that handle and compare if more is needed?
+		StreamingHandle->CancelHandle();
+	}
+	TArray<FSoftObjectPath> AssetsToLoad;
+	for (const FMassSpawnedEntityType& EntityType : EntityTypes)
+	{
+		const int32 EntityCount = int32(TotalSpawnCount * EntityType.Proportion / TotalProportion);
+		if (EntityCount > 0)
 		{
-			if (SpawnPointsGenerator.GeneratorInstance)
+			if (!EntityType.IsLoaded())
 			{
-				SpawnPointsGenerator.bPointsGenerated = false;
-				TotalProportion += SpawnPointsGenerator.Proportion;
+				AssetsToLoad.Add(EntityType.EntityConfig.ToSoftObjectPath());
 			}
 		}
+	}
 
-		if (TotalProportion <= 0.0f)
-		{
-			UE_VLOG_UELOG(this, LogMassSpawner, Error, TEXT("The total combined porportion of all the generator needs to be greater than 0.0f."));
-			return;
-		}
-		// Check if it needs loading
-		if (StreamingHandle.IsValid() && StreamingHandle->IsActive())
-		{
-			// @todo, instead of blindly canceling, we should remember what was asked to load with that handle and compare if more is needed?
-			StreamingHandle->CancelHandle();
-		}
-		TArray<FSoftObjectPath> AssetsToLoad;
-		for (const FMassSpawnedEntityType& EntityType : EntityTypes)
+	auto GenerateSpawningPoints = [this, TotalSpawnCount, TotalProportion]()
+	{
+		int32 SpawnCountRemaining = TotalSpawnCount;
+		float ProportionRemaining = TotalProportion;
+		for (FMassSpawnDataGenerator& Generator : SpawnDataGenerators)
 		{
-			const int32 EntityCount = int32(SpawnCount * EntityType.Proportion / TotalProportion);
-			if (EntityCount > 0)
-
+			if (Generator.Proportion == 0.0f || ProportionRemaining <= 0.0f || SpawnCountRemaining <= 0)
 			{
-				if (!EntityType.IsLoaded())
-				{
-					AssetsToLoad.Add(EntityType.EntityConfig.ToSoftObjectPath());
-				}
+				// If there's nothing to spawn, mark the generator done as OnSpawnDataGenerationFinished() will wait for all generators to complete before the actual spawning.
+				Generator.bDataGenerated = true;
+				continue;
 			}
-		}
 
-		auto GenerateSpawningPoints = [this, SpawnCount, TotalProportion]()
-		{
-			int32 SpawnPointCountRemaining = SpawnCount;
-			float ProportionRemaining = TotalProportion;
-			for (FMassSpawnPointGenerator& SpawnPointsGenerator : SpawnPointsGenerators)
+			if (Generator.GeneratorInstance)
 			{
-				// Still needs to call the OnSpawnPointGenerationFinished for Generator we don't want to output Spawn Points for as the system will wait
-				// for the bPointsGenerated == true on every Generators before proceeding to the next steps.
-				if (SpawnPointsGenerator.Proportion == 0.0f || ProportionRemaining <= 0.0f || SpawnPointCountRemaining <= 0)
+				const float ProportionRatio = FMath::Min(Generator.Proportion / ProportionRemaining, 1.0f);
+				const int32 SpawnCount = FMath::CeilToInt(SpawnCountRemaining * ProportionRatio);
+				if (SpawnCount > 0)
 				{
-					const TArray<FVector> EmptyLoc;
-					OnSpawnPointGenerationFinished(EmptyLoc, &SpawnPointsGenerator);
-					continue;
-				}
-
-				if (SpawnPointsGenerator.GeneratorInstance)
-				{
-					const float ProportionRatio = FMath::Min(SpawnPointsGenerator.Proportion / ProportionRemaining, 1.0f);
-					const int32 SpawnPointCount = FMath::CeilToInt(SpawnPointCountRemaining * ProportionRatio);
-					if (SpawnPointCount > 0)
-					{
-						FFinishedGeneratingSpawnPointsSignature Delegate = FFinishedGeneratingSpawnPointsSignature::CreateUObject(this, &AMassSpawner::OnSpawnPointGenerationFinished, &SpawnPointsGenerator);
-						SpawnPointsGenerator.GeneratorInstance->GenerateSpawnPoints(*this, SpawnPointCount, Delegate);
-						SpawnPointCountRemaining -= SpawnPointCount;
-						ProportionRemaining -= SpawnPointsGenerator.Proportion;
-					}
+					FFinishedGeneratingSpawnDataSignature Delegate = FFinishedGeneratingSpawnDataSignature::CreateUObject(this, &AMassSpawner::OnSpawnDataGenerationFinished, &Generator);
+					Generator.GeneratorInstance->Generate(*this, EntityTypes, SpawnCount, Delegate);
+					SpawnCountRemaining -= SpawnCount;
+					ProportionRemaining -= Generator.Proportion;
 				}
 			}
-		};
-		
-		if (AssetsToLoad.Num())
-		{
-			FStreamableManager& StreamableManager = UAssetManager::GetStreamableManager();
-			StreamingHandle = StreamableManager.RequestAsyncLoad(AssetsToLoad, GenerateSpawningPoints);
-		}
-		else
-		{
-			GenerateSpawningPoints();
 		}
+	};
+	
+	if (AssetsToLoad.Num())
+	{
+		FStreamableManager& StreamableManager = UAssetManager::GetStreamableManager();
+		StreamingHandle = StreamableManager.RequestAsyncLoad(AssetsToLoad, GenerateSpawningPoints);
+	}
+	else
+	{
+		GenerateSpawningPoints();
 	}
 }
 
-void AMassSpawner::OnSpawnPointGenerationFinished(const TArray<FVector>& Locations, FMassSpawnPointGenerator* FinishedGenerator)
+void AMassSpawner::OnSpawnDataGenerationFinished(TConstArrayView<FMassEntitySpawnDataGeneratorResult> Results, FMassSpawnDataGenerator* FinishedGenerator)
 {
-	AllGeneratedLocations.Append(Locations);
+	// @todo: this can be potentially expensive copy for the instanced structs, could there be a way to use move gere instead?
+	AllGeneratedResults.Append(Results.GetData(), Results.Num());
 
 	bool bAllSpawnPointsGenerated = true;
 	bool bFoundFinishedGenerator = false;
-	for (FMassSpawnPointGenerator& SpawnPointsGenerator : SpawnPointsGenerators)
+	for (FMassSpawnDataGenerator& Generator : SpawnDataGenerators)
 	{
-		if (&SpawnPointsGenerator == FinishedGenerator)
+		if (&Generator == FinishedGenerator)
 		{
-			SpawnPointsGenerator.bPointsGenerated = true;
+			Generator.bDataGenerated = true;
 			bFoundFinishedGenerator = true;
 		}
 
-		bAllSpawnPointsGenerated &= SpawnPointsGenerator.bPointsGenerated;
+		bAllSpawnPointsGenerated &= Generator.bDataGenerated;
 	}
 
 	checkf(bFoundFinishedGenerator, TEXT("Something went wrong, we are receiving a callback on an unknow spawn point generator"));
 	
 	if (bAllSpawnPointsGenerated)
 	{
-		SpawnAtLocations(AllGeneratedLocations);
+		SpawnGeneratedEntities(AllGeneratedResults);
+		AllGeneratedResults.Reset();
 	}
 }
 
@@ -393,75 +385,92 @@ int32 AMassSpawner::GetSpawnCount() const
 	return int32(FinalSpawningCountScale * Count);
 }
 
-void AMassSpawner::SpawnAtLocations(TConstArrayView<FVector> Locations)
+UMassProcessor* AMassSpawner::GetPostSpawnProcessor(TSubclassOf<UMassProcessor> ProcessorClass)
 {
-	UMassSpawnerSubsystem* SpawnerSystem = UWorld::GetSubsystem<UMassSpawnerSubsystem>(GetWorld());
-	if (SpawnerSystem == nullptr)
+	if (!ProcessorClass)
 	{
-		UE_VLOG_UELOG(this, LogMassSpawner, Error, TEXT("UMassSpawnerSubsystem missing while trying to spawn entities"));
-		return;
+		return nullptr;
 	}
 
-	FMassSpawnAuxData MassSpawnAuxData;
-	MassSpawnAuxData.Transforms.Reserve(Locations.Num());
-	for (const FVector& Location : Locations)
-	{
-		FTransform& Transform = MassSpawnAuxData.Transforms.AddDefaulted_GetRef();
+	TObjectPtr<UMassProcessor>* const Initializer = PostSpawnProcessors.FindByPredicate([ProcessorClass](const UMassProcessor* Processor)
+		{
+			return Processor && Processor->GetClass() == ProcessorClass;
+		}
+	);
 
-		Transform.SetLocation(Location);
+	if (Initializer == nullptr)
+	{
+		UMassProcessor* NewInitializer = NewObject<UMassProcessor>(this, ProcessorClass);
+		NewInitializer->Initialize(*this);
+		PostSpawnProcessors.Add(NewInitializer);
+		return NewInitializer;
 	}
 
-#if ENABLE_VISUAL_LOG
-	UE_VLOG(this, LogMassSpawner, Log, TEXT("Spawning at %d locations"), Locations.Num());
-	if (GetDefault<UMassGameplaySettings>()->bLogSpawnLocations)
+	return *Initializer;
+}
+
+void AMassSpawner::SpawnGeneratedEntities(TConstArrayView<FMassEntitySpawnDataGeneratorResult> Results)
+{
+	UMassSpawnerSubsystem* SpawnerSystem = UWorld::GetSubsystem<UMassSpawnerSubsystem>(GetWorld());
+	if (SpawnerSystem == nullptr)
 	{
-		if (FVisualLogEntry* LogEntry = FVisualLogger::Get().GetLastEntryForObject(this))
-		{
-			FVisualLogShapeElement Element(TEXT(""), FColor::Orange, /*Thickness*/20, LogMassSpawner.GetCategoryName());
-			Element.Points.Append(Locations.GetData(), Locations.Num());
-			Element.Type = EVisualLoggerShapeElement::SinglePoint;
-			Element.Verbosity = ELogVerbosity::Display;
-			LogEntry->AddElement(Element);
-		}
+		UE_VLOG_UELOG(this, LogMassSpawner, Error, TEXT("UMassSpawnerSubsystem missing while trying to spawn entities"));
+		return;
 	}
-#endif // ENABLE_VISUAL_LOG
 
-	const int32 SpawnCount = GetSpawnCount();
-	if (SpawnCount > 0)
+	for (const FMassEntitySpawnDataGeneratorResult& Result : Results)
 	{
-		float TotalProportion = 0.0f;
-		for (const FMassSpawnedEntityType& EntityType : EntityTypes)
+		if (Result.NumEntities <= 0)
 		{
-			TotalProportion += EntityType.Proportion;
+			continue;
 		}
+		
+		check(EntityTypes.IsValidIndex(Result.EntityConfigIndex));
+		check(Result.SpawnDataProcessor != nullptr);
+		
+		const FMassSpawnedEntityType& EntityType = EntityTypes[Result.EntityConfigIndex];
 
-		if (TotalProportion <= 0)
+		if (const UMassEntityConfigAsset* EntityConfig = EntityType.GetEntityConfig())
 		{
-			UE_VLOG_UELOG(this, LogMassSpawner, Error, TEXT("The total combined porportion of all the entity types needs to be greater than 0.0f."));
-			return;
+			const FMassEntityTemplate* EntityTemplate = EntityConfig->GetConfig().GetOrCreateEntityTemplate(*this, *EntityConfig);
+			if (EntityTemplate && EntityTemplate->IsValid())
+			{
+				FSpawnedEntities& SpawnedEntities = AllSpawnedEntities.AddDefaulted_GetRef();
+				SpawnedEntities.TemplateID = EntityTemplate->GetTemplateID();
+				SpawnerSystem->SpawnEntities(EntityTemplate->GetTemplateID(), Result.NumEntities, Result.SpawnData, Result.SpawnDataProcessor, SpawnedEntities.Entities);
+			}
 		}
+	}
 
-		AllSpawnedEntities.Reserve(AllSpawnedEntities.Num() + EntityTypes.Num());
+	// Run post spawn processors on all Mass entities that matches the queries.
+	// @todo: we might need a way to specify that these are ran only on the freshly spawned entities.
+	
+	TArray<UMassProcessor*> Processors;
+	TSet<TSubclassOf<UMassProcessor>> AddedProcessors;
 
-		for (const FMassSpawnedEntityType& EntityType : EntityTypes)
+	for (const FMassEntitySpawnDataGeneratorResult& Result : Results)
+	{
+		for (const TSubclassOf<UMassProcessor> ProcessorClass : Result.PostSpawnProcessors)
 		{
-			const int32 EntityCount = int32(SpawnCount * EntityType.Proportion / TotalProportion);
-			if (EntityCount > 0)
+			if (AddedProcessors.Contains(ProcessorClass) == false)
 			{
-				if (const UMassEntityConfigAsset* EntityConfig = EntityType.GetEntityConfig())
+				if (UMassProcessor* Processor = GetPostSpawnProcessor(ProcessorClass))
 				{
-					const FMassEntityTemplate* EntityTemplate = EntityConfig->GetConfig().GetOrCreateEntityTemplate(*this, *EntityConfig);
-					if (EntityTemplate && EntityTemplate->IsValid())
-					{
-						FSpawnedEntities& SpawnedEntities = AllSpawnedEntities.AddDefaulted_GetRef();
-						SpawnedEntities.TemplateID = EntityTemplate->GetTemplateID();
-						SpawnerSystem->SpawnEntities(EntityTemplate->GetTemplateID(), EntityCount, FStructView::Make(MassSpawnAuxData)
-							, UMassSpawnLocationProcessor::StaticClass(), SpawnedEntities.Entities);
-					}
+					Processors.Add(Processor);
 				}
+				AddedProcessors.Add(ProcessorClass);
 			}
 		}
 	}
+
+	if (Processors.Num() > 0)
+	{
+		UMassEntitySubsystem* EntitySystem = UWorld::GetSubsystem<UMassEntitySubsystem>(GetWorld());
+		check(EntitySystem);
+
+		FMassProcessingContext ProcessingContext(*EntitySystem, /*TimeDelta=*/0.0f);
+		UE::Mass::Executor::RunProcessorsView(Processors, ProcessingContext);
+	}
 }
 
 void AMassSpawner::DoDespawning()
@@ -533,24 +542,3 @@ void AMassSpawner::UnloadConfig()
 		StreamingHandle->CancelHandle();
 	}
 }
-
-void AMassSpawner::OnAdjustTickSchematics(UWorld* World,TArray<TSoftObjectPtr<UMassSchematic>>& InOutTickSchematics)
-{
-	// this function gets called from a static multicast delegate, so we need to differentiate based on the World it's being called for.
-	if (World != GetWorld())
-	{
-		return;
-	}
-
-	if (bOverrideSchematics)
-	{
-		UE_VLOG_UELOG(World, LogMassSpawner, Log, TEXT("Overriding TickSchematics with %s"), *GetName());
-		InOutTickSchematics = TickSchematics;
-	}
-	else
-	{
-		UE_VLOG_UELOG(World, LogMassSpawner, Log, TEXT("Appending to TickSchematics with %s"), *GetName());
-		InOutTickSchematics.Append(TickSchematics);
-	}
-}
-
diff --git a/Engine/Plugins/Runtime/MassGameplay/Source/MassSpawner/Private/MassSpawnerSubsystem.cpp b/Engine/Plugins/Runtime/MassGameplay/Source/MassSpawner/Private/MassSpawnerSubsystem.cpp
--- a/Engine/Plugins/Runtime/MassGameplay/Source/MassSpawner/Private/MassSpawnerSubsystem.cpp
+++ b/Engine/Plugins/Runtime/MassGameplay/Source/MassSpawner/Private/MassSpawnerSubsystem.cpp
@@ -125,14 +125,14 @@ void UMassSpawnerSubsystem::DestroyEntities(const FMassEntityTemplateID Template
 	}
 }
 
-UMassProcessor* UMassSpawnerSubsystem::GetSpawnLocationInitializer(TSubclassOf<UMassProcessor> InitializerClass)
+UMassProcessor* UMassSpawnerSubsystem::GetSpawnDataInitializer(TSubclassOf<UMassProcessor> InitializerClass)
 {	
 	if (!InitializerClass)
 	{
 		return nullptr;
 	}
 
-	UMassProcessor** const Initializer = SpawnLocationInitializers.FindByPredicate([InitializerClass](const UMassProcessor* Processor)
+	TObjectPtr<UMassProcessor>* const Initializer = SpawnDataInitializers.FindByPredicate([InitializerClass](const UMassProcessor* Processor)
 		{
 			return Processor && Processor->GetClass() == InitializerClass;
 		}
@@ -142,7 +142,7 @@ UMassProcessor* UMassSpawnerSubsystem::GetSpawnLocationInitializer(TSubclassOf<U
 	{
 		UMassProcessor* NewInitializer = NewObject<UMassProcessor>(this, InitializerClass);
 		NewInitializer->Initialize(*this);
-		SpawnLocationInitializers.Add(NewInitializer);
+		SpawnDataInitializers.Add(NewInitializer);
 		return NewInitializer;
 	}
 
@@ -166,7 +166,7 @@ void UMassSpawnerSubsystem::DoSpawning(const FMassEntityTemplate& EntityTemplate
 	// 1. Create required number of entities with EntityTemplate.Archetype
 	// 2. Copy data from FMassEntityTemplate.Fragments.
 	//		a. @todo, could be done as part of creation?
-	// 3. Run SpawlLocationInitializer if set
+	// 3. Run SpawlDataInitializer if set
 	// 4. "OnEntitiesCreated" notifies will be sent out once the CreationContext gets destroyed (via its destructor).
 
 	TArray<FMassEntityHandle> SpawnedEntities;
@@ -175,13 +175,13 @@ void UMassSpawnerSubsystem::DoSpawning(const FMassEntityTemplate& EntityTemplate
 	TConstArrayView<FInstancedStruct> FragmentInstances = EntityTemplate.GetInitialFragmentValues();
 	EntitySystem->BatchSetEntityFragmentsValues(CreationContext->GetChunkCollection(), FragmentInstances);
 	
-	UMassProcessor* SpawnLocationInitializer = SpawnData.IsValid() ? GetSpawnLocationInitializer(InitializerClass) : nullptr;
+	UMassProcessor* SpawnDataInitializer = SpawnData.IsValid() ? GetSpawnDataInitializer(InitializerClass) : nullptr;
 
-	if (SpawnLocationInitializer)
+	if (SpawnDataInitializer)
 	{
 		FMassProcessingContext ProcessingContext(*EntitySystem, /*TimeDelta=*/0.0f);
 		ProcessingContext.AuxData = SpawnData;
-		UE::Mass::Executor::RunProcessorsView(MakeArrayView(&SpawnLocationInitializer, 1), ProcessingContext, &CreationContext->GetChunkCollection());
+		UE::Mass::Executor::RunProcessorsView(MakeArrayView(&SpawnDataInitializer, 1), ProcessingContext, &CreationContext->GetChunkCollection());
 	}
 
 	OutEntities.Append(MoveTemp(SpawnedEntities));
diff --git a/Engine/Plugins/Runtime/MassGameplay/Source/MassSpawner/Public/MassEntityEQSSpawnPointsGenerator.h b/Engine/Plugins/Runtime/MassGameplay/Source/MassSpawner/Public/MassEntityEQSSpawnPointsGenerator.h
--- a/Engine/Plugins/Runtime/MassGameplay/Source/MassSpawner/Public/MassEntityEQSSpawnPointsGenerator.h
+++ b/Engine/Plugins/Runtime/MassGameplay/Source/MassSpawner/Public/MassEntityEQSSpawnPointsGenerator.h
@@ -3,7 +3,7 @@
 #pragma once 
 
 #include "CoreMinimal.h"
-#include "MassEntitySpawnPointsGeneratorBase.h"
+#include "MassEntitySpawnDataGeneratorBase.h"
 #include "EnvironmentQuery/EnvQueryTypes.h"
 #include "MassEntityEQSSpawnPointsGenerator.generated.h"
 
@@ -11,17 +11,17 @@
  * Describes the SpawnPoints Generator when we want to leverage the points given by an EQS Query
  */
 UCLASS(BlueprintType, meta=(DisplayName="EQS SpawnPoints Generator"))
-class MASSSPAWNER_API UMassEntityEQSSpawnPointsGenerator : public UMassEntitySpawnPointsGeneratorBase
+class MASSSPAWNER_API UMassEntityEQSSpawnPointsGenerator : public UMassEntitySpawnDataGeneratorBase
 {	
 	GENERATED_BODY()
 
 public:
 	UMassEntityEQSSpawnPointsGenerator();
 	
-	virtual void GenerateSpawnPoints(UObject& QueryOwner, int32 Count, FFinishedGeneratingSpawnPointsSignature& FinishedGeneratingSpawnPointsDelegate) const override;
+	virtual void Generate(UObject& QueryOwner, TConstArrayView<FMassSpawnedEntityType> EntityTypes, int32 Count, FFinishedGeneratingSpawnDataSignature& FinishedGeneratingSpawnPointsDelegate) const override;
 
 protected:
-	void OnEQSQueryFinished(TSharedPtr<FEnvQueryResult> Result, int32 Count, FFinishedGeneratingSpawnPointsSignature FinishedGeneratingSpawnPointsDelegate) const;
+	void OnEQSQueryFinished(TSharedPtr<FEnvQueryResult> EQSResult, TArray<FMassEntitySpawnDataGeneratorResult> Results, FFinishedGeneratingSpawnDataSignature FinishedGeneratingSpawnPointsDelegate) const;
 
 	UPROPERTY(Category = "Query", EditAnywhere)
 	FEQSParametrizedQueryExecutionRequest EQSRequest;
diff --git a/Engine/Plugins/Runtime/MassGameplay/Source/MassSpawner/Public/MassEntitySpawnDataGeneratorBase.h b/Engine/Plugins/Runtime/MassGameplay/Source/MassSpawner/Public/MassEntitySpawnDataGeneratorBase.h
new file mode 100644
--- /dev/null
+++ b/Engine/Plugins/Runtime/MassGameplay/Source/MassSpawner/Public/MassEntitySpawnDataGeneratorBase.h
@@ -0,0 +1,69 @@
+// Copyright Epic Games, Inc. All Rights Reserved.
+
+#pragma once 
+
+#include "CoreMinimal.h"
+#include "InstancedStruct.h"
+#include "MassSpawnerTypes.h"
+#include "MassEntitySpawnDataGeneratorBase.generated.h"
+
+class UMassEntityConfigAsset;
+class UMassProcessor;
+
+/**
+ * The result of the spawn point generator.
+ */
+USTRUCT()
+struct FMassEntitySpawnDataGeneratorResult
+{
+	GENERATED_BODY()
+
+	// Spawn data that is passed to the InitSpawnDataProcessor. E.g. the data could contain array of locations, one for each entity.
+	UPROPERTY()
+	FInstancedStruct SpawnData;
+	
+	// Processor that understands how to apply SpawnData to the spawned entities. 
+	UPROPERTY()
+	TSubclassOf<UMassProcessor> SpawnDataProcessor;
+
+	// Processors that are run for all entities after they entities for a MassSpawner are initialized.
+	UPROPERTY()
+	TArray<TSubclassOf<UMassProcessor>> PostSpawnProcessors;
+
+	// Index in the EntityTypes array passed to Generate().
+	UPROPERTY()
+	int32 EntityConfigIndex = INDEX_NONE;
+
+	// Number of entities to spawn.
+	UPROPERTY()
+	int32 NumEntities = 0;
+};
+
+
+DECLARE_DELEGATE_OneParam(FFinishedGeneratingSpawnDataSignature, TConstArrayView<FMassEntitySpawnDataGeneratorResult>);
+
+/**
+ * Base class for Mass Entity Spawn Points Generator.
+ * A Mass Spawn Points Generator can be of several type (EQS, ZoneGraph, Volume, Area, etc.)
+ * The concept is to override the GenerateSpawnPoints() method and requesting a certain number of Spawn Point Locations to the method.
+ */
+UCLASS(Abstract, Blueprintable, EditInlineNew)
+class MASSSPAWNER_API UMassEntitySpawnDataGeneratorBase : public UObject
+{
+	GENERATED_BODY()
+
+public:
+
+	/** Generate "Count" number of SpawnPoints and return as a list of position
+	 * @param Count of point to generate
+	 * @param FinishedGeneratingSpawnPointsDelegate is the callback to call once the generation is done
+	 */
+	virtual void Generate(UObject& QueryOwner, TConstArrayView<FMassSpawnedEntityType> EntityTypes, int32 Count, FFinishedGeneratingSpawnDataSignature& FinishedGeneratingSpawnPointsDelegate) const PURE_VIRTUAL(UMassEntitySpawnDataGeneratorBase::GenerateSpawnPoints, );
+
+	/** Populates empty generator results from EntityTypes based on the provided proportions. 
+	 * @param SpawnCount How many entities to distributed among the EntityTypes
+	 * @param EntityTypes Types of entities to generate data for.
+	 * @param OutResults Generator result for each entity type that had > 0 entities assigned to it.
+	 */
+	void BuildResultsFromEntityTypes(const int32 SpawnCount, TConstArrayView<FMassSpawnedEntityType> EntityTypes, TArray<FMassEntitySpawnDataGeneratorResult>& OutResults) const;
+};
diff --git a/Engine/Plugins/Runtime/MassGameplay/Source/MassSpawner/Public/MassEntitySpawnPointsGeneratorBase.h b/Engine/Plugins/Runtime/MassGameplay/Source/MassSpawner/Public/MassEntitySpawnPointsGeneratorBase.h
deleted file mode 100644
diff --git a/Engine/Plugins/Runtime/MassGameplay/Source/MassSpawner/Public/MassEntityZoneGraphSpawnPointsGenerator.h b/Engine/Plugins/Runtime/MassGameplay/Source/MassSpawner/Public/MassEntityZoneGraphSpawnPointsGenerator.h
--- a/Engine/Plugins/Runtime/MassGameplay/Source/MassSpawner/Public/MassEntityZoneGraphSpawnPointsGenerator.h
+++ b/Engine/Plugins/Runtime/MassGameplay/Source/MassSpawner/Public/MassEntityZoneGraphSpawnPointsGenerator.h
@@ -3,7 +3,7 @@
 #pragma once 
 
 #include "CoreMinimal.h"
-#include "MassEntitySpawnPointsGeneratorBase.h"
+#include "MassEntitySpawnDataGeneratorBase.h"
 #include "ZoneGraphTypes.h"
 
 #include "MassEntityZoneGraphSpawnPointsGenerator.generated.h"
@@ -14,12 +14,12 @@ class AZoneGraphData;
  * Describes the SpawnPoints Generator when we want to spawn directly on Zone Graph
  */
 UCLASS(BlueprintType, meta=(DisplayName="ZoneGraph SpawnPoints Generator"))
-class MASSSPAWNER_API UMassEntityZoneGraphSpawnPointsGenerator : public UMassEntitySpawnPointsGeneratorBase
+class MASSSPAWNER_API UMassEntityZoneGraphSpawnPointsGenerator : public UMassEntitySpawnDataGeneratorBase
 {
 	GENERATED_BODY()
 
 public:
-	virtual void GenerateSpawnPoints(UObject& QueryOwner, int32 Count, FFinishedGeneratingSpawnPointsSignature& FinishedGeneratingSpawnPointsDelegate) const override;
+	virtual void Generate(UObject& QueryOwner, TConstArrayView<FMassSpawnedEntityType> EntityTypes, int32 Count, FFinishedGeneratingSpawnDataSignature& FinishedGeneratingSpawnPointsDelegate) const override;
 
 protected:
 	void GeneratePointsForZoneGraphData(const ::AZoneGraphData& ZoneGraphData, TArray<FVector>& Locations, const FRandomStream& RandomStream) const;
diff --git a/Engine/Plugins/Runtime/MassGameplay/Source/MassSpawner/Public/MassSpawner.h b/Engine/Plugins/Runtime/MassGameplay/Source/MassSpawner/Public/MassSpawner.h
--- a/Engine/Plugins/Runtime/MassGameplay/Source/MassSpawner/Public/MassSpawner.h
+++ b/Engine/Plugins/Runtime/MassGameplay/Source/MassSpawner/Public/MassSpawner.h
@@ -3,13 +3,11 @@
 #pragma once 
 
 #include "GameFramework/Actor.h"
-#include "EnvironmentQuery/EnvQueryTypes.h"
-#include "InstancedStruct.h"
 #include "Engine/World.h"
 #include "MassSpawnerTypes.h"
+#include "MassEntitySpawnDataGeneratorBase.h"
 #include "MassSpawner.generated.h"
 
-class UMassEntitySpawnPointsGeneratorConfigAsset;
 class UMassSchematic;
 struct FStreamableHandle;
 
@@ -24,7 +22,6 @@ public:
 protected:
 	virtual void BeginPlay() override;
 	virtual void EndPlay(const EEndPlayReason::Type EndPlayReason) override;
-	virtual void PostInitProperties() override;
 	virtual void PostLoad() override;
 	virtual void PostRegisterAllComponents() override;
 	virtual void BeginDestroy() override;
@@ -80,12 +77,12 @@ public:
 	void ScaleSpawningCount(float Scale) { SpawningCountScale = Scale; }
 
 protected:
-	void OnAdjustTickSchematics(UWorld* World, TArray<TSoftObjectPtr<UMassSchematic>>& InOutTickSchematics);
 	void OnPostWorldInit(UWorld* World, const UWorld::InitializationValues);
-	void SpawnAtLocations(TConstArrayView<FVector> Locations);
-	void OnSpawnPointGenerationFinished(const TArray<FVector>& Locations, FMassSpawnPointGenerator* FinishedGenerator);
+	void SpawnGeneratedEntities(TConstArrayView<FMassEntitySpawnDataGeneratorResult> Results);
+	void OnSpawnDataGenerationFinished(TConstArrayView<FMassEntitySpawnDataGeneratorResult> Results, FMassSpawnDataGenerator* FinishedGenerator);
 
 	int32 GetSpawnCount() const;
+	UMassProcessor* GetPostSpawnProcessor(TSubclassOf<UMassProcessor> ProcessorClass);
 
 protected:
 
@@ -98,12 +95,13 @@ protected:
 	UPROPERTY(EditAnywhere, Category = "Mass|Spawn")
 	int32 Count;
 
+	/** Array of entity types to spawn. These define which entities to spawn. */
 	UPROPERTY(BlueprintReadWrite, EditAnywhere, Category = "Mass|Spawn")
 	TArray<FMassSpawnedEntityType> EntityTypes;
 
-	/** Asset that describes the way we want to generate SpawnPoints of the entity */
+	/** Array of entity spawn generators. These define where to spawn entities. */
 	UPROPERTY(EditAnywhere, Category = "Mass|Spawn")
-	TArray<FMassSpawnPointGenerator> SpawnPointsGenerators;
+	TArray<FMassSpawnDataGenerator> SpawnDataGenerators;
 
 	UPROPERTY(Category = "Mass|Spawn", EditAnywhere)
 	uint32 bAutoSpawnOnBeginPlay : 1;
@@ -113,22 +111,20 @@ protected:
 	UPROPERTY(Category = "Mass|Simulation", EditAnywhere)
 	uint32 bOverrideSchematics : 1;
 
-	UPROPERTY(Category = "Mass|Simulation", EditAnywhere)
-	TArray<TSoftObjectPtr<UMassSchematic>> TickSchematics;
+	UPROPERTY()
+	TArray<TObjectPtr<UMassProcessor>> PostSpawnProcessors;
 
 	/** Scale of the spawning count */
 	float SpawningCountScale = 1.0f;
 
-	FDelegateHandle TickSchematicsAdjustHandle;
-
 	FDelegateHandle SimulationStartedHandle;
 
 	FDelegateHandle OnPostWorldInitDelegateHandle;
 
 	TArray<FSpawnedEntities> AllSpawnedEntities;
-
-	TArray<FVector> AllGeneratedLocations;
-
+	
+	TArray<FMassEntitySpawnDataGeneratorResult> AllGeneratedResults;
+	
 	TSharedPtr<FStreamableHandle> StreamingHandle;
 
 #if WITH_EDITORONLY_DATA
diff --git a/Engine/Plugins/Runtime/MassGameplay/Source/MassSpawner/Public/MassSpawnerSubsystem.h b/Engine/Plugins/Runtime/MassGameplay/Source/MassSpawner/Public/MassSpawnerSubsystem.h
--- a/Engine/Plugins/Runtime/MassGameplay/Source/MassSpawner/Public/MassSpawnerSubsystem.h
+++ b/Engine/Plugins/Runtime/MassGameplay/Source/MassSpawner/Public/MassSpawnerSubsystem.h
@@ -47,10 +47,10 @@ protected:
 
 	void DoSpawning(const FMassEntityTemplate& EntityTemplate, const int32 NumToSpawn, FConstStructView SpawnData, TSubclassOf<UMassProcessor> InitializerClass, TArray<FMassEntityHandle>& OutEntities);
 
-	UMassProcessor* GetSpawnLocationInitializer(TSubclassOf<UMassProcessor> InitializerClass);
+	UMassProcessor* GetSpawnDataInitializer(TSubclassOf<UMassProcessor> InitializerClass);
 
 	UPROPERTY()
-	TArray<UMassProcessor*> SpawnLocationInitializers;
+	TArray<TObjectPtr<UMassProcessor>> SpawnDataInitializers;
 
 	UPROPERTY()
 	UMassEntitySubsystem* EntitySystem;
diff --git a/Engine/Plugins/Runtime/MassGameplay/Source/MassSpawner/Public/MassSpawnerTypes.h b/Engine/Plugins/Runtime/MassGameplay/Source/MassSpawner/Public/MassSpawnerTypes.h
--- a/Engine/Plugins/Runtime/MassGameplay/Source/MassSpawner/Public/MassSpawnerTypes.h
+++ b/Engine/Plugins/Runtime/MassGameplay/Source/MassSpawner/Public/MassSpawnerTypes.h
@@ -2,13 +2,8 @@
 
 #pragma once 
 
-#include "Engine/DataTable.h"
 #include "Logging/LogMacros.h"
-#include "MassCommonTypes.h"
-#include "InstancedStruct.h"
 #include "MassEntityTemplate.h"
-#include "MassEntitySpawnPointsGeneratorBase.h"
-
 #include "MassSpawnerTypes.generated.h"
 
 
@@ -17,9 +12,10 @@ DECLARE_LOG_CATEGORY_EXTERN(LogMassSpawner, Warning, All);
 class APawn;
 class UCurveFloat;
 class UMassEntityConfigAsset;
+class UMassEntitySpawnDataGeneratorBase;
 
 USTRUCT()
-struct FMassSpawnAuxData
+struct FMassTransformsSpawnData
 {
 	GENERATED_BODY()
 
@@ -59,23 +55,23 @@ private:
 };
 
 USTRUCT(BlueprintType)
-struct FMassSpawnPointGenerator
+struct FMassSpawnDataGenerator
 {
 	GENERATED_BODY()
 
 	/** The Generator to use to generate the spawn the points */
 	UPROPERTY()
-	TSubclassOf<UMassEntitySpawnPointsGeneratorBase> GeneratorClass;
+	TSubclassOf<UMassEntitySpawnDataGeneratorBase> GeneratorClass;
 
 	UPROPERTY(BlueprintReadWrite, EditAnywhere, Instanced, Category = "Mass|Generator")
-	TObjectPtr<UMassEntitySpawnPointsGeneratorBase> GeneratorInstance = nullptr;
+	TObjectPtr<UMassEntitySpawnDataGeneratorBase> GeneratorInstance = nullptr;
 
 	/** Proportion of the spawn points that should be generated by this generator, (the proportions will be normalized with other sibling generators) */
 	UPROPERTY(BlueprintReadWrite, EditAnywhere, Category = "Mass|Generator", meta = (ClampMin = 0.0f, UIMin = 0.0f))
 	float Proportion = 1.0f;
 
-	/** Runtime value to know if we received the generated points for this generator */
-	bool bPointsGenerated = false;
+	/** Runtime value to know if we received the generated data for this generator */
+	bool bDataGenerated = false;
 };
 
 USTRUCT()
