Moved representation processors configuration into shared fragments

#rnx

diff --git a/Engine/Plugins/AI/MassCrowd/Source/MassCrowd/Private/MassCrowdServerRepresentationProcessor.cpp b/Engine/Plugins/AI/MassCrowd/Source/MassCrowd/Private/MassCrowdServerRepresentationProcessor.cpp
--- a/Engine/Plugins/AI/MassCrowd/Source/MassCrowd/Private/MassCrowdServerRepresentationProcessor.cpp
+++ b/Engine/Plugins/AI/MassCrowd/Source/MassCrowd/Private/MassCrowdServerRepresentationProcessor.cpp
@@ -9,9 +9,4 @@ UMassCrowdServerRepresentationProcessor::UMassCrowdServerRepresentationProcessor
 	ExecutionFlags = (int32)EProcessorExecutionFlags::Server;
 
 	ExecutionOrder.ExecuteAfter.Add(UE::Mass::ProcessorGroupNames::LOD);
-
-	LODRepresentation[EMassLOD::High] = ERepresentationType::HighResSpawnedActor;
-	LODRepresentation[EMassLOD::Medium] = ERepresentationType::None;
-	LODRepresentation[EMassLOD::Low] = ERepresentationType::None;
-	LODRepresentation[EMassLOD::Off] = ERepresentationType::None;
 }
\ No newline at end of file
diff --git a/Engine/Plugins/AI/MassCrowd/Source/MassCrowd/Private/MassCrowdServerRepresentationTrait.cpp b/Engine/Plugins/AI/MassCrowd/Source/MassCrowd/Private/MassCrowdServerRepresentationTrait.cpp
--- a/Engine/Plugins/AI/MassCrowd/Source/MassCrowd/Private/MassCrowdServerRepresentationTrait.cpp
+++ b/Engine/Plugins/AI/MassCrowd/Source/MassCrowd/Private/MassCrowdServerRepresentationTrait.cpp
@@ -8,6 +8,20 @@
 #include "Engine/World.h"
 #include "MassCrowdRepresentationActorManagement.h"
 
+UMassCrowdServerRepresentationTrait::UMassCrowdServerRepresentationTrait()
+{
+	Config.RepresentationActorManagementClass = UMassCrowdRepresentationActorManagement::StaticClass();
+	Config.LODRepresentation[EMassLOD::High] = ERepresentationType::HighResSpawnedActor;
+	Config.LODRepresentation[EMassLOD::Medium] = ERepresentationType::None;
+	Config.LODRepresentation[EMassLOD::Low] = ERepresentationType::None;
+	Config.LODRepresentation[EMassLOD::Off] = ERepresentationType::None;
+	Config.bKeepLowResActors = false;
+	Config.bKeepActorExtraFrame = false;
+	Config.bSpreadFirstVisualizationUpdate = false;
+	Config.WorldPartitionGridNameContainingCollision = NAME_None;
+	Config.NotVisibleUpdateRate = 0.5f;
+}
+
 void UMassCrowdServerRepresentationTrait::BuildTemplate(FMassEntityTemplateBuildContext& BuildContext, UWorld& World) const
 {
 	// This should only be ran on NM_DedicatedServer network mode
@@ -33,10 +47,9 @@ void UMassCrowdServerRepresentationTrait::BuildTemplate(FMassEntityTemplateBuild
 	FSharedStruct SubsystemFragment = EntitySubsystem->GetOrCreateSharedFragment<FMassRepresentationSubsystemFragment>(SubsystemHash, Subsystem);
 	BuildContext.AddSharedFragment(SubsystemFragment);
 
-	FMassRepresentationConfig Config;
-	Config.RepresentationActorManagement = UMassCrowdRepresentationActorManagement::StaticClass()->GetDefaultObject<UMassCrowdRepresentationActorManagement>();
 	uint32 ConfigHash = UE::StructUtils::GetStructCrc32(FConstStructView::Make(Config));
 	FConstSharedStruct ConfigFragment = EntitySubsystem->GetOrCreateConstSharedFragment<FMassRepresentationConfig>(ConfigHash, Config);
+	ConfigFragment.Get<FMassRepresentationConfig>().ComputeCachedValues();
 	BuildContext.AddConstSharedFragment(ConfigFragment);
 
 	FMassRepresentationFragment& RepresentationFragment = BuildContext.AddFragment_GetRef<FMassRepresentationFragment>();
diff --git a/Engine/Plugins/AI/MassCrowd/Source/MassCrowd/Private/MassCrowdVisualizationProcessor.cpp b/Engine/Plugins/AI/MassCrowd/Source/MassCrowd/Private/MassCrowdVisualizationProcessor.cpp
--- a/Engine/Plugins/AI/MassCrowd/Source/MassCrowd/Private/MassCrowdVisualizationProcessor.cpp
+++ b/Engine/Plugins/AI/MassCrowd/Source/MassCrowd/Private/MassCrowdVisualizationProcessor.cpp
@@ -31,7 +31,6 @@ UMassCrowdVisualizationProcessor::UMassCrowdVisualizationProcessor()
 	ExecutionOrder.ExecuteAfter.Add(UE::Mass::ProcessorGroupNames::LOD);
 
 	bRequiresGameThreadExecution = true;
-	bKeepActorExtraFrame = true;
 }
 
 void UMassCrowdVisualizationProcessor::ConfigureQueries()
diff --git a/Engine/Plugins/AI/MassCrowd/Source/MassCrowd/Private/MassCrowdVisualizationTrait.cpp b/Engine/Plugins/AI/MassCrowd/Source/MassCrowd/Private/MassCrowdVisualizationTrait.cpp
--- a/Engine/Plugins/AI/MassCrowd/Source/MassCrowd/Private/MassCrowdVisualizationTrait.cpp
+++ b/Engine/Plugins/AI/MassCrowd/Source/MassCrowd/Private/MassCrowdVisualizationTrait.cpp
@@ -9,5 +9,14 @@ UMassCrowdVisualizationTrait::UMassCrowdVisualizationTrait()
 {
 	// Override the subsystem to support parallelization of the crowd
 	RepresentationSubsystemClass = UMassCrowdRepresentationSubsystem::StaticClass();
-	RepresentationActorManagementClass = UMassCrowdRepresentationActorManagement::StaticClass();
+	Config.RepresentationActorManagementClass = UMassCrowdRepresentationActorManagement::StaticClass();
+	Config.LODRepresentation[EMassLOD::High] = ERepresentationType::HighResSpawnedActor;
+	Config.LODRepresentation[EMassLOD::Medium] = ERepresentationType::LowResSpawnedActor;
+	Config.LODRepresentation[EMassLOD::Low] = ERepresentationType::StaticMeshInstance;
+	Config.LODRepresentation[EMassLOD::Off] = ERepresentationType::None;
+	Config.bKeepLowResActors = true;
+	Config.bKeepActorExtraFrame = true;
+	Config.bSpreadFirstVisualizationUpdate = false;
+	Config.WorldPartitionGridNameContainingCollision = NAME_None;
+	Config.NotVisibleUpdateRate = 0.5f;
 }
diff --git a/Engine/Plugins/AI/MassCrowd/Source/MassCrowd/Public/MassCrowdServerRepresentationTrait.h b/Engine/Plugins/AI/MassCrowd/Source/MassCrowd/Public/MassCrowdServerRepresentationTrait.h
--- a/Engine/Plugins/AI/MassCrowd/Source/MassCrowd/Public/MassCrowdServerRepresentationTrait.h
+++ b/Engine/Plugins/AI/MassCrowd/Source/MassCrowd/Public/MassCrowdServerRepresentationTrait.h
@@ -4,7 +4,9 @@
 
 #include "MassEntityTraitBase.h"
 #include "MassRepresentationTypes.h"
+#include "MassRepresentationFragments.h"
 #include "GameFramework/Actor.h"
+
 #include "MassCrowdServerRepresentationTrait.generated.h"
 
 
@@ -13,6 +15,8 @@ class MASSCROWD_API UMassCrowdServerRepresentationTrait : public UMassEntityTrai
 {
 	GENERATED_BODY()
 
+	UMassCrowdServerRepresentationTrait();
+
 protected:
 
 	virtual void BuildTemplate(FMassEntityTemplateBuildContext& BuildContext, UWorld& World) const override;
@@ -20,4 +24,8 @@ protected:
 	/** Actor class of this agent when spawned on server */
 	UPROPERTY(EditAnywhere, Category = "Mass|Visual")
 	TSubclassOf<AActor> TemplateActor;
+
+	/** Configuration parameters for the representation processor */
+	UPROPERTY(EditAnywhere, Category = "Mass|Visual")
+	FMassRepresentationConfig Config;
 };
diff --git a/Engine/Plugins/Runtime/MassGameplay/Source/MassRepresentation/Private/MassRepresentationActorManagement.cpp b/Engine/Plugins/Runtime/MassGameplay/Source/MassRepresentation/Private/MassRepresentationActorManagement.cpp
--- a/Engine/Plugins/Runtime/MassGameplay/Source/MassRepresentation/Private/MassRepresentationActorManagement.cpp
+++ b/Engine/Plugins/Runtime/MassGameplay/Source/MassRepresentation/Private/MassRepresentationActorManagement.cpp
@@ -1,10 +1,17 @@
 // Copyright Epic Games, Inc. All Rights Reserved.
 
 #include "MassRepresentationActorManagement.h"
+#include "MassRepresentationSubsystem.h"
+#include "MassRepresentationFragments.h"
 #include "MassCommandBuffer.h"
 #include "MassEntityView.h"
 #include "MassActorSubsystem.h"
-#include "MassRepresentationSubsystem.h"
+
+float UMassRepresentationActorManagement::GetSpawnPriority(const FMassRepresentationLODFragment& Representation) const
+{
+	// Bump up the spawning priority on the visible entities
+	return Representation.LODSignificance - (Representation.Visibility == EMassVisibility::CanBeSeen ? 1.0f : 0.0f);
+}
 
 AActor* UMassRepresentationActorManagement::GetOrSpawnActor(UMassRepresentationSubsystem& RepresentationSubsystem, UMassEntitySubsystem& EntitySubsystem, const FMassEntityHandle MassAgent, FDataFragment_Actor& ActorInfo, const FTransform& Transform, const int16 TemplateActorIndex, FMassActorSpawnRequestHandle& SpawnRequestHandle, const float Priority) const
 {
diff --git a/Engine/Plugins/Runtime/MassGameplay/Source/MassRepresentation/Private/MassRepresentationProcessor.cpp b/Engine/Plugins/Runtime/MassGameplay/Source/MassRepresentation/Private/MassRepresentationProcessor.cpp
--- a/Engine/Plugins/Runtime/MassGameplay/Source/MassRepresentation/Private/MassRepresentationProcessor.cpp
+++ b/Engine/Plugins/Runtime/MassGameplay/Source/MassRepresentation/Private/MassRepresentationProcessor.cpp
@@ -25,11 +25,6 @@ UMassRepresentationProcessor::UMassRepresentationProcessor()
 	bAutoRegisterWithProcessingPhases = false;
 
 	ExecutionOrder.ExecuteInGroup = UE::Mass::ProcessorGroupNames::Representation;
-
-	LODRepresentation[EMassLOD::High] = ERepresentationType::HighResSpawnedActor;
-	LODRepresentation[EMassLOD::Medium] = ERepresentationType::LowResSpawnedActor;
-	LODRepresentation[EMassLOD::Low] = ERepresentationType::StaticMeshInstance;
-	LODRepresentation[EMassLOD::Off] = ERepresentationType::None;
 }
 
 void UMassRepresentationProcessor::ConfigureQueries()
@@ -49,20 +44,6 @@ void UMassRepresentationProcessor::Initialize(UObject& Owner)
 
 	World = Owner.GetWorld();
 	CachedEntitySubsystem = UWorld::GetSubsystem<UMassEntitySubsystem>(World);
-
-	// Calculate the default representation when actor isn't spawned yet.
-	for (int32 LOD = EMassLOD::High; LOD < EMassLOD::Max; LOD++)
-	{
-		// Find the first representation type after any actors
-		if (LODRepresentation[LOD] == ERepresentationType::HighResSpawnedActor ||
-			LODRepresentation[LOD] == ERepresentationType::LowResSpawnedActor)
-		{
-			continue;
-		}
-
-		DefaultRepresentationType = LODRepresentation[LOD];
-		break;
-	}
 }
 
 void UMassRepresentationProcessor::UpdateRepresentation(FMassExecutionContext& Context)
@@ -70,7 +51,7 @@ void UMassRepresentationProcessor::UpdateRepresentation(FMassExecutionContext& C
 	UMassRepresentationSubsystem* RepresentationSubsystem = Context.GetMutableSharedFragment<FMassRepresentationSubsystemFragment>().RepresentationSubsystem;
 	check(RepresentationSubsystem);
 	const FMassRepresentationConfig& RepresentationConfig = Context.GetConstSharedFragment<FMassRepresentationConfig>();
-	UMassRepresentationActorManagement* RepresentationActorManagement = RepresentationConfig.RepresentationActorManagement;
+	UMassRepresentationActorManagement* RepresentationActorManagement = RepresentationConfig.CachedRepresentationActorManagement;
 	check(RepresentationActorManagement);
 
 	const TConstArrayView<FDataFragment_Transform> TransformList = Context.GetFragmentView<FDataFragment_Transform>();
@@ -78,7 +59,7 @@ void UMassRepresentationProcessor::UpdateRepresentation(FMassExecutionContext& C
 	const TConstArrayView<FMassRepresentationLODFragment> RepresentationLODList = Context.GetFragmentView<FMassRepresentationLODFragment>();
 	const TArrayView<FDataFragment_Actor> ActorList = Context.GetMutableFragmentView<FDataFragment_Actor>();
 
-	const bool bDoKeepActorExtraFrame = UE::MassRepresentation::bAllowKeepActorExtraFrame ? bKeepLowResActors : false;
+	const bool bDoKeepActorExtraFrame = UE::MassRepresentation::bAllowKeepActorExtraFrame ? RepresentationConfig.bKeepLowResActors : false;
 
 	const int32 NumEntities = Context.GetNumEntities();
 	for (int32 EntityIdx = 0; EntityIdx < NumEntities; EntityIdx++)
@@ -93,13 +74,13 @@ void UMassRepresentationProcessor::UpdateRepresentation(FMassExecutionContext& C
 		const ERepresentationType PrevRepresentationCopy = Representation.PrevRepresentation;
 		Representation.PrevRepresentation = Representation.CurrentRepresentation;
 		
-		ERepresentationType WantedRepresentationType = LODRepresentation[FMath::Min((int32)RepresentationLOD.LOD, (int32)EMassLOD::Off)];
+		ERepresentationType WantedRepresentationType = RepresentationConfig.LODRepresentation[FMath::Min((int32)RepresentationLOD.LOD, (int32)EMassLOD::Off)];
 
 		// Make sure we do not have actor spawned in areas not fully loaded
 		if ((WantedRepresentationType == ERepresentationType::HighResSpawnedActor || WantedRepresentationType == ERepresentationType::LowResSpawnedActor) &&
-			!RepresentationSubsystem->IsCollisionLoaded(WorldPartitionGridNameContainingCollision, TransformFragment.GetTransform()))
+			!RepresentationSubsystem->IsCollisionLoaded(RepresentationConfig.WorldPartitionGridNameContainingCollision, TransformFragment.GetTransform()))
 		{
-			WantedRepresentationType = DefaultRepresentationType;
+			WantedRepresentationType = RepresentationConfig.CachedDefaultRepresentationType;
 		}
 		
 		auto DisableActorForISM = [&](AActor*& Actor)
@@ -108,7 +89,7 @@ void UMassRepresentationProcessor::UpdateRepresentation(FMassExecutionContext& C
 			{
 				// Execute only if the high res is different than the low res Actor 
 				// Or if we do not wish to keep the low res actor while in ISM
-				if (Representation.HighResTemplateActorIndex != Representation.LowResTemplateActorIndex || !bKeepLowResActors)
+				if (Representation.HighResTemplateActorIndex != Representation.LowResTemplateActorIndex || !RepresentationConfig.bKeepLowResActors)
 				{
 					// Try releasing the high actor or any high res spawning request
 					if (ReleaseActorOrCancelSpawning(*RepresentationSubsystem, MassAgent, ActorInfo, Representation.HighResTemplateActorIndex, Representation.ActorSpawnRequestHandle, Context.Defer()))
@@ -116,7 +97,7 @@ void UMassRepresentationProcessor::UpdateRepresentation(FMassExecutionContext& C
 						Actor = ActorInfo.GetOwnedByMassMutable();
 					}
 					// Do not do the same with low res if indicated so
-					if (!bKeepLowResActors && ReleaseActorOrCancelSpawning(*RepresentationSubsystem, MassAgent, ActorInfo, Representation.LowResTemplateActorIndex, Representation.ActorSpawnRequestHandle, Context.Defer()))
+					if (!RepresentationConfig.bKeepLowResActors && ReleaseActorOrCancelSpawning(*RepresentationSubsystem, MassAgent, ActorInfo, Representation.LowResTemplateActorIndex, Representation.ActorSpawnRequestHandle, Context.Defer()))
 					{
 						Actor = ActorInfo.GetOwnedByMassMutable();
 					}
@@ -197,7 +178,7 @@ void UMassRepresentationProcessor::UpdateRepresentation(FMassExecutionContext& C
 					}
 					else if (!Actor)
 					{
-						Representation.CurrentRepresentation = DefaultRepresentationType;
+						Representation.CurrentRepresentation = RepresentationConfig.CachedDefaultRepresentationType;
 					}
 					break;
 				}
@@ -303,6 +284,7 @@ void UMassRepresentationProcessor::UpdateVisualization(FMassExecutionContext& Co
 
 FMassVisualizationChunkFragment& UMassRepresentationProcessor::UpdateChunkVisibility(FMassExecutionContext& Context) const
 {
+	const FMassRepresentationConfig& RepresentationConfig = Context.GetConstSharedFragment<FMassRepresentationConfig>();
 	bool bFirstUpdate = false;
 
 	// Setup chunk fragment data about visibility
@@ -313,7 +295,7 @@ FMassVisualizationChunkFragment& UMassRepresentationProcessor::UpdateChunkVisibi
 		// The visibility on the chunk fragment data isn't set yet, let see if the Archetype has an visibility tag and set it on the ChunkData
 		ChunkVisibility = UE::MassRepresentation::GetVisibilityFromArchetype(Context);
 		ChunkData.SetVisibility(ChunkVisibility);
-		bFirstUpdate = bSpreadFirstVisualizationUpdate;
+		bFirstUpdate = RepresentationConfig.bSpreadFirstVisualizationUpdate;
 	}
 	else
 	{
@@ -326,13 +308,13 @@ FMassVisualizationChunkFragment& UMassRepresentationProcessor::UpdateChunkVisibi
 		if (bFirstUpdate)
 		{
 			// A DeltaTime of 0.0f means it will tick this frame.
-			DeltaTime = FMath::RandRange(0.0f, NotVisibleUpdateRate);
+			DeltaTime = FMath::RandRange(0.0f, RepresentationConfig.NotVisibleUpdateRate);
 		}
 		else 
 		{
 			if (DeltaTime < 0.0f)
 			{
-				DeltaTime += NotVisibleUpdateRate * (1.0f + FMath::RandRange(-0.1f, 0.1f));
+				DeltaTime += RepresentationConfig.NotVisibleUpdateRate * (1.0f + FMath::RandRange(-0.1f, 0.1f));
 			}
 			DeltaTime -= Context.GetDeltaTimeSeconds();
 		}
diff --git a/Engine/Plugins/Runtime/MassGameplay/Source/MassRepresentation/Private/MassVisualizationTrait.cpp b/Engine/Plugins/Runtime/MassGameplay/Source/MassRepresentation/Private/MassVisualizationTrait.cpp
--- a/Engine/Plugins/Runtime/MassGameplay/Source/MassRepresentation/Private/MassVisualizationTrait.cpp
+++ b/Engine/Plugins/Runtime/MassGameplay/Source/MassRepresentation/Private/MassVisualizationTrait.cpp
@@ -11,7 +11,12 @@
 UMassVisualizationTrait::UMassVisualizationTrait()
 {
 	RepresentationSubsystemClass = UMassRepresentationSubsystem::StaticClass();
-	RepresentationActorManagementClass = UMassRepresentationActorManagement::StaticClass();
+
+	Config.RepresentationActorManagementClass = UMassRepresentationActorManagement::StaticClass();
+	Config.LODRepresentation[EMassLOD::High] = ERepresentationType::HighResSpawnedActor;
+	Config.LODRepresentation[EMassLOD::Medium] = ERepresentationType::LowResSpawnedActor;
+	Config.LODRepresentation[EMassLOD::Low] = ERepresentationType::StaticMeshInstance;
+	Config.LODRepresentation[EMassLOD::Off] = ERepresentationType::None;
 }
 
 void UMassVisualizationTrait::BuildTemplate(FMassEntityTemplateBuildContext& BuildContext, UWorld& World) const
@@ -45,15 +50,13 @@ void UMassVisualizationTrait::BuildTemplate(FMassEntityTemplateBuildContext& Bui
 	FSharedStruct SubsystemFragment = EntitySubsystem->GetOrCreateSharedFragment<FMassRepresentationSubsystemFragment>(SubsystemHash, Subsystem);
 	BuildContext.AddSharedFragment(SubsystemFragment);
 
-	FMassRepresentationConfig Config;
-	Config.RepresentationActorManagement = RepresentationActorManagementClass.GetDefaultObject();
-	if (Config.RepresentationActorManagement == nullptr)
+	if (!Config.RepresentationActorManagementClass)
 	{
 		UE_LOG(LogMassRepresentation, Error, TEXT("Expecting a valid class for the representation actor management"));
-		Config.RepresentationActorManagement = UMassRepresentationActorManagement::StaticClass()->GetDefaultObject<UMassRepresentationActorManagement>();
 	}
 	uint32 ConfigHash = UE::StructUtils::GetStructCrc32(FConstStructView::Make(Config));
 	FConstSharedStruct ConfigFragment = EntitySubsystem->GetOrCreateConstSharedFragment<FMassRepresentationConfig>(ConfigHash, Config);
+	ConfigFragment.Get<FMassRepresentationConfig>().ComputeCachedValues();
 	BuildContext.AddConstSharedFragment(ConfigFragment);
 
 	FMassRepresentationFragment& RepresentationFragment = BuildContext.AddFragment_GetRef<FMassRepresentationFragment>();
diff --git a/Engine/Plugins/Runtime/MassGameplay/Source/MassRepresentation/Public/MassRepresentationActorManagement.h b/Engine/Plugins/Runtime/MassGameplay/Source/MassRepresentation/Public/MassRepresentationActorManagement.h
--- a/Engine/Plugins/Runtime/MassGameplay/Source/MassRepresentation/Public/MassRepresentationActorManagement.h
+++ b/Engine/Plugins/Runtime/MassGameplay/Source/MassRepresentation/Public/MassRepresentationActorManagement.h
@@ -2,7 +2,6 @@
 
 #pragma once
 
-#include "MassRepresentationFragments.h"
 #include "MassRepresentationActorManagement.generated.h"
 
 struct FMassActorSpawnRequestHandle;
@@ -14,6 +13,9 @@ struct FMassActorSpawnRequestHandle;
 struct FStructView;
 struct FMassEntityHandle;
 struct FDataFragment_Actor;
+struct FMassRepresentationLODFragment;
+struct FMassRepresentationFragment;
+class UMassRepresentationSubsystem;
 
 UCLASS()
 class MASSREPRESENTATION_API UMassRepresentationActorManagement : public UObject
@@ -26,11 +28,7 @@ public:
 	 * Returns the spawn priority from the LOD
 	 * @param Representation is the type of enabling to do on this actor
 	 */
-	virtual float GetSpawnPriority(const FMassRepresentationLODFragment& Representation) const
-	{
-		// Bump up the spawning priority on the visible entities
-		return Representation.LODSignificance - (Representation.Visibility == EMassVisibility::CanBeSeen ? 1.0f : 0.0f);
-	}
+	virtual float GetSpawnPriority(const FMassRepresentationLODFragment& Representation) const;
 
 	/**
 	 * Returns an actor of the template type and setup fragments values from it
diff --git a/Engine/Plugins/Runtime/MassGameplay/Source/MassRepresentation/Public/MassRepresentationFragments.h b/Engine/Plugins/Runtime/MassGameplay/Source/MassRepresentation/Public/MassRepresentationFragments.h
--- a/Engine/Plugins/Runtime/MassGameplay/Source/MassRepresentation/Public/MassRepresentationFragments.h
+++ b/Engine/Plugins/Runtime/MassGameplay/Source/MassRepresentation/Public/MassRepresentationFragments.h
@@ -7,6 +7,7 @@
 #include "MassEntitySubsystem.h"
 #include "MassActorSpawnerSubsystem.h"
 #include "MassRepresentationTypes.h"
+#include "MassRepresentationActorManagement.h"
 
 #include "MassRepresentationFragments.generated.h"
 
@@ -67,6 +68,65 @@ struct FMassRepresentationConfig : public FMassSharedFragment
 {
 	GENERATED_BODY()
 
+
+	/** Allow subclasses to override the representation actor management behavior */
+	UPROPERTY(EditAnywhere, Category = "Mass|Visual")
+	TSubclassOf<UMassRepresentationActorManagement> RepresentationActorManagementClass;
+
+	/** What should be the representation of this entity for each specific LOD */
+	UPROPERTY(EditAnywhere, Category = "Mass|Representation", config)
+	ERepresentationType LODRepresentation[EMassLOD::Max];
+
+	/** If true, LowRes actors will be kept around, disabled, whilst StaticMeshInstance representation is active */
+	UPROPERTY(EditAnywhere, Category = "Mass|Representation", config)
+	bool bKeepLowResActors = true;
+
+	/** When switching to ISM keep the actor an extra frame, helps cover rendering glitches (i.e. occlusion query being one frame late) */
+	UPROPERTY(EditAnywhere, Category = "Mass|Representation", config)
+	bool bKeepActorExtraFrame = false;
+
+	/** If true, will spread the first visualization update over the period specified in NotVisibleUpdateRate member */
+	UPROPERTY(EditAnywhere, Category = "Mass|Representation", config)
+	bool bSpreadFirstVisualizationUpdate = false;
+
+	/** World Partition grid name to test collision against, default None will be the main grid */
+	UPROPERTY(EditAnywhere, Category = "Mass|Representation", config)
+	FName WorldPartitionGridNameContainingCollision;
+
+	/** At what rate should the not visible entity be updated in seconds */
+	UPROPERTY(EditAnywhere, Category = "Mass|Visualization", config)
+	float NotVisibleUpdateRate = 0.5f;
+
+	inline void ComputeCachedValues() const;
+
+	/** Default representation when unable to spawn an actor, gets calculated at initialization */
+	UPROPERTY(Transient)
+	mutable ERepresentationType CachedDefaultRepresentationType = ERepresentationType::None;
+
 	UPROPERTY(Transient)
-	UMassRepresentationActorManagement* RepresentationActorManagement;
+	mutable UMassRepresentationActorManagement* CachedRepresentationActorManagement;
 };
+
+inline void FMassRepresentationConfig::ComputeCachedValues() const
+{
+	// Calculate the default representation when actor isn't spawned yet.
+	for (int32 LOD = EMassLOD::High; LOD < EMassLOD::Max; LOD++)
+	{
+		// Find the first representation type after any actors
+		if (LODRepresentation[LOD] == ERepresentationType::HighResSpawnedActor ||
+			LODRepresentation[LOD] == ERepresentationType::LowResSpawnedActor)
+		{
+			continue;
+		}
+
+		CachedDefaultRepresentationType = LODRepresentation[LOD];
+		break;
+	}
+
+	CachedRepresentationActorManagement = RepresentationActorManagementClass.GetDefaultObject();
+	if (CachedRepresentationActorManagement == nullptr)
+	{
+		// We should have warn about it in the traits.
+		CachedRepresentationActorManagement = UMassRepresentationActorManagement::StaticClass()->GetDefaultObject<UMassRepresentationActorManagement>();
+	}
+}
diff --git a/Engine/Plugins/Runtime/MassGameplay/Source/MassRepresentation/Public/MassRepresentationProcessor.h b/Engine/Plugins/Runtime/MassGameplay/Source/MassRepresentation/Public/MassRepresentationProcessor.h
--- a/Engine/Plugins/Runtime/MassGameplay/Source/MassRepresentation/Public/MassRepresentationProcessor.h
+++ b/Engine/Plugins/Runtime/MassGameplay/Source/MassRepresentation/Public/MassRepresentationProcessor.h
@@ -81,26 +81,6 @@ protected:
 	 */
 	void UpdateVisualization(FMassExecutionContext& Context);
 
-	/** What should be the representation of this entity for each specificLOD */
-	UPROPERTY(EditAnywhere, Category = "Mass|Representation", config)
-	ERepresentationType LODRepresentation[EMassLOD::Max];
-
-	/** If true, LowRes actors will be kept around, disabled, whilst StaticMeshInstance representation is active */
-	UPROPERTY(EditAnywhere, Category = "Mass|Representation", config)
-	bool bKeepLowResActors = true;
-
-	/** When switching to ISM keep the actor an extra frame, helps cover rendering glitches (i.e. occlusion query being one frame late) */
-	UPROPERTY(EditAnywhere, Category = "Mass|Representation", config)
-	bool bKeepActorExtraFrame = false;
-
-	/** If true, will spread the first visualization update over the period specified in NotVisibleUpdateRate member */
-	UPROPERTY(EditAnywhere, Category = "Mass|Representation", config)
-	bool bSpreadFirstVisualizationUpdate = false;
-
-	/** World Partition grid name to test collision against, default None will be the main grid */
-	UPROPERTY(EditAnywhere, Category = "Mass|Representation", config)
-	FName WorldPartitionGridNameContainingCollision;
-
 	/** Caching ptr to our associated world */
 	UPROPERTY(Transient)
 	UWorld* World;
@@ -108,14 +88,9 @@ protected:
 	UPROPERTY(Transient)
 	UMassEntitySubsystem* CachedEntitySubsystem;
 
-	/** Default representation when unable to spawn an actor */
-	ERepresentationType DefaultRepresentationType = ERepresentationType::None;
 
 	FMassEntityQuery EntityQuery;
 
-	/** At what rate should the not visible entity be updated in seconds */
-	UPROPERTY(EditAnywhere, Category = "Mass|Visualization", config)
-	float NotVisibleUpdateRate = 0.5f;
 };
 
 UCLASS()
diff --git a/Engine/Plugins/Runtime/MassGameplay/Source/MassRepresentation/Public/MassVisualizationTrait.h b/Engine/Plugins/Runtime/MassGameplay/Source/MassRepresentation/Public/MassVisualizationTrait.h
--- a/Engine/Plugins/Runtime/MassGameplay/Source/MassRepresentation/Public/MassVisualizationTrait.h
+++ b/Engine/Plugins/Runtime/MassGameplay/Source/MassRepresentation/Public/MassVisualizationTrait.h
@@ -4,7 +4,9 @@
 
 #include "MassEntityTraitBase.h"
 #include "MassRepresentationTypes.h"
+#include "MassRepresentationFragments.h"
 #include "GameFramework/Actor.h"
+
 #include "MassVisualizationTrait.generated.h"
 
 class UMassRepresentationSubsystem;
@@ -36,11 +38,7 @@ public:
 	UPROPERTY(EditAnywhere, Category = "Mass|Visual")
 	TSubclassOf<UMassRepresentationSubsystem> RepresentationSubsystemClass;
 
-	/** Allow subclasses to override the representation actor management behavior */
+	/** Configuration parameters for the representaion processor */
 	UPROPERTY(EditAnywhere, Category = "Mass|Visual")
-	TSubclassOf<UMassRepresentationActorManagement> RepresentationActorManagementClass;
-
-	// @todo the following property will be cut once new de/initializers are in
-	/** Allow subclasses to override the representation fragment deinitializer */
-	TSubclassOf<UMassProcessor> RepresentationFragmentDeinitializerClass;
+	FMassRepresentationConfig Config;
 };
