Fix uninitialized member struct variable which is causing a crash from my last checkin

#rnx

diff --git a/Engine/Plugins/Runtime/MassGameplay/Source/MassRepresentation/Public/MassRepresentationFragments.h b/Engine/Plugins/Runtime/MassGameplay/Source/MassRepresentation/Public/MassRepresentationFragments.h
--- a/Engine/Plugins/Runtime/MassGameplay/Source/MassRepresentation/Public/MassRepresentationFragments.h
+++ b/Engine/Plugins/Runtime/MassGameplay/Source/MassRepresentation/Public/MassRepresentationFragments.h
@@ -104,7 +104,7 @@ struct FMassRepresentationConfig : public FMassSharedFragment
 	mutable ERepresentationType CachedDefaultRepresentationType = ERepresentationType::None;
 
 	UPROPERTY(Transient)
-	mutable UMassRepresentationActorManagement* CachedRepresentationActorManagement;
+	mutable UMassRepresentationActorManagement* CachedRepresentationActorManagement = nullptr;
 };
 
 inline void FMassRepresentationConfig::ComputeCachedValues() const
