Support optional chunk fragment

#rnx

diff --git a/Engine/Plugins/Runtime/MassGameplay/Source/MassLOD/Public/MassLODTypes.h b/Engine/Plugins/Runtime/MassGameplay/Source/MassLOD/Public/MassLODTypes.h
--- a/Engine/Plugins/Runtime/MassGameplay/Source/MassLOD/Public/MassLODTypes.h
+++ b/Engine/Plugins/Runtime/MassGameplay/Source/MassLOD/Public/MassLODTypes.h
@@ -201,7 +201,8 @@ struct MASSLOD_API FMassVisualizationChunkFragment : public FMassChunkFragment
 	 */
 	static bool ShouldUpdateVisualizationForChunk(const FMassExecutionContext& Context)
 	{
-		return Context.GetChunkFragment<FMassVisualizationChunkFragment>().ShouldUpdateVisualization();
+		const FMassVisualizationChunkFragment* ChunkFragment = Context.GetChunkFragmentPtr<FMassVisualizationChunkFragment>();
+		return ChunkFragment == nullptr || ChunkFragment->ShouldUpdateVisualization();
 	}
 
 	/**
