Updated UMassAgentSubsystem::UnregisterAgentComponent to test both SimulationSystem->GetPhaseManager().IsDuringMassProcessing() as well as EntitySystem->IsProcessing()
This is a temporary solution. The disrepency between these two functions is confusing and should be addressed sustem-wide. Filed as UE-140265.

#rnx

diff --git a/Engine/Plugins/Runtime/MassGameplay/Source/MassActors/Private/MassAgentSubsystem.cpp b/Engine/Plugins/Runtime/MassGameplay/Source/MassActors/Private/MassAgentSubsystem.cpp
--- a/Engine/Plugins/Runtime/MassGameplay/Source/MassActors/Private/MassAgentSubsystem.cpp
+++ b/Engine/Plugins/Runtime/MassGameplay/Source/MassActors/Private/MassAgentSubsystem.cpp
@@ -181,7 +181,8 @@ void UMassAgentSubsystem::UnregisterAgentComponent(UMassAgentComponent& AgentCom
 				const FMassEntityHandle Entity = AgentComp.GetEntityHandle();
 
 			    // remove fragments that have been added for the puppet agent
-			    if (SimulationSystem->GetPhaseManager().IsDuringMassProcessing())
+			    if (SimulationSystem->GetPhaseManager().IsDuringMassProcessing()
+					|| EntitySystem->IsProcessing())
 			    {
 				    // need to request via command buffer since we can't move entities while processing is happening
 				    EntitySystem->Defer().PushCommand(FCommandRemoveComposition(Entity, AgentComp.GetPuppetSpecificAddition()));
@@ -212,7 +213,8 @@ void UMassAgentSubsystem::UnregisterAgentComponent(UMassAgentComponent& AgentCom
 			AgentComp.ClearEntityHandle();
 
 			// Destroy the entity
-			if (SimulationSystem->GetPhaseManager().IsDuringMassProcessing())
+			if (SimulationSystem->GetPhaseManager().IsDuringMassProcessing()
+				|| EntitySystem->IsProcessing())
 			{
 				// need to request via command buffer since we can't move entities while processing is happening
 				EntitySystem->Defer().DestroyEntity(Entity);
