MassSpawner: Fix for non-unity compile.

diff --git a/Engine/Plugins/Runtime/MassGameplay/Source/MassSpawner/Private/MassEntitySpawnDataGeneratorBase.cpp b/Engine/Plugins/Runtime/MassGameplay/Source/MassSpawner/Private/MassEntitySpawnDataGeneratorBase.cpp
--- a/Engine/Plugins/Runtime/MassGameplay/Source/MassSpawner/Private/MassEntitySpawnDataGeneratorBase.cpp
+++ b/Engine/Plugins/Runtime/MassGameplay/Source/MassSpawner/Private/MassEntitySpawnDataGeneratorBase.cpp
@@ -1,6 +1,7 @@
 ﻿// Copyright Epic Games, Inc. All Rights Reserved.
 
 #include "MassEntitySpawnDataGeneratorBase.h"
+#include "VisualLogger/VisualLogger.h"
 #include "MassSpawnerTypes.h"
 
 void UMassEntitySpawnDataGeneratorBase::BuildResultsFromEntityTypes(const int32 SpawnCount, TConstArrayView<FMassSpawnedEntityType> EntityTypes, TArray<FMassEntitySpawnDataGeneratorResult>& OutResults) const
