Mass Movement and Navigation cleanup

diff --git a/Engine/Plugins/Runtime/MassGameplay/Source/MassMovement/Public/MassMovementSettings.h b/Engine/Plugins/Runtime/MassGameplay/Source/MassMovement/Public/MassMovementSettings.h
new file mode 100644
--- /dev/null
+++ b/Engine/Plugins/Runtime/MassGameplay/Source/MassMovement/Public/MassMovementSettings.h
@@ -0,0 +1,34 @@
+// Copyright Epic Games, Inc. All Rights Reserved.
+
+#pragma once
+
+#include "CoreMinimal.h"
+#include "MassSettings.h"
+#include "MassMovementTypes.h"
+#include "MassMovementSettings.generated.h"
+
+UCLASS(config = Mass, defaultconfig, meta = (DisplayName = "Mass Movement"))
+class MASSMOVEMENT_API UMassMovementSettings : public UMassModuleSettings
+{
+	GENERATED_BODY()
+
+public:
+	UMassMovementSettings(const FObjectInitializer& ObjectInitializer = FObjectInitializer::Get());
+	
+	static const UMassMovementSettings* Get()
+	{
+		return GetDefault<UMassMovementSettings>();
+	}
+
+	TConstArrayView<FMassMovementStyle> GetMovementStyles() const { return MovementStyles; }
+	const FMassMovementStyle* GetMovementStyleByID(const FGuid ID) const; 
+	
+private:
+
+#if WITH_EDITOR
+	virtual void PostEditChangeChainProperty(FPropertyChangedChainEvent& PropertyChangedEvent) override;
+#endif
+
+	UPROPERTY(EditAnywhere, config, Category = Movement);
+    TArray<FMassMovementStyle> MovementStyles;
+};
diff --git a/Engine/Plugins/Runtime/MassGameplay/Source/MassMovement/Public/MassMovementTypes.h b/Engine/Plugins/Runtime/MassGameplay/Source/MassMovement/Public/MassMovementTypes.h
--- a/Engine/Plugins/Runtime/MassGameplay/Source/MassMovement/Public/MassMovementTypes.h
+++ b/Engine/Plugins/Runtime/MassGameplay/Source/MassMovement/Public/MassMovementTypes.h
@@ -1,32 +1,79 @@
-// Copyright Epic Games, Inc. All Rights Reserved.
+﻿// Copyright Epic Games, Inc. All Rights Reserved.
 
 #pragma once
 
-#include "MassEntityTypes.h"
-#include "HierarchicalHashGrid2D.h"
+#include "CoreTypes.h"
+#include "MassMovementTypes.generated.h"
 
-namespace UE::Mass::ProcessorGroupNames
+namespace UE::MassMovement
 {
-	const FName Avoidance = FName(TEXT("Avoidance"));
-}
+	extern MASSMOVEMENT_API int32 bFreezeMovement;
+};
 
-UENUM()
-enum class EMassAvoidanceObstacleItemFlags : uint8
+/** Reference to movement style in MassMovementSettings. */
+USTRUCT()
+struct MASSMOVEMENT_API FMassMovementStyleRef
 {
-	None			= 0,
-	HasColliderData = 1 << 0,
+	GENERATED_BODY()
+
+	UPROPERTY(EditAnywhere, Category = "Movement")
+	FName Name;
+
+	UPROPERTY(EditAnywhere, Category = "Movement", meta = (IgnoreForMemberInitializationTest, EditCondition = "false", EditConditionHides))
+	FGuid ID;
 };
-ENUM_CLASS_FLAGS(EMassAvoidanceObstacleItemFlags)
 
-struct FMassAvoidanceObstacleItem
+/** Describes movement style name. */
+USTRUCT()
+struct MASSMOVEMENT_API FMassMovementStyle
 {
-	bool operator==(const FMassAvoidanceObstacleItem& Other) const
-	{
-		return Entity == Other.Entity;
-	}
+	GENERATED_BODY()
 
-	FMassEntityHandle Entity;
-	EMassAvoidanceObstacleItemFlags ItemFlags = EMassAvoidanceObstacleItemFlags::None;
+	UPROPERTY(EditAnywhere, Category = "Movement")
+	FName Name;
+
+	UPROPERTY(EditAnywhere, Category = "Movement", meta = (IgnoreForMemberInitializationTest, EditCondition = "false", EditConditionHides))
+	FGuid ID;
 };
 
-typedef THierarchicalHashGrid2D<2, 4, FMassAvoidanceObstacleItem> FAvoidanceObstacleHashGrid2D;	// 2 levels of hierarchy, 4 ratio between levels
\ No newline at end of file
+/**
+ * Movement style consists of multiple speeds which are assigned to agents based on agents unique ID.
+ * Same speed is assigned consistently for the same ID.
+ */
+USTRUCT()
+struct MASSMOVEMENT_API FMassMovementStyleSpeedParameters
+{
+	GENERATED_BODY()
+
+	/** Desired speed */
+	UPROPERTY(EditAnywhere, Category = "Movement Style", meta = (UIMin = 0.0, ClampMin = 0.0, ForceUnits="cm/s"))
+	float Speed = 140.0f;
+
+	/** How much default desired speed is varied randomly. */
+	UPROPERTY(EditAnywhere, Category = "Movement Style", meta = (ClampMin = 0.0, ClampMax = 1.0))
+	float Variance = 0.1f;
+
+	/** Probability to assign this speed. */
+	UPROPERTY(EditAnywhere, Category = "Movement Style", meta = (UIMin = 1, ClampMin = 1, UIMax = 100, ClampMax = 100))
+	float Probability = 1.0f;
+
+	/** Probability threshold for this style, used to faster lookup. Update via FMassMovementConfig::Update(). */
+	float ProbabilityThreshold = 0.0f;
+};
+
+/** Movement style parameters. A movement style abstracts movement properties for specific style. Behaviors can refer to specific styles when handling speeds. */
+USTRUCT()
+struct MASSMOVEMENT_API FMassMovementStyleParameters
+{
+	GENERATED_BODY()
+
+	/** Style of the movement */
+	UPROPERTY(EditAnywhere, Category = "Movement Style")
+	FMassMovementStyleRef Style;
+
+	/** Array of desired speeds (cm/s) assigned to agents based on probability. */
+	UPROPERTY(EditAnywhere, Category = "Movement Style")
+	TArray<FMassMovementStyleSpeedParameters> DesiredSpeeds;
+};
+
+
diff --git a/Engine/Plugins/Runtime/MassGameplay/Source/MassMovement/Public/Movement/MassMovementProcessors.h b/Engine/Plugins/Runtime/MassGameplay/Source/MassMovement/Public/Movement/MassMovementProcessors.h
new file mode 100644
--- /dev/null
+++ b/Engine/Plugins/Runtime/MassGameplay/Source/MassMovement/Public/Movement/MassMovementProcessors.h
@@ -0,0 +1,26 @@
+﻿// Copyright Epic Games, Inc. All Rights Reserved.
+
+#pragma once
+
+#include "MassProcessor.h"
+#include "MassMovementProcessors.generated.h"
+
+/**
+ * Updates entities position based on force and velocity.
+ * Not applied on Off-LOD entities.
+ */
+UCLASS()
+class MASSMOVEMENT_API UMassApplyMovementProcessor : public UMassProcessor
+{
+	GENERATED_BODY()
+
+public:
+	UMassApplyMovementProcessor();
+
+protected:
+	virtual void ConfigureQueries() override;
+	virtual void Execute(UMassEntitySubsystem& EntitySubsystem, FMassExecutionContext& Context) override;
+
+private:
+	FMassEntityQuery EntityQuery;
+};
diff --git a/Engine/Plugins/Runtime/MassGameplay/Source/MassMovement/Public/Movement/MassMovementTrait.h b/Engine/Plugins/Runtime/MassGameplay/Source/MassMovement/Public/Movement/MassMovementTrait.h
new file mode 100644
--- /dev/null
+++ b/Engine/Plugins/Runtime/MassGameplay/Source/MassMovement/Public/Movement/MassMovementTrait.h
@@ -0,0 +1,19 @@
+﻿// Copyright Epic Games, Inc. All Rights Reserved.
+
+#pragma once 
+
+#include "MassEntityTraitBase.h"
+#include "MassMovementFragments.h"
+#include "MassMovementTrait.generated.h"
+
+UCLASS(meta = (DisplayName = "Movement"))
+class MASSMOVEMENT_API UMassMovementTrait : public UMassEntityTraitBase
+{
+	GENERATED_BODY()
+
+protected:
+	virtual void BuildTemplate(FMassEntityTemplateBuildContext& BuildContext, UWorld& World) const override;
+
+	UPROPERTY(Category="Movement", EditAnywhere)
+	FMassMovementParameters Movement;
+};
diff --git a/Engine/Plugins/Runtime/MassGameplay/Source/MassMovementEditor/Public/IMassMovementEditor.h b/Engine/Plugins/Runtime/MassGameplay/Source/MassMovementEditor/Public/IMassMovementEditor.h
new file mode 100644
--- /dev/null
+++ b/Engine/Plugins/Runtime/MassGameplay/Source/MassMovementEditor/Public/IMassMovementEditor.h
@@ -0,0 +1,15 @@
+// Copyright Epic Games, Inc. All Rights Reserved.
+
+#pragma once
+
+#include "CoreMinimal.h"
+#include "Toolkits/AssetEditorToolkit.h"
+
+/** MassMovement Editor public interface */
+class MASSMOVEMENTEDITOR_API IMassMovementEditor : public FAssetEditorToolkit
+{
+public:
+
+};
+
+
diff --git a/Engine/Plugins/Runtime/MassGameplay/Source/MassMovementEditor/Public/MassMovementEditorModule.h b/Engine/Plugins/Runtime/MassGameplay/Source/MassMovementEditor/Public/MassMovementEditorModule.h
new file mode 100644
--- /dev/null
+++ b/Engine/Plugins/Runtime/MassGameplay/Source/MassMovementEditor/Public/MassMovementEditorModule.h
@@ -0,0 +1,25 @@
+// Copyright Epic Games, Inc. All Rights Reserved.
+
+#pragma  once
+
+#include "CoreMinimal.h"
+#include "Modules/ModuleManager.h"
+#include "AssetTypeCategories.h"
+#include "Toolkits/IToolkitHost.h"
+#include "Toolkits/AssetEditorToolkit.h"
+
+
+class IMassMovementEditor;
+
+/**
+* The public interface to this module
+*/
+class MASSMOVEMENTEDITOR_API FMassMovementEditorModule : public IModuleInterface//, public IHasMenuExtensibility, public IHasToolBarExtensibility
+{
+public:
+	// Begin IModuleInterface
+	virtual void StartupModule() override;
+	virtual void ShutdownModule() override;
+
+private:
+};
diff --git a/Engine/Plugins/Runtime/MassGameplay/Source/MassReplication/Public/MassReplicationTransformHandlers.h b/Engine/Plugins/Runtime/MassGameplay/Source/MassReplication/Public/MassReplicationTransformHandlers.h
--- a/Engine/Plugins/Runtime/MassGameplay/Source/MassReplication/Public/MassReplicationTransformHandlers.h
+++ b/Engine/Plugins/Runtime/MassGameplay/Source/MassReplication/Public/MassReplicationTransformHandlers.h
@@ -2,7 +2,6 @@
 
 #pragma once
 
-#include "MassMovementTypes.h"
 #include "MassReplicationTypes.h"
 #include "MassEntityView.h"
 #include "MassClientBubbleHandler.h"
diff --git a/Engine/Plugins/Runtime/ZoneGraphAnnotations/Source/ZoneGraphAnnotations/Public/ZoneGraphAnnotationTypes.h b/Engine/Plugins/Runtime/ZoneGraphAnnotations/Source/ZoneGraphAnnotations/Public/ZoneGraphAnnotationTypes.h
--- a/Engine/Plugins/Runtime/ZoneGraphAnnotations/Source/ZoneGraphAnnotations/Public/ZoneGraphAnnotationTypes.h
+++ b/Engine/Plugins/Runtime/ZoneGraphAnnotations/Source/ZoneGraphAnnotations/Public/ZoneGraphAnnotationTypes.h
@@ -33,6 +33,8 @@ public:
 
 	uint64 GetValue() const { return Value; }
 
+	bool IsValid() const { return Value != MAX_uint64; }
+
 private:
 	static uint64 NextUniqueID;
 	
