Exported public types in MassSimulationLOD.h required for variable tick rate setup.

diff --git a/Engine/Plugins/Runtime/MassGameplay/Source/MassLOD/Public/MassSimulationLOD.h b/Engine/Plugins/Runtime/MassGameplay/Source/MassLOD/Public/MassSimulationLOD.h
--- a/Engine/Plugins/Runtime/MassGameplay/Source/MassLOD/Public/MassSimulationLOD.h
+++ b/Engine/Plugins/Runtime/MassGameplay/Source/MassLOD/Public/MassSimulationLOD.h
@@ -37,7 +37,7 @@ struct MASSLOD_API FMassSimulationVariableTickFragment : public FMassFragment
 };
 
 USTRUCT()
-struct FMassSimulationVariableTickChunkFragment : public FMassVariableTickChunkFragment
+struct MASSLOD_API FMassSimulationVariableTickChunkFragment : public FMassVariableTickChunkFragment
 {
 	GENERATED_BODY();
 
@@ -60,7 +60,7 @@ struct FMassSimulationVariableTickChunkFragment : public FMassVariableTickChunkF
 };
 
 USTRUCT()
-struct FMassSimulationLODConfig : public FMassSharedFragment
+struct MASSLOD_API FMassSimulationLODConfig : public FMassSharedFragment
 {
 	GENERATED_BODY()
 
@@ -84,7 +84,7 @@ struct FMassSimulationLODConfig : public FMassSharedFragment
 };
 
 USTRUCT()
-struct FMassSimulationVariableTickConfig : public FMassSharedFragment
+struct MASSLOD_API FMassSimulationVariableTickConfig : public FMassSharedFragment
 {
 	GENERATED_BODY()
 
@@ -100,7 +100,7 @@ struct FMassSimulationVariableTickConfig : public FMassSharedFragment
 };
 
 USTRUCT()
-struct FMassSimulationLODSharedFragment : public FMassSharedFragment
+struct MASSLOD_API FMassSimulationLODSharedFragment : public FMassSharedFragment
 {
 	GENERATED_BODY()
 
@@ -113,7 +113,7 @@ struct FMassSimulationLODSharedFragment : public FMassSharedFragment
 };
 
 USTRUCT()
-struct FMassSimulationVariableTickSharedFragment : public FMassSharedFragment
+struct MASSLOD_API FMassSimulationVariableTickSharedFragment : public FMassSharedFragment
 {
 	GENERATED_BODY()
 
