Fix an optimization on the SimulationLODProcessor.

#rnx

diff --git a/Engine/Plugins/Runtime/MassGameplay/Source/MassLOD/Private/MassSimulationLOD.cpp b/Engine/Plugins/Runtime/MassGameplay/Source/MassLOD/Private/MassSimulationLOD.cpp
--- a/Engine/Plugins/Runtime/MassGameplay/Source/MassLOD/Private/MassSimulationLOD.cpp
+++ b/Engine/Plugins/Runtime/MassGameplay/Source/MassLOD/Private/MassSimulationLOD.cpp
@@ -87,7 +87,7 @@ void UMassSimulationLODProcessor::ConfigureQueries()
 		const FMassSimulationLODSharedFragment& LODSharedFragment = Context.GetSharedFragment<FMassSimulationLODSharedFragment>();
 		return LODSharedFragment.bHasAdjustedDistancesFromCount;
 	});
-	EntityQueryCalculateLOD.SetChunkFilter(&FMassSimulationVariableTickSharedFragment::ShouldAdjustLODFromCountForChunk);
+	EntityQueryAdjustDistances.SetChunkFilter(&FMassSimulationVariableTickSharedFragment::ShouldAdjustLODFromCountForChunk);
 
 	EntityQueryVariableTick.AddRequirement<FMassSimulationLODFragment>(EMassFragmentAccess::ReadOnly);
 	EntityQueryVariableTick.AddRequirement<FMassSimulationVariableTickFragment>(EMassFragmentAccess::ReadWrite);
