Fix warnings spamming output

#rnx

diff --git a/Engine/Plugins/Runtime/MassGameplay/Source/MassRepresentation/Private/MassRepresentationProcessor.cpp b/Engine/Plugins/Runtime/MassGameplay/Source/MassRepresentation/Private/MassRepresentationProcessor.cpp
--- a/Engine/Plugins/Runtime/MassGameplay/Source/MassRepresentation/Private/MassRepresentationProcessor.cpp
+++ b/Engine/Plugins/Runtime/MassGameplay/Source/MassRepresentation/Private/MassRepresentationProcessor.cpp
@@ -230,6 +230,11 @@ void UMassRepresentationProcessor::Execute(UMassEntitySubsystem& InEntitySubsyst
 
 bool UMassRepresentationProcessor::ReleaseActorOrCancelSpawning(UMassRepresentationSubsystem& RepresentationSubsystem, const FMassEntityHandle MassAgent, FDataFragment_Actor& ActorInfo, const int16 TemplateActorIndex, FMassActorSpawnRequestHandle& SpawnRequestHandle, FMassCommandBuffer& CommandBuffer, bool bCancelSpawningOnly /*= false*/)
 {
+	if (TemplateActorIndex == INDEX_NONE)
+	{
+		// Nothing to release
+		return false;
+	}
 	check(!ActorInfo.IsValid() || ActorInfo.IsOwnedByMass());
 
 	AActor* Actor = ActorInfo.GetOwnedByMassMutable();
