Fix non-unity errors

diff --git a/Engine/Plugins/Runtime/MassGameplay/Source/MassSpawner/Private/MassEntityZoneGraphSpawnPointsGenerator.cpp b/Engine/Plugins/Runtime/MassGameplay/Source/MassSpawner/Private/MassEntityZoneGraphSpawnPointsGenerator.cpp
--- a/Engine/Plugins/Runtime/MassGameplay/Source/MassSpawner/Private/MassEntityZoneGraphSpawnPointsGenerator.cpp
+++ b/Engine/Plugins/Runtime/MassGameplay/Source/MassSpawner/Private/MassEntityZoneGraphSpawnPointsGenerator.cpp
@@ -5,6 +5,7 @@
 #include "ZoneGraphSubsystem.h"
 #include "ZoneGraphQuery.h"
 #include "MassSpawnerTypes.h"
+#include "Engine/World.h"
 #include "VisualLogger/VisualLogger.h"
 #include "MassGameplaySettings.h"
 #include "MassSpawnLocationProcessor.h"
diff --git a/Engine/Source/Runtime/CUDA/Source/Private/CudaModule.cpp b/Engine/Source/Runtime/CUDA/Source/Private/CudaModule.cpp
--- a/Engine/Source/Runtime/CUDA/Source/Private/CudaModule.cpp
+++ b/Engine/Source/Runtime/CUDA/Source/Private/CudaModule.cpp
@@ -16,6 +16,7 @@
     #include "Windows/HideWindowsPlatformTypes.h"
 #endif
 
+#include "RHI.h"
 #include "HAL/UnrealMemory.h"
 
 DEFINE_LOG_CATEGORY(LogCUDA);
