Split all the specific visualization code of the representation processor into visualization processor.
Move into a separate processor any custom visualization update code.

#rnx

diff --git a/Engine/Plugins/AI/MassCrowd/Source/MassCrowd/Private/MassCrowdRepresentationProcessor.cpp b/Engine/Plugins/AI/MassCrowd/Source/MassCrowd/Private/MassCrowdRepresentationProcessor.cpp
deleted file mode 100644
diff --git a/Engine/Plugins/AI/MassCrowd/Source/MassCrowd/Private/MassCrowdServerRepresentationProcessor.cpp b/Engine/Plugins/AI/MassCrowd/Source/MassCrowd/Private/MassCrowdServerRepresentationProcessor.cpp
--- a/Engine/Plugins/AI/MassCrowd/Source/MassCrowd/Private/MassCrowdServerRepresentationProcessor.cpp
+++ b/Engine/Plugins/AI/MassCrowd/Source/MassCrowd/Private/MassCrowdServerRepresentationProcessor.cpp
@@ -8,5 +8,13 @@ UMassCrowdServerRepresentationProcessor::UMassCrowdServerRepresentationProcessor
 {
 	ExecutionFlags = (int32)EProcessorExecutionFlags::Server;
 
+	bAutoRegisterWithProcessingPhases = true;
+
 	ExecutionOrder.ExecuteAfter.Add(UE::Mass::ProcessorGroupNames::LOD);
-}
\ No newline at end of file
+}
+
+void UMassCrowdServerRepresentationProcessor::ConfigureQueries()
+{
+	Super::ConfigureQueries();
+	EntityQuery.AddTagRequirement<FTagFragment_MassCrowd>(EMassFragmentPresence::All);
+}
diff --git a/Engine/Plugins/AI/MassCrowd/Source/MassCrowd/Private/MassCrowdVisualizationProcessor.cpp b/Engine/Plugins/AI/MassCrowd/Source/MassCrowd/Private/MassCrowdVisualizationProcessor.cpp
--- a/Engine/Plugins/AI/MassCrowd/Source/MassCrowd/Private/MassCrowdVisualizationProcessor.cpp
+++ b/Engine/Plugins/AI/MassCrowd/Source/MassCrowd/Private/MassCrowdVisualizationProcessor.cpp
@@ -28,6 +28,8 @@ UMassCrowdVisualizationProcessor::UMassCrowdVisualizationProcessor()
 {
 	ExecutionFlags = (int32)(EProcessorExecutionFlags::Client | EProcessorExecutionFlags::Standalone);
 
+	bAutoRegisterWithProcessingPhases = true;
+
 	ExecutionOrder.ExecuteAfter.Add(UE::Mass::ProcessorGroupNames::LOD);
 
 	bRequiresGameThreadExecution = true;
@@ -36,15 +38,41 @@ UMassCrowdVisualizationProcessor::UMassCrowdVisualizationProcessor()
 void UMassCrowdVisualizationProcessor::ConfigureQueries()
 {
 	Super::ConfigureQueries();
+	EntityQuery.AddTagRequirement<FTagFragment_MassCrowd>(EMassFragmentPresence::All);
 }
 
-void UMassCrowdVisualizationProcessor::Execute(UMassEntitySubsystem& EntitySubsystem, FMassExecutionContext& Context)
+//----------------------------------------------------------------------//
+// UMassDebugCrowdVisualizationProcessor
+//----------------------------------------------------------------------//
+UMassDebugCrowdVisualizationProcessor::UMassDebugCrowdVisualizationProcessor()
 {
-	EntityQuery.ForEachEntityChunk(EntitySubsystem, Context, [this](FMassExecutionContext& Context)
-	{
-		UpdateVisualization(Context);
-	});
+	ExecutionFlags = (int32)(EProcessorExecutionFlags::Client | EProcessorExecutionFlags::Standalone);
+
+	ExecutionOrder.ExecuteAfter.Add(UE::Mass::ProcessorGroupNames::LOD);
+	ExecutionOrder.ExecuteAfter.Add(UMassCrowdVisualizationProcessor::StaticClass()->GetFName());
+
+	bRequiresGameThreadExecution = true;
+}
 
+void UMassDebugCrowdVisualizationProcessor::ConfigureQueries()
+{
+	EntityQuery.AddTagRequirement<FTagFragment_MassCrowd>(EMassFragmentPresence::All);
+
+	EntityQuery.AddRequirement<FDataFragment_Transform>(EMassFragmentAccess::ReadOnly);
+	EntityQuery.AddRequirement<FMassRepresentationFragment>(EMassFragmentAccess::ReadWrite);
+	EntityQuery.AddRequirement<FDataFragment_Actor>(EMassFragmentAccess::ReadWrite);
+}
+
+void UMassDebugCrowdVisualizationProcessor::Initialize(UObject& Owner)
+{
+	Super::Initialize(Owner);
+
+	World = Owner.GetWorld();
+	check(World);
+}
+
+void UMassDebugCrowdVisualizationProcessor::Execute(UMassEntitySubsystem& EntitySubsystem, FMassExecutionContext& Context)
+{
 	if (UE::MassCrowd::bDebugCrowdVisualType)
 	{
 		check(World);
diff --git a/Engine/Plugins/AI/MassCrowd/Source/MassCrowd/Public/MassCrowdRepresentationProcessor.h b/Engine/Plugins/AI/MassCrowd/Source/MassCrowd/Public/MassCrowdRepresentationProcessor.h
deleted file mode 100644
diff --git a/Engine/Plugins/AI/MassCrowd/Source/MassCrowd/Public/MassCrowdServerRepresentationProcessor.h b/Engine/Plugins/AI/MassCrowd/Source/MassCrowd/Public/MassCrowdServerRepresentationProcessor.h
--- a/Engine/Plugins/AI/MassCrowd/Source/MassCrowd/Public/MassCrowdServerRepresentationProcessor.h
+++ b/Engine/Plugins/AI/MassCrowd/Source/MassCrowd/Public/MassCrowdServerRepresentationProcessor.h
@@ -2,7 +2,7 @@
 
 #pragma once
 
-#include "MassCrowdRepresentationProcessor.h"
+#include "MassRepresentationProcessor.h"
 
 #include "MassCrowdServerRepresentationProcessor.generated.h"
 
@@ -11,10 +11,13 @@
  * It is the counter part of the crowd visualization processor on the client.
  */
 UCLASS(meta = (DisplayName = "Mass Crowd Server Representation"))
-class MASSCROWD_API UMassCrowdServerRepresentationProcessor : public UMassCrowdRepresentationProcessor
+class MASSCROWD_API UMassCrowdServerRepresentationProcessor : public UMassRepresentationProcessor
 {
 	GENERATED_BODY()
 
 public:
 	UMassCrowdServerRepresentationProcessor();
+
+	/** Configure the owned FMassEntityQuery instances to express processor's requirements */
+	virtual void ConfigureQueries() override;
 };
\ No newline at end of file
diff --git a/Engine/Plugins/AI/MassCrowd/Source/MassCrowd/Public/MassCrowdVisualizationProcessor.h b/Engine/Plugins/AI/MassCrowd/Source/MassCrowd/Public/MassCrowdVisualizationProcessor.h
--- a/Engine/Plugins/AI/MassCrowd/Source/MassCrowd/Public/MassCrowdVisualizationProcessor.h
+++ b/Engine/Plugins/AI/MassCrowd/Source/MassCrowd/Public/MassCrowdVisualizationProcessor.h
@@ -2,7 +2,7 @@
 
 #pragma once
 
-#include "MassCrowdRepresentationProcessor.h"
+#include "MassRepresentationProcessor.h"
 
 #include "MassCrowdVisualizationProcessor.generated.h"
 
@@ -10,7 +10,7 @@
  * Overridden visualization processor to make it tied to the crowd via the requirements
  */
 UCLASS(meta = (DisplayName = "Mass Crowd Visualization"))
-class MASSCROWD_API UMassCrowdVisualizationProcessor : public UMassCrowdRepresentationProcessor
+class MASSCROWD_API UMassCrowdVisualizationProcessor : public UMassVisualizationProcessor
 {
 	GENERATED_BODY()
 public:
@@ -20,6 +20,24 @@ protected:
 
 	/** Configure the owned FMassEntityQuery instances to express processor's requirements */
 	virtual void ConfigureQueries() override;
+};
+
+/**
+ * A custom visualization processor for debugging mass crowd
+ */
+UCLASS(meta = (DisplayName = "Mass Crowd Visualization"))
+class MASSCROWD_API UMassDebugCrowdVisualizationProcessor : public UMassProcessor
+{
+	GENERATED_BODY()
+public:
+	UMassDebugCrowdVisualizationProcessor();
+
+protected:
+
+	/** Configure the owned FMassEntityQuery instances to express processor's requirements */
+	virtual void ConfigureQueries() override;
+
+	virtual void Initialize(UObject& Owner) override;
 
 	/**
 	 * Execution method for this processor
@@ -27,4 +45,9 @@ protected:
 	 * @param Context is the execution context to be passed when executing the lambdas
 	 */
 	virtual void Execute(UMassEntitySubsystem& EntitySubsystem, FMassExecutionContext& Context) override;
+
+	UPROPERTY(Transient)
+	UWorld* World;
+
+	FMassEntityQuery EntityQuery;
 };
\ No newline at end of file
diff --git a/Engine/Plugins/Runtime/MassGameplay/Source/MassRepresentation/Private/MassRepresentationProcessor.cpp b/Engine/Plugins/Runtime/MassGameplay/Source/MassRepresentation/Private/MassRepresentationProcessor.cpp
--- a/Engine/Plugins/Runtime/MassGameplay/Source/MassRepresentation/Private/MassRepresentationProcessor.cpp
+++ b/Engine/Plugins/Runtime/MassGameplay/Source/MassRepresentation/Private/MassRepresentationProcessor.cpp
@@ -33,7 +33,6 @@ void UMassRepresentationProcessor::ConfigureQueries()
 	EntityQuery.AddRequirement<FMassRepresentationFragment>(EMassFragmentAccess::ReadWrite);
 	EntityQuery.AddRequirement<FMassRepresentationLODFragment>(EMassFragmentAccess::ReadOnly);
 	EntityQuery.AddRequirement<FDataFragment_Actor>(EMassFragmentAccess::ReadWrite);
-	EntityQuery.AddChunkRequirement<FMassVisualizationChunkFragment>(EMassFragmentAccess::ReadWrite);
 	EntityQuery.AddConstSharedRequirement<FMassRepresentationConfig>();
 	EntityQuery.AddSharedRequirement<FMassRepresentationSubsystemFragment>(EMassFragmentAccess::ReadWrite);
 }
@@ -221,7 +220,7 @@ void UMassRepresentationProcessor::UpdateRepresentation(FMassExecutionContext& C
 
 void UMassRepresentationProcessor::Execute(UMassEntitySubsystem& InEntitySubsystem, FMassExecutionContext& Context)
 {
-	// Visualize entities
+	// Update entities representation
 	EntityQuery.ForEachEntityChunk(InEntitySubsystem, Context, [this](FMassExecutionContext& Context)
 	{
 		UpdateRepresentation(Context);
@@ -263,7 +262,10 @@ bool UMassRepresentationProcessor::ReleaseActorOrCancelSpawning(UMassRepresentat
 	return false;
 }
 
-void UMassRepresentationProcessor::UpdateVisualization(FMassExecutionContext& Context)
+//----------------------------------------------------------------------//
+// UMassVisualizationProcessor 
+//----------------------------------------------------------------------//
+void UMassVisualizationProcessor::UpdateVisualization(FMassExecutionContext& Context)
 {
 	FMassVisualizationChunkFragment& ChunkData = UpdateChunkVisibility(Context);
 	if (!ChunkData.ShouldUpdateVisualization())
@@ -287,7 +289,13 @@ void UMassRepresentationProcessor::UpdateVisualization(FMassExecutionContext& Co
 	}
 }
 
-FMassVisualizationChunkFragment& UMassRepresentationProcessor::UpdateChunkVisibility(FMassExecutionContext& Context) const
+void UMassVisualizationProcessor::ConfigureQueries()
+{
+	Super::ConfigureQueries();
+	EntityQuery.AddChunkRequirement<FMassVisualizationChunkFragment>(EMassFragmentAccess::ReadWrite);
+}
+
+FMassVisualizationChunkFragment& UMassVisualizationProcessor::UpdateChunkVisibility(FMassExecutionContext& Context) const
 {
 	const FMassRepresentationConfig& RepresentationConfig = Context.GetConstSharedFragment<FMassRepresentationConfig>();
 	bool bFirstUpdate = false;
@@ -330,7 +338,7 @@ FMassVisualizationChunkFragment& UMassRepresentationProcessor::UpdateChunkVisibi
 	return ChunkData;
 }
 
-void UMassRepresentationProcessor::UpdateEntityVisibility(const FMassEntityHandle Entity, const FMassRepresentationFragment& Representation, const FMassRepresentationLODFragment& RepresentationLOD, FMassVisualizationChunkFragment& ChunkData, FMassCommandBuffer& CommandBuffer)
+void UMassVisualizationProcessor::UpdateEntityVisibility(const FMassEntityHandle Entity, const FMassRepresentationFragment& Representation, const FMassRepresentationLODFragment& RepresentationLOD, FMassVisualizationChunkFragment& ChunkData, FMassCommandBuffer& CommandBuffer)
 {
 	// Move the visible entities together into same chunks so we can skip entire chunk when not visible as an optimization
 	const EMassVisibility Visibility = Representation.CurrentRepresentation != ERepresentationType::None ? 
@@ -343,6 +351,14 @@ void UMassRepresentationProcessor::UpdateEntityVisibility(const FMassEntityHandl
 	}
 }
 
+void UMassVisualizationProcessor::Execute(UMassEntitySubsystem& EntitySubsystem, FMassExecutionContext& Context)
+{
+	// Update entities visualization
+	EntityQuery.ForEachEntityChunk(EntitySubsystem, Context, [this](FMassExecutionContext& Context)
+	{
+		UpdateVisualization(Context);
+	});
+}
 
 //----------------------------------------------------------------------//
 // UMassRepresentationFragmentDestructor 
diff --git a/Engine/Plugins/Runtime/MassGameplay/Source/MassRepresentation/Public/MassRepresentationProcessor.h b/Engine/Plugins/Runtime/MassGameplay/Source/MassRepresentation/Public/MassRepresentationProcessor.h
--- a/Engine/Plugins/Runtime/MassGameplay/Source/MassRepresentation/Public/MassRepresentationProcessor.h
+++ b/Engine/Plugins/Runtime/MassGameplay/Source/MassRepresentation/Public/MassRepresentationProcessor.h
@@ -50,6 +50,39 @@ protected:
 	 */
 	bool ReleaseActorOrCancelSpawning(UMassRepresentationSubsystem& RepresentationSubsystem, const FMassEntityHandle MassAgent, FDataFragment_Actor& ActorInfo, const int16 TemplateActorIndex, FMassActorSpawnRequestHandle& SpawnRequestHandle, FMassCommandBuffer& CommandBuffer, const bool bCancelSpawningOnly = false);
 
+	/*
+	 * Update representation type for each entity, must be called within a ForEachEntityChunk
+	 * @param Context of the execution from the entity sub system
+	 */
+	void UpdateRepresentation(FMassExecutionContext& Context);
+
+	/** Caching ptr to our associated world */
+	UPROPERTY(Transient)
+	UWorld* World;
+
+	UPROPERTY(Transient)
+	UMassEntitySubsystem* CachedEntitySubsystem;
+
+	FMassEntityQuery EntityQuery;
+};
+
+UCLASS()
+class MASSREPRESENTATION_API UMassVisualizationProcessor : public UMassRepresentationProcessor
+{
+	GENERATED_BODY()
+
+protected:
+
+	/** Configure the owned FMassEntityQuery instances to express processor's requirements */
+	virtual void ConfigureQueries() override;
+
+	/**
+	 * Execution method for this processor
+	 * @param EntitySubsystem is the system to execute the lambdas on each entity chunk
+	 * @param Context is the execution context to be passed when executing the lambdas
+	 */
+	virtual void Execute(UMassEntitySubsystem& EntitySubsystem, FMassExecutionContext& Context) override;
+
 	/**
 	 * Updates chunk visibility info for later chunk logic optimization
 	 * @param Context of the execution from the entity sub system
@@ -67,32 +100,14 @@ protected:
 	 */
 	static void UpdateEntityVisibility(const FMassEntityHandle Entity, const FMassRepresentationFragment& Representation, const FMassRepresentationLODFragment& RepresentationLOD, FMassVisualizationChunkFragment& ChunkData, FMassCommandBuffer& CommandBuffer);
 
-protected:
-
-	/*
-	 * Update representation type for each entity, must be called within a ForEachEntityChunk
-	 * @param Context of the execution from the entity sub system
-	 */
-	void UpdateRepresentation(FMassExecutionContext& Context);
-
-	/** 
+	/**
 	 * Update representation and visibility for each entity, must be called within a ForEachEntityChunk
 	 * @param Context of the execution from the entity sub system
 	 */
 	void UpdateVisualization(FMassExecutionContext& Context);
-
-	/** Caching ptr to our associated world */
-	UPROPERTY(Transient)
-	UWorld* World;
-
-	UPROPERTY(Transient)
-	UMassEntitySubsystem* CachedEntitySubsystem;
-
-
-	FMassEntityQuery EntityQuery;
-
 };
 
+
 UCLASS()
 class MASSREPRESENTATION_API UMassRepresentationFragmentDestructor : public UMassFragmentDeinitializer
 {
