Fix non unity build: Issue #139199

diff --git a/Engine/Plugins/Runtime/MassGameplay/Source/MassSpawner/Private/MassEntityZoneGraphSpawnPointsGenerator.cpp b/Engine/Plugins/Runtime/MassGameplay/Source/MassSpawner/Private/MassEntityZoneGraphSpawnPointsGenerator.cpp
--- a/Engine/Plugins/Runtime/MassGameplay/Source/MassSpawner/Private/MassEntityZoneGraphSpawnPointsGenerator.cpp
+++ b/Engine/Plugins/Runtime/MassGameplay/Source/MassSpawner/Private/MassEntityZoneGraphSpawnPointsGenerator.cpp
@@ -1,7 +1,6 @@
 // Copyright Epic Games, Inc. All Rights Reserved.
 
 #include "MassEntityZoneGraphSpawnPointsGenerator.h"
-#include "ZoneGraphData.h"
 #include "ZoneGraphSubsystem.h"
 #include "ZoneGraphQuery.h"
 #include "MassSpawnerTypes.h"
@@ -9,6 +8,7 @@
 #include "VisualLogger/VisualLogger.h"
 #include "MassGameplaySettings.h"
 #include "MassSpawnLocationProcessor.h"
+#include "Engine/World.h"
 
 void UMassEntityZoneGraphSpawnPointsGenerator::Generate(UObject& QueryOwner, TConstArrayView<FMassSpawnedEntityType> EntityTypes, int32 Count, FFinishedGeneratingSpawnDataSignature& FinishedGeneratingSpawnPointsDelegate) const
 {
