Moved LOD Fragments into their own file

#rnx

diff --git a/Engine/Plugins/AI/MassAI/Source/MassAIBehavior/Private/MassLookAtProcessors.cpp b/Engine/Plugins/AI/MassAI/Source/MassAIBehavior/Private/MassLookAtProcessors.cpp
--- a/Engine/Plugins/AI/MassAI/Source/MassAIBehavior/Private/MassLookAtProcessors.cpp
+++ b/Engine/Plugins/AI/MassAI/Source/MassAIBehavior/Private/MassLookAtProcessors.cpp
@@ -17,6 +17,7 @@
 #include "BezierUtilities.h"
 #include "Algo/RandomShuffle.h"
 #include "Engine/World.h"
+#include "MassLODFragments.h"
 
 #define UNSAFE_FOR_MT 1
 
diff --git a/Engine/Plugins/AI/MassCrowd/Source/MassCrowd/Private/MassCrowdServerRepresentationTrait.cpp b/Engine/Plugins/AI/MassCrowd/Source/MassCrowd/Private/MassCrowdServerRepresentationTrait.cpp
--- a/Engine/Plugins/AI/MassCrowd/Source/MassCrowd/Private/MassCrowdServerRepresentationTrait.cpp
+++ b/Engine/Plugins/AI/MassCrowd/Source/MassCrowd/Private/MassCrowdServerRepresentationTrait.cpp
@@ -4,6 +4,7 @@
 #include "MassEntityTemplateRegistry.h"
 #include "MassCrowdRepresentationSubsystem.h"
 #include "MassCommonFragments.h"
+#include "MassLODFragments.h"
 #include "MassRepresentationFragments.h"
 #include "Engine/World.h"
 #include "MassCrowdRepresentationActorManagement.h"
diff --git a/Engine/Plugins/Runtime/MassGameplay/Source/MassLOD/Private/MassLODTrait.cpp b/Engine/Plugins/Runtime/MassGameplay/Source/MassLOD/Private/MassLODTrait.cpp
--- a/Engine/Plugins/Runtime/MassGameplay/Source/MassLOD/Private/MassLODTrait.cpp
+++ b/Engine/Plugins/Runtime/MassGameplay/Source/MassLOD/Private/MassLODTrait.cpp
@@ -5,7 +5,7 @@
 #include "Engine/World.h"
 #include "MassCommonFragments.h"
 #include "StructUtilsTypes.h"
-#include "MassLODTypes.h"
+#include "MassLODFragments.h"
 
 void UMassLODCollectorTrait::BuildTemplate(FMassEntityTemplateBuildContext& BuildContext, UWorld& World) const
 {
diff --git a/Engine/Plugins/Runtime/MassGameplay/Source/MassLOD/Public/MassLODCollectorProcessor.h b/Engine/Plugins/Runtime/MassGameplay/Source/MassLOD/Public/MassLODCollectorProcessor.h
--- a/Engine/Plugins/Runtime/MassGameplay/Source/MassLOD/Public/MassLODCollectorProcessor.h
+++ b/Engine/Plugins/Runtime/MassGameplay/Source/MassLOD/Public/MassLODCollectorProcessor.h
@@ -3,7 +3,6 @@
 #pragma once
 
 #include "MassEntityTypes.h"
-#include "MassLODTypes.h"
 #include "MassLODManager.h"
 #include "MassLODCollector.h"
 
diff --git a/Engine/Plugins/Runtime/MassGameplay/Source/MassLOD/Public/MassLODFragments.h b/Engine/Plugins/Runtime/MassGameplay/Source/MassLOD/Public/MassLODFragments.h
new file mode 100644
--- /dev/null
+++ b/Engine/Plugins/Runtime/MassGameplay/Source/MassLOD/Public/MassLODFragments.h
@@ -0,0 +1,208 @@
+// Copyright Epic Games, Inc. All Rights Reserved.
+
+#pragma once
+
+#include "MassLODTypes.h"
+
+#include "MassLODFragments.generated.h"
+
+USTRUCT()
+struct MASSLOD_API FMassHighLODTag : public FMassTag
+{
+	GENERATED_BODY()
+};
+
+USTRUCT()
+struct MASSLOD_API FMassMediumLODTag : public FMassTag
+{
+	GENERATED_BODY()
+};
+
+USTRUCT()
+struct MASSLOD_API FMassLowLODTag : public FMassTag
+{
+	GENERATED_BODY()
+};
+
+USTRUCT()
+struct MASSLOD_API FMassOffLODTag : public FMassTag
+{
+	GENERATED_BODY()
+};
+
+/*
+ * Data fragment to store the calculated distances to viewers
+ */
+USTRUCT()
+struct MASSLOD_API FMassViewerInfoFragment : public FMassFragment
+{
+	GENERATED_BODY()
+
+	// Closest viewer distance
+	float ClosestViewerDistanceSq;
+
+	// Closest distance to frustum
+	float ClosestDistanceToFrustum;
+};
+
+USTRUCT()
+struct MASSLOD_API FMassVariableTickChunkFragment : public FMassChunkFragment
+{
+	GENERATED_BODY();
+
+	bool ShouldTickThisFrame() const
+	{
+		return bShouldTickThisFrame;
+	}
+
+	float GetTimeUntilNextTick() const
+	{
+		return TimeUntilNextTick;
+	}
+
+	int32 GetLastChunkSerialModificationNumber() const
+	{
+		return LastChunkSerialModificationNumber;
+	}
+
+	EMassLOD::Type GetLOD() const
+	{
+		return LOD;
+	}
+
+	void SetLOD(EMassLOD::Type InLOD)
+	{
+		checkf(LOD == EMassLOD::Max, TEXT("Chunk LOD should never change, it is allowed to only set it once"))
+		LOD = InLOD;
+	}
+
+	void Update(const bool bInShouldTickThisFrame, const float InTimeUntilNextTick, int32 ChunkSerialModificationNumber)
+	{
+		bShouldTickThisFrame = bInShouldTickThisFrame;
+		TimeUntilNextTick = InTimeUntilNextTick;
+		LastChunkSerialModificationNumber = ChunkSerialModificationNumber;
+	}
+
+private:
+	bool bShouldTickThisFrame = true;
+	EMassLOD::Type LOD = EMassLOD::Max;
+	float TimeUntilNextTick = 0.0f;
+	int32 LastChunkSerialModificationNumber = INDEX_NONE;
+};
+
+USTRUCT()
+struct FMassCollectLODViewerInfoTag : public FMassTag
+{
+	GENERATED_BODY();
+};
+
+USTRUCT()
+struct MASSLOD_API FMassVisibilityCanBeSeenTag : public FMassTag
+{
+	GENERATED_BODY()
+};
+
+USTRUCT()
+struct MASSLOD_API FMassVisibilityCulledByFrustumTag : public FMassTag
+{
+	GENERATED_BODY()
+};
+
+USTRUCT()
+struct MASSLOD_API FMassVisibilityCulledByDistanceTag : public FMassTag
+{
+	GENERATED_BODY()
+};
+
+USTRUCT()
+struct MASSLOD_API FMassVisualizationChunkFragment : public FMassChunkFragment
+{
+	GENERATED_BODY();
+
+	/**
+	 * Fetched existing chunk fragment to know if there is a possibility of a an entity that is visible
+	 * In the case that there is not chunk information, we cannot assume that all entities are not visible.
+	 *
+	 * @param Context of the execution from the entity sub system
+	 * @return true if there is a possibility that the chunk contains a visible entity
+	 */
+	static bool AreAnyEntitiesVisibleInChunk(const FMassExecutionContext& Context)
+	{
+		return Context.GetChunkFragment<FMassVisualizationChunkFragment>().AreAnyEntitiesVisible();
+	}
+
+	/**
+	 * Returns if there could be a visible entities in that chunk
+	 *
+	 * @return true if that is the case
+	 */
+	bool AreAnyEntitiesVisible() const
+	{
+		return (Visibility != EMassVisibility::CulledByDistance && Visibility != EMassVisibility::CulledByFrustum) || bContainsNewlyVisibleEntity;
+	}
+
+	/**
+	 * ShouldUpdateVisualizationForChunk
+	 *
+	* @param Context of the execution from the entity sub system
+	 * @return bool true if we should update the representation type for this chunk
+	 */
+	static bool ShouldUpdateVisualizationForChunk(const FMassExecutionContext& Context)
+	{
+		const FMassVisualizationChunkFragment* ChunkFragment = Context.GetChunkFragmentPtr<FMassVisualizationChunkFragment>();
+		return ChunkFragment == nullptr || ChunkFragment->ShouldUpdateVisualization();
+	}
+
+	/**
+	 * Representation type of all currently visible entities are always updated
+	 * But as an optimization, we use a frequency check on the not visible one.
+	 *
+	 * @return true if we should update the representation type for this chunk
+	 */
+	bool ShouldUpdateVisualization() const
+	{
+		return Visibility != EMassVisibility::CulledByDistance || DeltaTime <= 0.0f;
+	}
+
+	void SetContainsNewlyVisibleEntity(bool bInContainsNewlyVisibleEntity)
+	{
+		if (bInContainsNewlyVisibleEntity)
+		{
+			checkfSlow(Visibility != EMassVisibility::CanBeSeen, TEXT("Something is not adding up, how can an entity be newly visible in a can be seen chunk?"));
+			bContainsNewlyVisibleEntity = true;
+		}
+	}
+
+	void SetVisibility(const EMassVisibility InVisibility)
+	{
+		checkf(Visibility == EMassVisibility::Max, TEXT("Chunk visibility should never change, it is allowed to only set it once"));
+		Visibility = InVisibility;
+	}
+
+	EMassVisibility GetVisibility() const
+	{
+		return Visibility;
+	}
+
+	float GetDeltaTime() const
+	{
+		return DeltaTime;
+	}
+
+	void Update(float InDeltaTime)
+	{
+		bContainsNewlyVisibleEntity = false;
+		DeltaTime = InDeltaTime;
+	}
+
+protected:
+
+	/** Visibility of the current chunk, should never change */
+	EMassVisibility Visibility = EMassVisibility::Max;
+
+	/** Not visible chunks, might contains entity that are newly visible and not yet moved. */
+	bool bContainsNewlyVisibleEntity = true;
+
+	/** Not visible chunks delta time until next update */
+	float DeltaTime = 0;
+};
diff --git a/Engine/Plugins/Runtime/MassGameplay/Source/MassLOD/Public/MassLODTypes.h b/Engine/Plugins/Runtime/MassGameplay/Source/MassLOD/Public/MassLODTypes.h
--- a/Engine/Plugins/Runtime/MassGameplay/Source/MassLOD/Public/MassLODTypes.h
+++ b/Engine/Plugins/Runtime/MassGameplay/Source/MassLOD/Public/MassLODTypes.h
@@ -38,7 +38,7 @@ namespace UE::Mass::ProcessorGroupNames
 // We are not using enum class here because we are doing so many arithmetic operation and comparison on them 
 // that it is not worth polluting int32 casts everywhere in the code.
 UENUM()
-namespace EMassLOD 
+namespace EMassLOD
 {
 	enum Type
 	{
@@ -50,89 +50,6 @@ namespace EMassLOD
 	};
 }
 
-USTRUCT()
-struct MASSLOD_API FMassHighLODTag : public FMassTag
-{
-	GENERATED_BODY()
-};
-
-USTRUCT()
-struct MASSLOD_API FMassMediumLODTag : public FMassTag
-{
-	GENERATED_BODY()
-};
-
-USTRUCT()
-struct MASSLOD_API FMassLowLODTag : public FMassTag
-{
-	GENERATED_BODY()
-};
-
-USTRUCT()
-struct MASSLOD_API FMassOffLODTag : public FMassTag
-{
-	GENERATED_BODY()
-};
-
-/*
- * Data fragment to store the calculated distances to viewers
- */
-USTRUCT()
-struct MASSLOD_API FMassViewerInfoFragment : public FMassFragment
-{
-	GENERATED_BODY()
-
-	// Closest viewer distance
-	float ClosestViewerDistanceSq;
-
-	// Closest distance to frustum
-	float ClosestDistanceToFrustum;
-};
-
-USTRUCT()
-struct MASSLOD_API FMassVariableTickChunkFragment : public FMassChunkFragment
-{
-	GENERATED_BODY();
-
-	bool ShouldTickThisFrame() const
-	{
-		return bShouldTickThisFrame;
-	}
-
-	float GetTimeUntilNextTick() const
-	{
-		return TimeUntilNextTick;
-	}
-
-	int32 GetLastChunkSerialModificationNumber() const
-	{
-		return LastChunkSerialModificationNumber;
-	}
-
-	EMassLOD::Type GetLOD() const
-	{
-		return LOD;
-	}
-
-	void SetLOD(EMassLOD::Type InLOD)
-	{
-		checkf(LOD == EMassLOD::Max, TEXT("Chunk LOD should never change, it is allowed to only set it once"))
-		LOD = InLOD;
-	}
-
-	void Update(const bool bInShouldTickThisFrame, const float InTimeUntilNextTick, int32 ChunkSerialModificationNumber)
-	{
-		bShouldTickThisFrame = bInShouldTickThisFrame;
-		TimeUntilNextTick = InTimeUntilNextTick;
-		LastChunkSerialModificationNumber = ChunkSerialModificationNumber;
-	}
-
-private:
-	bool bShouldTickThisFrame = true;
-	EMassLOD::Type LOD = EMassLOD::Max;
-	float TimeUntilNextTick = 0.0f;
-	int32 LastChunkSerialModificationNumber = INDEX_NONE;
-};
 
 enum class EMassVisibility : uint8
 {
@@ -140,121 +57,4 @@ enum class EMassVisibility : uint8
 	CulledByFrustum, // Not in camera frustum but within visibility distance
 	CulledByDistance, // Too far whether in or out of frustum
 	Max
-};
-
-USTRUCT()
-struct FMassCollectLODViewerInfoTag : public FMassTag
-{
-	GENERATED_BODY();
-};
-
-USTRUCT()
-struct MASSLOD_API FMassVisibilityCanBeSeenTag : public FMassTag
-{
-	GENERATED_BODY()
-};
-
-USTRUCT()
-struct MASSLOD_API FMassVisibilityCulledByFrustumTag : public FMassTag
-{
-	GENERATED_BODY()
-};
-
-USTRUCT()
-struct MASSLOD_API FMassVisibilityCulledByDistanceTag : public FMassTag
-{
-	GENERATED_BODY()
-};
-
-USTRUCT()
-struct MASSLOD_API FMassVisualizationChunkFragment : public FMassChunkFragment
-{
-	GENERATED_BODY();
-
-	/**
-	 * Fetched existing chunk fragment to know if there is a possibility of a an entity that is visible
-	 * In the case that there is not chunk information, we cannot assume that all entities are not visible.
-	 *
-	 * @param Context of the execution from the entity sub system
-	 * @return true if there is a possibility that the chunk contains a visible entity
-	 */
-	static bool AreAnyEntitiesVisibleInChunk(const FMassExecutionContext& Context)
-	{
-		return Context.GetChunkFragment<FMassVisualizationChunkFragment>().AreAnyEntitiesVisible();
-	}
-
-	/**
-	 * Returns if there could be a visible entities in that chunk
-	 *
-	 * @return true if that is the case
-	 */
-	bool AreAnyEntitiesVisible() const
-	{
-		return (Visibility != EMassVisibility::CulledByDistance && Visibility != EMassVisibility::CulledByFrustum) || bContainsNewlyVisibleEntity;
-	}
-
-	/**
-	 * ShouldUpdateVisualizationForChunk
-	 *
-	* @param Context of the execution from the entity sub system
-	 * @return bool true if we should update the representation type for this chunk
-	 */
-	static bool ShouldUpdateVisualizationForChunk(const FMassExecutionContext& Context)
-	{
-		const FMassVisualizationChunkFragment* ChunkFragment = Context.GetChunkFragmentPtr<FMassVisualizationChunkFragment>();
-		return ChunkFragment == nullptr || ChunkFragment->ShouldUpdateVisualization();
-	}
-
-	/**
-	 * Representation type of all currently visible entities are always updated
-	 * But as an optimization, we use a frequency check on the not visible one.
-	 *
-	 * @return true if we should update the representation type for this chunk
-	 */
-	bool ShouldUpdateVisualization() const
-	{
-		return Visibility != EMassVisibility::CulledByDistance || DeltaTime <= 0.0f;
-	}
-
-	void SetContainsNewlyVisibleEntity(bool bInContainsNewlyVisibleEntity)
-	{
-		if (bInContainsNewlyVisibleEntity)
-		{
-			checkfSlow(Visibility != EMassVisibility::CanBeSeen, TEXT("Something is not adding up, how can an entity be newly visible in a can be seen chunk?"));
-			bContainsNewlyVisibleEntity = true;
-		}
-	}
-
-	void SetVisibility(const EMassVisibility InVisibility)
-	{
-		checkf(Visibility == EMassVisibility::Max, TEXT("Chunk visibility should never change, it is allowed to only set it once"));
-		Visibility = InVisibility;
-	}
-
-	EMassVisibility GetVisibility() const
-	{
-		return Visibility;
-	}
-
-	float GetDeltaTime() const
-	{
-		return DeltaTime;
-	}
-
-	void Update(float InDeltaTime)
-	{
-		bContainsNewlyVisibleEntity = false;
-		DeltaTime = InDeltaTime;
-	}
-
-protected:
-
-	/** Visibility of the current chunk, should never change */
-	EMassVisibility Visibility = EMassVisibility::Max;
-
-	/** Not visible chunks, might contains entity that are newly visible and not yet moved. */
-	bool bContainsNewlyVisibleEntity = true;
-
-	/** Not visible chunks delta time until next update */
-	float DeltaTime = 0;
-};
+};
\ No newline at end of file
diff --git a/Engine/Plugins/Runtime/MassGameplay/Source/MassLOD/Public/MassLODUtils.h b/Engine/Plugins/Runtime/MassGameplay/Source/MassLOD/Public/MassLODUtils.h
--- a/Engine/Plugins/Runtime/MassGameplay/Source/MassLOD/Public/MassLODUtils.h
+++ b/Engine/Plugins/Runtime/MassGameplay/Source/MassLOD/Public/MassLODUtils.h
@@ -5,7 +5,7 @@
 #include "MassEntityTypes.h"
 #include "MassEntitySubsystem.h"
 #include "MassCommonTypes.h"
-#include "MassLODTypes.h"
+#include "MassLODFragments.h"
 #include "Logging/LogMacros.h"
 
 MASSLOD_API DECLARE_LOG_CATEGORY_EXTERN(LogMassLOD, Log, All);
diff --git a/Engine/Plugins/Runtime/MassGameplay/Source/MassLOD/Public/MassSimulationLOD.h b/Engine/Plugins/Runtime/MassGameplay/Source/MassLOD/Public/MassSimulationLOD.h
--- a/Engine/Plugins/Runtime/MassGameplay/Source/MassLOD/Public/MassSimulationLOD.h
+++ b/Engine/Plugins/Runtime/MassGameplay/Source/MassLOD/Public/MassSimulationLOD.h
@@ -8,6 +8,7 @@
 #include "MassLODTickRateController.h"
 #include "MassLODLogic.h"
 #include "MassEntityTypes.h"
+#include "MassLODFragments.h"
 
 #include "MassSimulationLOD.generated.h"
 
diff --git a/Engine/Plugins/Runtime/MassGameplay/Source/MassRepresentation/Private/MassUpdateISMProcessor.cpp b/Engine/Plugins/Runtime/MassGameplay/Source/MassRepresentation/Private/MassUpdateISMProcessor.cpp
--- a/Engine/Plugins/Runtime/MassGameplay/Source/MassRepresentation/Private/MassUpdateISMProcessor.cpp
+++ b/Engine/Plugins/Runtime/MassGameplay/Source/MassRepresentation/Private/MassUpdateISMProcessor.cpp
@@ -6,7 +6,7 @@
 #include "MassEntitySubsystem.h"
 #include "MassRepresentationFragments.h"
 #include "MassCommonFragments.h"
-#include "MassLODTypes.h"
+#include "MassLODFragments.h"
 #include "Engine/World.h"
 
 UMassUpdateISMProcessor::UMassUpdateISMProcessor()
diff --git a/Engine/Plugins/Runtime/MassGameplay/Source/MassRepresentation/Private/MassVisualizationTrait.cpp b/Engine/Plugins/Runtime/MassGameplay/Source/MassRepresentation/Private/MassVisualizationTrait.cpp
--- a/Engine/Plugins/Runtime/MassGameplay/Source/MassRepresentation/Private/MassVisualizationTrait.cpp
+++ b/Engine/Plugins/Runtime/MassGameplay/Source/MassRepresentation/Private/MassVisualizationTrait.cpp
@@ -7,6 +7,7 @@
 #include "MassRepresentationFragments.h"
 #include "MassRepresentationActorManagement.h"
 #include "Engine/World.h"
+#include "MassLODFragments.h"
 
 UMassVisualizationTrait::UMassVisualizationTrait()
 {
diff --git a/Engine/Plugins/Runtime/MassGameplay/Source/MassRepresentation/Public/MassRepresentationProcessor.h b/Engine/Plugins/Runtime/MassGameplay/Source/MassRepresentation/Public/MassRepresentationProcessor.h
--- a/Engine/Plugins/Runtime/MassGameplay/Source/MassRepresentation/Public/MassRepresentationProcessor.h
+++ b/Engine/Plugins/Runtime/MassGameplay/Source/MassRepresentation/Public/MassRepresentationProcessor.h
@@ -2,6 +2,7 @@
 
 #pragma once
 
+#include "MassLODFragments.h"
 #include "MassRepresentationFragments.h"
 #include "MassActorSpawnerSubsystem.h"
 #include "MassObserverProcessor.h"
