Rename some Enum to add Mass in them
Changed const FStructView& method paramaters into FConstStructView

#rnx

diff --git a/Engine/Plugins/AI/MassCrowd/Source/MassCrowd/Private/MassCrowdRepresentationActorManagement.cpp b/Engine/Plugins/AI/MassCrowd/Source/MassCrowd/Private/MassCrowdRepresentationActorManagement.cpp
--- a/Engine/Plugins/AI/MassCrowd/Source/MassCrowd/Private/MassCrowdRepresentationActorManagement.cpp
+++ b/Engine/Plugins/AI/MassCrowd/Source/MassCrowd/Private/MassCrowdRepresentationActorManagement.cpp
@@ -8,11 +8,11 @@
 #include "MassAgentComponent.h"
 #include "MassRepresentationSubsystem.h"
 
-void UMassCrowdRepresentationActorManagement::SetActorEnabled(const EActorEnabledType EnabledType, AActor& Actor, const int32 EntityIdx, FMassCommandBuffer& CommandBuffer) const
+void UMassCrowdRepresentationActorManagement::SetActorEnabled(const EMassActorEnabledType EnabledType, AActor& Actor, const int32 EntityIdx, FMassCommandBuffer& CommandBuffer) const
 {
 	Super::SetActorEnabled(EnabledType, Actor, EntityIdx, CommandBuffer);
 
-	const bool bEnabled = EnabledType != EActorEnabledType::Disabled;
+	const bool bEnabled = EnabledType != EMassActorEnabledType::Disabled;
 
 	USkeletalMeshComponent* SkeletalMeshComponent = Actor.FindComponentByClass<USkeletalMeshComponent>();
 	if (SkeletalMeshComponent)
diff --git a/Engine/Plugins/AI/MassCrowd/Source/MassCrowd/Private/MassCrowdServerRepresentationTrait.cpp b/Engine/Plugins/AI/MassCrowd/Source/MassCrowd/Private/MassCrowdServerRepresentationTrait.cpp
--- a/Engine/Plugins/AI/MassCrowd/Source/MassCrowd/Private/MassCrowdServerRepresentationTrait.cpp
+++ b/Engine/Plugins/AI/MassCrowd/Source/MassCrowd/Private/MassCrowdServerRepresentationTrait.cpp
@@ -12,10 +12,10 @@
 UMassCrowdServerRepresentationTrait::UMassCrowdServerRepresentationTrait()
 {
 	Config.RepresentationActorManagementClass = UMassCrowdRepresentationActorManagement::StaticClass();
-	Config.LODRepresentation[EMassLOD::High] = ERepresentationType::HighResSpawnedActor;
-	Config.LODRepresentation[EMassLOD::Medium] = ERepresentationType::None;
-	Config.LODRepresentation[EMassLOD::Low] = ERepresentationType::None;
-	Config.LODRepresentation[EMassLOD::Off] = ERepresentationType::None;
+	Config.LODRepresentation[EMassLOD::High] = EMassRepresentationType::HighResSpawnedActor;
+	Config.LODRepresentation[EMassLOD::Medium] = EMassRepresentationType::None;
+	Config.LODRepresentation[EMassLOD::Low] = EMassRepresentationType::None;
+	Config.LODRepresentation[EMassLOD::Off] = EMassRepresentationType::None;
 	Config.bKeepLowResActors = false;
 	Config.bKeepActorExtraFrame = false;
 	Config.bSpreadFirstVisualizationUpdate = false;
diff --git a/Engine/Plugins/AI/MassCrowd/Source/MassCrowd/Private/MassCrowdVisualizationLODProcessor.cpp b/Engine/Plugins/AI/MassCrowd/Source/MassCrowd/Private/MassCrowdVisualizationLODProcessor.cpp
--- a/Engine/Plugins/AI/MassCrowd/Source/MassCrowd/Private/MassCrowdVisualizationLODProcessor.cpp
+++ b/Engine/Plugins/AI/MassCrowd/Source/MassCrowd/Private/MassCrowdVisualizationLODProcessor.cpp
@@ -108,7 +108,7 @@ void UMassCrowdVisualizationLODProcessor::Execute(UMassEntitySubsystem& EntitySu
 			{
 				const FMassRepresentationFragment& RepresentationFragment = RepresentationFragmentList[EntityIdx];
 				const FMassViewerInfoFragment& LODInfoFragment = LODInfoFragmentList[EntityIdx];
-				if (RepresentationFragment.CurrentRepresentation == ERepresentationType::StaticMeshInstance && SpecifiedRangeSquaredCentimeters > LODInfoFragment.ClosestViewerDistanceSq)
+				if (RepresentationFragment.CurrentRepresentation == EMassRepresentationType::StaticMeshInstance && SpecifiedRangeSquaredCentimeters > LODInfoFragment.ClosestViewerDistanceSq)
 				{
 					const FDataFragment_Transform& EntityLocation = LocationList[EntityIdx];
 					DrawDebugSolidBox(World, EntityLocation.GetTransform().GetLocation() + FVector(0.0f, 0.0f, 150.0f), FVector(50.0f), FColor::Red);
diff --git a/Engine/Plugins/AI/MassCrowd/Source/MassCrowd/Private/MassCrowdVisualizationTrait.cpp b/Engine/Plugins/AI/MassCrowd/Source/MassCrowd/Private/MassCrowdVisualizationTrait.cpp
--- a/Engine/Plugins/AI/MassCrowd/Source/MassCrowd/Private/MassCrowdVisualizationTrait.cpp
+++ b/Engine/Plugins/AI/MassCrowd/Source/MassCrowd/Private/MassCrowdVisualizationTrait.cpp
@@ -10,10 +10,10 @@ UMassCrowdVisualizationTrait::UMassCrowdVisualizationTrait()
 	// Override the subsystem to support parallelization of the crowd
 	RepresentationSubsystemClass = UMassCrowdRepresentationSubsystem::StaticClass();
 	Config.RepresentationActorManagementClass = UMassCrowdRepresentationActorManagement::StaticClass();
-	Config.LODRepresentation[EMassLOD::High] = ERepresentationType::HighResSpawnedActor;
-	Config.LODRepresentation[EMassLOD::Medium] = ERepresentationType::LowResSpawnedActor;
-	Config.LODRepresentation[EMassLOD::Low] = ERepresentationType::StaticMeshInstance;
-	Config.LODRepresentation[EMassLOD::Off] = ERepresentationType::None;
+	Config.LODRepresentation[EMassLOD::High] = EMassRepresentationType::HighResSpawnedActor;
+	Config.LODRepresentation[EMassLOD::Medium] = EMassRepresentationType::LowResSpawnedActor;
+	Config.LODRepresentation[EMassLOD::Low] = EMassRepresentationType::StaticMeshInstance;
+	Config.LODRepresentation[EMassLOD::Off] = EMassRepresentationType::None;
 	// Set bKeepLowResActor to true as a spawning optimization, this will keep the low-res actor if available while showing the static mesh instance
 	Config.bKeepLowResActors = true;
 	Config.bKeepActorExtraFrame = true;
diff --git a/Engine/Plugins/AI/MassCrowd/Source/MassCrowd/Public/MassCrowdRepresentationActorManagement.h b/Engine/Plugins/AI/MassCrowd/Source/MassCrowd/Public/MassCrowdRepresentationActorManagement.h
--- a/Engine/Plugins/AI/MassCrowd/Source/MassCrowd/Public/MassCrowdRepresentationActorManagement.h
+++ b/Engine/Plugins/AI/MassCrowd/Source/MassCrowd/Public/MassCrowdRepresentationActorManagement.h
@@ -24,7 +24,7 @@ protected:
 	 * @param EntityIdx is the entity index currently processing
 	 * @param Context is the current Mass execution context
 	 */
-	virtual void SetActorEnabled(const EActorEnabledType EnabledType, AActor& Actor, const int32 EntityIdx, FMassCommandBuffer& CommandBuffer) const override;
+	virtual void SetActorEnabled(const EMassActorEnabledType EnabledType, AActor& Actor, const int32 EntityIdx, FMassCommandBuffer& CommandBuffer) const override;
 
 	/**
 	 * Returns an actor of the template type and setup fragments values from it
diff --git a/Engine/Plugins/Runtime/MassGameplay/Source/MassActors/Private/MassActorSpawnerSubsystem.cpp b/Engine/Plugins/Runtime/MassGameplay/Source/MassActors/Private/MassActorSpawnerSubsystem.cpp
--- a/Engine/Plugins/Runtime/MassGameplay/Source/MassActors/Private/MassActorSpawnerSubsystem.cpp
+++ b/Engine/Plugins/Runtime/MassGameplay/Source/MassActors/Private/MassActorSpawnerSubsystem.cpp
@@ -175,7 +175,7 @@ FMassActorSpawnRequestHandle UMassActorSpawnerSubsystem::GetNextRequestToSpawn()
 	return BestSpawnRequestHandle;
 }
 
-AActor* UMassActorSpawnerSubsystem::SpawnOrRetrieveFromPool(const FStructView& SpawnRequestView)
+AActor* UMassActorSpawnerSubsystem::SpawnOrRetrieveFromPool(FConstStructView SpawnRequestView)
 {
 	if (UE::MassActors::bUseActorPooling != 0 && bActorPoolingEnabled)
 	{
@@ -204,7 +204,7 @@ AActor* UMassActorSpawnerSubsystem::SpawnOrRetrieveFromPool(const FStructView& S
 	return SpawnActor(SpawnRequestView);
 }
 
-AActor* UMassActorSpawnerSubsystem::SpawnActor(const FStructView& SpawnRequestView) const
+AActor* UMassActorSpawnerSubsystem::SpawnActor(FConstStructView SpawnRequestView) const
 {
 	TRACE_CPUPROFILER_EVENT_SCOPE(UMassActorSpawnerSubsystem::SpawnActor);
 
diff --git a/Engine/Plugins/Runtime/MassGameplay/Source/MassActors/Public/MassActorSpawnerSubsystem.h b/Engine/Plugins/Runtime/MassGameplay/Source/MassActors/Public/MassActorSpawnerSubsystem.h
--- a/Engine/Plugins/Runtime/MassGameplay/Source/MassActors/Public/MassActorSpawnerSubsystem.h
+++ b/Engine/Plugins/Runtime/MassGameplay/Source/MassActors/Public/MassActorSpawnerSubsystem.h
@@ -30,13 +30,13 @@ struct MASSACTORS_API FMassActorSpawnRequestHandle : public FIndexedHandleBase
 typedef FIndexedHandleManager<FMassActorSpawnRequestHandle, true/*bOptimizeHandleReuse*/> FMassEntityHandleManager_ActorSpawnRequest;
 
 
-DECLARE_DELEGATE_TwoParams(FMassActorPreSpawnDelegate, const FMassActorSpawnRequestHandle& SpawnRequestHandle, const FStructView& SpawnRequest);
+DECLARE_DELEGATE_TwoParams(FMassActorPreSpawnDelegate, const FMassActorSpawnRequestHandle& SpawnRequestHandle, FConstStructView SpawnRequest);
 enum class EMassActorSpawnRequestAction : uint8
 {
 	Keep, // Will leave spawning request in the queue and it will be users job to call RemoveActorSpawnRequest
 	Remove, // Will remove the spawning request from the queue once the callback ends
 };
-DECLARE_DELEGATE_RetVal_TwoParams(EMassActorSpawnRequestAction, FMassActorPostSpawnDelegate, const FMassActorSpawnRequestHandle& SpawnRequestHandle, const FStructView& SpawnRequest);
+DECLARE_DELEGATE_RetVal_TwoParams(EMassActorSpawnRequestAction, FMassActorPostSpawnDelegate, const FMassActorSpawnRequestHandle& SpawnRequestHandle, FConstStructView SpawnRequest);
 
 UENUM()
 enum class ESpawnRequestStatus : uint8
@@ -194,11 +194,11 @@ protected:
 	 *  @return the next best handle to spawn. */
 	virtual FMassActorSpawnRequestHandle GetNextRequestToSpawn() const;
 
-	virtual AActor* SpawnOrRetrieveFromPool(const FStructView& SpawnRequest);
+	virtual AActor* SpawnOrRetrieveFromPool(FConstStructView SpawnRequest);
 
 	/** Actual code that will spawn the actor, overridable by subclass if need to be.
 	 *  @return spawned actor if succeeded. */
-	virtual AActor* SpawnActor(const FStructView& SpawnRequest) const;
+	virtual AActor* SpawnActor(FConstStructView SpawnRequest) const;
 
 	/** Go through the spawning request and spawn them until we reach the budget 
 	 * @param MaxTimeSlicePerTick is the budget in seconds allowed to do spawning */
diff --git a/Engine/Plugins/Runtime/MassGameplay/Source/MassRepresentation/Private/MassRepresentationActorManagement.cpp b/Engine/Plugins/Runtime/MassGameplay/Source/MassRepresentation/Private/MassRepresentationActorManagement.cpp
--- a/Engine/Plugins/Runtime/MassGameplay/Source/MassRepresentation/Private/MassRepresentationActorManagement.cpp
+++ b/Engine/Plugins/Runtime/MassGameplay/Source/MassRepresentation/Private/MassRepresentationActorManagement.cpp
@@ -21,9 +21,9 @@ AActor* UMassRepresentationActorManagement::GetOrSpawnActor(UMassRepresentationS
 }
 
 
-void UMassRepresentationActorManagement::SetActorEnabled(const EActorEnabledType EnabledType, AActor& Actor, const int32 EntityIdx, FMassCommandBuffer& CommandBuffer) const
+void UMassRepresentationActorManagement::SetActorEnabled(const EMassActorEnabledType EnabledType, AActor& Actor, const int32 EntityIdx, FMassCommandBuffer& CommandBuffer) const
 {
-	const bool bEnabled = EnabledType != EActorEnabledType::Disabled;
+	const bool bEnabled = EnabledType != EMassActorEnabledType::Disabled;
 	if (Actor.IsActorTickEnabled() != bEnabled)
 	{
 		Actor.SetActorTickEnabled(bEnabled);
@@ -49,7 +49,7 @@ void UMassRepresentationActorManagement::TeleportActor(const FTransform& Transfo
 	}
 }
 
-void UMassRepresentationActorManagement::OnPreActorSpawn(const FMassActorSpawnRequestHandle& SpawnRequestHandle, const FStructView& SpawnRequest, UMassEntitySubsystem* EntitySubsystem) const
+void UMassRepresentationActorManagement::OnPreActorSpawn(const FMassActorSpawnRequestHandle& SpawnRequestHandle, FConstStructView SpawnRequest, UMassEntitySubsystem* EntitySubsystem) const
 {
 	check(EntitySubsystem);
 
@@ -80,7 +80,7 @@ void UMassRepresentationActorManagement::OnPreActorSpawn(const FMassActorSpawnRe
 	}
 }
 
-EMassActorSpawnRequestAction UMassRepresentationActorManagement::OnPostActorSpawn(const FMassActorSpawnRequestHandle& SpawnRequestHandle, const FStructView& SpawnRequest, UMassEntitySubsystem* EntitySubsystem) const
+EMassActorSpawnRequestAction UMassRepresentationActorManagement::OnPostActorSpawn(const FMassActorSpawnRequestHandle& SpawnRequestHandle, FConstStructView SpawnRequest, UMassEntitySubsystem* EntitySubsystem) const
 {
 	check(EntitySubsystem);
 
diff --git a/Engine/Plugins/Runtime/MassGameplay/Source/MassRepresentation/Private/MassRepresentationModule.cpp b/Engine/Plugins/Runtime/MassGameplay/Source/MassRepresentation/Private/MassRepresentationModule.cpp
--- a/Engine/Plugins/Runtime/MassGameplay/Source/MassRepresentation/Private/MassRepresentationModule.cpp
+++ b/Engine/Plugins/Runtime/MassGameplay/Source/MassRepresentation/Private/MassRepresentationModule.cpp
@@ -26,6 +26,11 @@ void FMassRepresentationModule::StartupModule()
 	Redirects.Emplace(ECoreRedirectFlags::Type_Class, TEXT("/Script/MassRepresentation.MassRepresentationLODProcessor"), TEXT("/Script/MassRepresentation.MassVisualizationLODProcessor"));
 	Redirects.Emplace(ECoreRedirectFlags::Type_Class, TEXT("/Script/MassVisualization.MassVisualizationFeature"), TEXT("/Script/MassRepresentation.MassVisualizationTrait"));
 	Redirects.Emplace(ECoreRedirectFlags::Type_Property, TEXT("/Script/MassRepresentation.MassVisualizationTrait.TemplateActor"), TEXT("/Script/MassRepresentation.MassVisualizationTrait.HighResTemplateActor"));
+	Redirects.Emplace(ECoreRedirectFlags::Type_Enum, TEXT("/Script/MassRepresentation.RepresentationType"), TEXT("/Script/MassRepresentation.MassRepresentationType"));
+	Redirects.Emplace(ECoreRedirectFlags::Type_Enum, TEXT("RepresentationType"),                            TEXT("/Script/MassRepresentation.MassRepresentationType"));
+	Redirects.Emplace(ECoreRedirectFlags::Type_Enum, TEXT("/Script/MassRepresentation.ActorEnabledType"), TEXT("/Script/MassRepresentation.MassActorEnabledType"));
+	Redirects.Emplace(ECoreRedirectFlags::Type_Enum, TEXT("ActorEnabledType"),                            TEXT("/Script/MassRepresentation.MassActorEnabledType"));
+
 
 	FCoreRedirects::AddRedirectList(Redirects, TEXT("MassRepresentation"));
 }
diff --git a/Engine/Plugins/Runtime/MassGameplay/Source/MassRepresentation/Private/MassRepresentationProcessor.cpp b/Engine/Plugins/Runtime/MassGameplay/Source/MassRepresentation/Private/MassRepresentationProcessor.cpp
--- a/Engine/Plugins/Runtime/MassGameplay/Source/MassRepresentation/Private/MassRepresentationProcessor.cpp
+++ b/Engine/Plugins/Runtime/MassGameplay/Source/MassRepresentation/Private/MassRepresentationProcessor.cpp
@@ -70,13 +70,13 @@ void UMassRepresentationProcessor::UpdateRepresentation(FMassExecutionContext& C
 		FDataFragment_Actor& ActorInfo = ActorList[EntityIdx];
 
 		// Keeping a copy of the that last calculated previous representation
-		const ERepresentationType PrevRepresentationCopy = Representation.PrevRepresentation;
+		const EMassRepresentationType PrevRepresentationCopy = Representation.PrevRepresentation;
 		Representation.PrevRepresentation = Representation.CurrentRepresentation;
 		
-		ERepresentationType WantedRepresentationType = RepresentationConfig.LODRepresentation[FMath::Min((int32)RepresentationLOD.LOD, (int32)EMassLOD::Off)];
+		EMassRepresentationType WantedRepresentationType = RepresentationConfig.LODRepresentation[FMath::Min((int32)RepresentationLOD.LOD, (int32)EMassLOD::Off)];
 
 		// Make sure we do not have actor spawned in areas not fully loaded
-		if ((WantedRepresentationType == ERepresentationType::HighResSpawnedActor || WantedRepresentationType == ERepresentationType::LowResSpawnedActor) &&
+		if ((WantedRepresentationType == EMassRepresentationType::HighResSpawnedActor || WantedRepresentationType == EMassRepresentationType::LowResSpawnedActor) &&
 			!RepresentationSubsystem->IsCollisionLoaded(RepresentationConfig.WorldPartitionGridNameContainingCollision, TransformFragment.GetTransform()))
 		{
 			WantedRepresentationType = RepresentationConfig.CachedDefaultRepresentationType;
@@ -109,14 +109,14 @@ void UMassRepresentationProcessor::UpdateRepresentation(FMassExecutionContext& C
 			}
 			if (Actor != nullptr)
 			{
-				RepresentationActorManagement->SetActorEnabled(EActorEnabledType::Disabled, *Actor, EntityIdx, Context.Defer());
+				RepresentationActorManagement->SetActorEnabled(EMassActorEnabledType::Disabled, *Actor, EntityIdx, Context.Defer());
 			}
 		};
 
 		// Process switch between representation if there is a change in the representation or there is a pending spawning request
 		if (WantedRepresentationType != Representation.CurrentRepresentation || Representation.ActorSpawnRequestHandle.IsValid())
 		{
-			if (Representation.CurrentRepresentation == ERepresentationType::None)
+			if (Representation.CurrentRepresentation == EMassRepresentationType::None)
 			{
 				Representation.PrevTransform = TransformFragment.GetTransform();
 				Representation.PrevLODSignificance = RepresentationLOD.LODSignificance;
@@ -125,10 +125,10 @@ void UMassRepresentationProcessor::UpdateRepresentation(FMassExecutionContext& C
 			AActor* Actor = ActorInfo.GetMutable();
 			switch (WantedRepresentationType)
 			{
-				case ERepresentationType::HighResSpawnedActor:
-				case ERepresentationType::LowResSpawnedActor:
+				case EMassRepresentationType::HighResSpawnedActor:
+				case EMassRepresentationType::LowResSpawnedActor:
 				{
-					const bool bHighResActor = WantedRepresentationType == ERepresentationType::HighResSpawnedActor;
+					const bool bHighResActor = WantedRepresentationType == EMassRepresentationType::HighResSpawnedActor;
 
 					// Reuse actor, if it is valid and not owned by mass or same representation as low res without a valid spawning request
 					AActor* NewActor = nullptr;
@@ -167,12 +167,12 @@ void UMassRepresentationProcessor::UpdateRepresentation(FMassExecutionContext& C
 					{
 						// Make sure our (re)activated actor is at the simulated position
 						// Needs to be done before enabling the actor so the animation initialization can use the new values
-						if (Representation.CurrentRepresentation == ERepresentationType::StaticMeshInstance)
+						if (Representation.CurrentRepresentation == EMassRepresentationType::StaticMeshInstance)
 						{
 							RepresentationActorManagement->TeleportActor(Representation.PrevTransform, *NewActor, Context.Defer());
 						}
 
-						RepresentationActorManagement->SetActorEnabled(bHighResActor ? EActorEnabledType::HighRes : EActorEnabledType::LowRes, *NewActor, EntityIdx, Context.Defer());
+						RepresentationActorManagement->SetActorEnabled(bHighResActor ? EMassActorEnabledType::HighRes : EMassActorEnabledType::LowRes, *NewActor, EntityIdx, Context.Defer());
 						Representation.CurrentRepresentation = WantedRepresentationType;
 					}
 					else if (!Actor)
@@ -181,16 +181,16 @@ void UMassRepresentationProcessor::UpdateRepresentation(FMassExecutionContext& C
 					}
 					break;
 				}
-				case ERepresentationType::StaticMeshInstance:
+				case EMassRepresentationType::StaticMeshInstance:
 					if (!bDoKeepActorExtraFrame || 
-					   (Representation.PrevRepresentation != ERepresentationType::HighResSpawnedActor && Representation.PrevRepresentation != ERepresentationType::LowResSpawnedActor))
+					   (Representation.PrevRepresentation != EMassRepresentationType::HighResSpawnedActor && Representation.PrevRepresentation != EMassRepresentationType::LowResSpawnedActor))
 					{
 						DisableActorForISM(Actor);
 					}
  
-					Representation.CurrentRepresentation = ERepresentationType::StaticMeshInstance;
+					Representation.CurrentRepresentation = EMassRepresentationType::StaticMeshInstance;
 					break;
-				case ERepresentationType::None:
+				case EMassRepresentationType::None:
 					if (!Actor || ActorInfo.IsOwnedByMass())
 					{
 						// Try releasing both, could have an high res spawned actor and a spawning request for a low res one
@@ -199,9 +199,9 @@ void UMassRepresentationProcessor::UpdateRepresentation(FMassExecutionContext& C
 					}
 					else
 					{
-						RepresentationActorManagement->SetActorEnabled(EActorEnabledType::Disabled, *Actor, EntityIdx, Context.Defer());
+						RepresentationActorManagement->SetActorEnabled(EMassActorEnabledType::Disabled, *Actor, EntityIdx, Context.Defer());
 					}
-					Representation.CurrentRepresentation = ERepresentationType::None;
+					Representation.CurrentRepresentation = EMassRepresentationType::None;
 					break;
 				default:
 					checkf(false, TEXT("Unsupported LOD type"));
@@ -209,8 +209,8 @@ void UMassRepresentationProcessor::UpdateRepresentation(FMassExecutionContext& C
 			}
 		}
 		else if (bDoKeepActorExtraFrame && 
-				 Representation.PrevRepresentation == ERepresentationType::StaticMeshInstance &&
-			    (PrevRepresentationCopy == ERepresentationType::HighResSpawnedActor || PrevRepresentationCopy == ERepresentationType::LowResSpawnedActor))
+				 Representation.PrevRepresentation == EMassRepresentationType::StaticMeshInstance &&
+			    (PrevRepresentationCopy == EMassRepresentationType::HighResSpawnedActor || PrevRepresentationCopy == EMassRepresentationType::LowResSpawnedActor))
 		{
 			AActor* Actor = ActorInfo.GetMutable();
 			DisableActorForISM(Actor);
@@ -341,7 +341,7 @@ FMassVisualizationChunkFragment& UMassVisualizationProcessor::UpdateChunkVisibil
 void UMassVisualizationProcessor::UpdateEntityVisibility(const FMassEntityHandle Entity, const FMassRepresentationFragment& Representation, const FMassRepresentationLODFragment& RepresentationLOD, FMassVisualizationChunkFragment& ChunkData, FMassCommandBuffer& CommandBuffer)
 {
 	// Move the visible entities together into same chunks so we can skip entire chunk when not visible as an optimization
-	const EMassVisibility Visibility = Representation.CurrentRepresentation != ERepresentationType::None ? 
+	const EMassVisibility Visibility = Representation.CurrentRepresentation != EMassRepresentationType::None ? 
 		EMassVisibility::CanBeSeen : RepresentationLOD.Visibility;
 	const EMassVisibility ChunkVisibility = ChunkData.GetVisibility();
 	if (ChunkVisibility != Visibility)
diff --git a/Engine/Plugins/Runtime/MassGameplay/Source/MassRepresentation/Private/MassRepresentationSubsystem.cpp b/Engine/Plugins/Runtime/MassGameplay/Source/MassRepresentation/Private/MassRepresentationSubsystem.cpp
--- a/Engine/Plugins/Runtime/MassGameplay/Source/MassRepresentation/Private/MassRepresentationSubsystem.cpp
+++ b/Engine/Plugins/Runtime/MassGameplay/Source/MassRepresentation/Private/MassRepresentationSubsystem.cpp
@@ -399,7 +399,7 @@ void UMassRepresentationSubsystem::OnMassAgentComponentEntityDetaching(const UMa
 		if (FMassRepresentationFragment* Representation = EntityView.GetFragmentDataPtr<FMassRepresentationFragment>())
 		{
 			// Force a reevaluate of the current representation
-			Representation->CurrentRepresentation = ERepresentationType::None;
+			Representation->CurrentRepresentation = EMassRepresentationType::None;
 		}
 	}
 }
diff --git a/Engine/Plugins/Runtime/MassGameplay/Source/MassRepresentation/Private/MassUpdateISMProcessor.cpp b/Engine/Plugins/Runtime/MassGameplay/Source/MassRepresentation/Private/MassUpdateISMProcessor.cpp
--- a/Engine/Plugins/Runtime/MassGameplay/Source/MassRepresentation/Private/MassUpdateISMProcessor.cpp
+++ b/Engine/Plugins/Runtime/MassGameplay/Source/MassRepresentation/Private/MassUpdateISMProcessor.cpp
@@ -45,7 +45,7 @@ void UMassUpdateISMProcessor::Execute(UMassEntitySubsystem& EntitySubsystem, FMa
 			const FMassRepresentationLODFragment& RepresentationLOD = RepresentationLODList[EntityIdx];
 			FMassRepresentationFragment& Representation = RepresentationList[EntityIdx];
 
-			if (Representation.CurrentRepresentation == ERepresentationType::StaticMeshInstance)
+			if (Representation.CurrentRepresentation == EMassRepresentationType::StaticMeshInstance)
 			{
 				UpdateISMTransform(GetTypeHash(Context.GetEntity(EntityIdx)), ISMInfo[Representation.StaticMeshDescIndex], TransformFragment.GetTransform(), Representation.PrevTransform, RepresentationLOD.LODSignificance, Representation.PrevLODSignificance);
 			}
diff --git a/Engine/Plugins/Runtime/MassGameplay/Source/MassRepresentation/Private/MassVisualizationTrait.cpp b/Engine/Plugins/Runtime/MassGameplay/Source/MassRepresentation/Private/MassVisualizationTrait.cpp
--- a/Engine/Plugins/Runtime/MassGameplay/Source/MassRepresentation/Private/MassVisualizationTrait.cpp
+++ b/Engine/Plugins/Runtime/MassGameplay/Source/MassRepresentation/Private/MassVisualizationTrait.cpp
@@ -14,10 +14,10 @@ UMassVisualizationTrait::UMassVisualizationTrait()
 	RepresentationSubsystemClass = UMassRepresentationSubsystem::StaticClass();
 
 	Config.RepresentationActorManagementClass = UMassRepresentationActorManagement::StaticClass();
-	Config.LODRepresentation[EMassLOD::High] = ERepresentationType::HighResSpawnedActor;
-	Config.LODRepresentation[EMassLOD::Medium] = ERepresentationType::LowResSpawnedActor;
-	Config.LODRepresentation[EMassLOD::Low] = ERepresentationType::StaticMeshInstance;
-	Config.LODRepresentation[EMassLOD::Off] = ERepresentationType::None;
+	Config.LODRepresentation[EMassLOD::High] = EMassRepresentationType::HighResSpawnedActor;
+	Config.LODRepresentation[EMassLOD::Medium] = EMassRepresentationType::LowResSpawnedActor;
+	Config.LODRepresentation[EMassLOD::Low] = EMassRepresentationType::StaticMeshInstance;
+	Config.LODRepresentation[EMassLOD::Off] = EMassRepresentationType::None;
 }
 
 void UMassVisualizationTrait::BuildTemplate(FMassEntityTemplateBuildContext& BuildContext, UWorld& World) const
diff --git a/Engine/Plugins/Runtime/MassGameplay/Source/MassRepresentation/Public/MassRepresentationActorManagement.h b/Engine/Plugins/Runtime/MassGameplay/Source/MassRepresentation/Public/MassRepresentationActorManagement.h
--- a/Engine/Plugins/Runtime/MassGameplay/Source/MassRepresentation/Public/MassRepresentationActorManagement.h
+++ b/Engine/Plugins/Runtime/MassGameplay/Source/MassRepresentation/Public/MassRepresentationActorManagement.h
@@ -8,9 +8,9 @@ struct FMassActorSpawnRequestHandle;
 struct FMassCommandBuffer;
 class UMassEntitySubsystem;
 enum class EMassActorSpawnRequestAction : uint8;
-enum class EActorEnabledType : uint8;
+enum class EMassActorEnabledType : uint8;
 struct FMassActorSpawnRequestHandle;
-struct FStructView;
+struct FConstStructView;
 struct FMassEntityHandle;
 struct FDataFragment_Actor;
 struct FMassRepresentationLODFragment;
@@ -51,7 +51,7 @@ public:
 	 * @param EntityIdx is the entity index currently processing
 	 * @param CommandBuffer to queue up anything that is thread sensitive
 	 */
-	virtual void SetActorEnabled(const EActorEnabledType EnabledType, AActor& Actor, const int32 EntityIdx, FMassCommandBuffer& CommandBuffer) const;
+	virtual void SetActorEnabled(const EMassActorEnabledType EnabledType, AActor& Actor, const int32 EntityIdx, FMassCommandBuffer& CommandBuffer) const;
 
 	/**
 	 * Teleports the actor at the specified transform by preserving its velocity and without collision.
@@ -69,7 +69,7 @@ public:
 	 * @param SpawnRequest of the actor that is about to spawn
 	 * @param EntitySubsystem to use to retrieve the mass agent fragments
 	 */
-	virtual void OnPreActorSpawn(const FMassActorSpawnRequestHandle& SpawnRequestHandle, const FStructView& SpawnRequest, UMassEntitySubsystem* EntitySubsystem) const;
+	virtual void OnPreActorSpawn(const FMassActorSpawnRequestHandle& SpawnRequestHandle, FConstStructView SpawnRequest, UMassEntitySubsystem* EntitySubsystem) const;
 
 	/**
 	 * Method that will be bound to a delegate used post-spawn to notify and let the requester configure the actor
@@ -78,7 +78,7 @@ public:
 	 * @param EntitySubsystem to use to retrieve the mass agent fragments
 	 * @return The action to take on the spawn request, either keep it there or remove it.
 	 */
-	virtual EMassActorSpawnRequestAction OnPostActorSpawn(const FMassActorSpawnRequestHandle& SpawnRequestHandle, const FStructView& SpawnRequest, UMassEntitySubsystem* EntitySubsystem) const;
+	virtual EMassActorSpawnRequestAction OnPostActorSpawn(const FMassActorSpawnRequestHandle& SpawnRequestHandle, FConstStructView SpawnRequest, UMassEntitySubsystem* EntitySubsystem) const;
 
 	/**
 	 * Static methods to Release an actor or cancel its spawning (calls ReleaseAnyActorOrCancelAnySpawning)
diff --git a/Engine/Plugins/Runtime/MassGameplay/Source/MassRepresentation/Public/MassRepresentationFragments.h b/Engine/Plugins/Runtime/MassGameplay/Source/MassRepresentation/Public/MassRepresentationFragments.h
--- a/Engine/Plugins/Runtime/MassGameplay/Source/MassRepresentation/Public/MassRepresentationFragments.h
+++ b/Engine/Plugins/Runtime/MassGameplay/Source/MassRepresentation/Public/MassRepresentationFragments.h
@@ -36,9 +36,9 @@ struct MASSREPRESENTATION_API FMassRepresentationFragment : public FMassFragment
 {
 	GENERATED_BODY()
 
-	ERepresentationType CurrentRepresentation = ERepresentationType::None;
+	EMassRepresentationType CurrentRepresentation = EMassRepresentationType::None;
 
-	ERepresentationType PrevRepresentation = ERepresentationType::None;
+	EMassRepresentationType PrevRepresentation = EMassRepresentationType::None;
 
 	int16 HighResTemplateActorIndex = INDEX_NONE;
 
@@ -75,7 +75,7 @@ struct FMassRepresentationConfig : public FMassSharedFragment
 
 	/** What should be the representation of this entity for each specific LOD */
 	UPROPERTY(EditAnywhere, Category = "Mass|Representation", config)
-	ERepresentationType LODRepresentation[EMassLOD::Max];
+	EMassRepresentationType LODRepresentation[EMassLOD::Max];
 
 	/** If true, LowRes actors will be kept around, disabled, whilst StaticMeshInstance representation is active */
 	UPROPERTY(EditAnywhere, Category = "Mass|Representation", config)
@@ -101,7 +101,7 @@ struct FMassRepresentationConfig : public FMassSharedFragment
 
 	/** Default representation when unable to spawn an actor, gets calculated at initialization */
 	UPROPERTY(Transient)
-	mutable ERepresentationType CachedDefaultRepresentationType = ERepresentationType::None;
+	mutable EMassRepresentationType CachedDefaultRepresentationType = EMassRepresentationType::None;
 
 	UPROPERTY(Transient)
 	mutable UMassRepresentationActorManagement* CachedRepresentationActorManagement = nullptr;
@@ -113,8 +113,8 @@ inline void FMassRepresentationConfig::ComputeCachedValues() const
 	for (int32 LOD = EMassLOD::High; LOD < EMassLOD::Max; LOD++)
 	{
 		// Find the first representation type after any actors
-		if (LODRepresentation[LOD] == ERepresentationType::HighResSpawnedActor ||
-			LODRepresentation[LOD] == ERepresentationType::LowResSpawnedActor)
+		if (LODRepresentation[LOD] == EMassRepresentationType::HighResSpawnedActor ||
+			LODRepresentation[LOD] == EMassRepresentationType::LowResSpawnedActor)
 		{
 			continue;
 		}
diff --git a/Engine/Plugins/Runtime/MassGameplay/Source/MassRepresentation/Public/MassRepresentationTypes.h b/Engine/Plugins/Runtime/MassGameplay/Source/MassRepresentation/Public/MassRepresentationTypes.h
--- a/Engine/Plugins/Runtime/MassGameplay/Source/MassRepresentation/Public/MassRepresentationTypes.h
+++ b/Engine/Plugins/Runtime/MassGameplay/Source/MassRepresentation/Public/MassRepresentationTypes.h
@@ -19,7 +19,7 @@ namespace UE::Mass::ProcessorGroupNames
 }
 
 UENUM()
-enum class ERepresentationType : uint8
+enum class EMassRepresentationType : uint8
 {
 	HighResSpawnedActor,
 	LowResSpawnedActor,
@@ -27,7 +27,7 @@ enum class ERepresentationType : uint8
 	None,
 };
 
-enum class EActorEnabledType : uint8
+enum class EMassActorEnabledType : uint8
 {
 	Disabled,
 	LowRes,
