Fix CIS for uninitialized variable

#rnx

diff --git a/Engine/Plugins/Runtime/MassGameplay/Source/MassRepresentation/Public/MassRepresentationFragments.h b/Engine/Plugins/Runtime/MassGameplay/Source/MassRepresentation/Public/MassRepresentationFragments.h
--- a/Engine/Plugins/Runtime/MassGameplay/Source/MassRepresentation/Public/MassRepresentationFragments.h
+++ b/Engine/Plugins/Runtime/MassGameplay/Source/MassRepresentation/Public/MassRepresentationFragments.h
@@ -60,7 +60,7 @@ struct FMassRepresentationSubsystemFragment : public FMassSharedFragment
 	GENERATED_BODY()
 
 	UPROPERTY(Transient)
-	UMassRepresentationSubsystem* RepresentationSubsystem;
+	UMassRepresentationSubsystem* RepresentationSubsystem = nullptr;
 };
 
 USTRUCT()
