Moved replication fragments into its own file to match the rest of the other systems

#rnx

diff --git a/Engine/Plugins/Runtime/MassGameplay/Source/MassActors/Private/MassAgentComponent.cpp b/Engine/Plugins/Runtime/MassGameplay/Source/MassActors/Private/MassAgentComponent.cpp
--- a/Engine/Plugins/Runtime/MassGameplay/Source/MassActors/Private/MassAgentComponent.cpp
+++ b/Engine/Plugins/Runtime/MassGameplay/Source/MassActors/Private/MassAgentComponent.cpp
@@ -11,6 +11,7 @@
 #include "MassSpawner.h"
 #include "Net/UnrealNetwork.h"
 #include "MassReplicationManager.h"
+#include "MassReplicationFragments.h"
 #include "MassMovementFragments.h"
 #include "GameFramework/CharacterMovementComponent.h"
 
diff --git a/Engine/Plugins/Runtime/MassGameplay/Source/MassReplication/Private/MassReplicationFragments.cpp b/Engine/Plugins/Runtime/MassGameplay/Source/MassReplication/Private/MassReplicationFragments.cpp
new file mode 100644
--- /dev/null
+++ b/Engine/Plugins/Runtime/MassGameplay/Source/MassReplication/Private/MassReplicationFragments.cpp
@@ -0,0 +1,53 @@
+// Copyright Epic Games, Inc. All Rights Reserved.
+
+#include "MassReplicationFragments.h"
+#include "MassReplicationManager.h"
+#include "Engine/World.h"
+
+//----------------------------------------------------------------------//
+// UMassNetworkIDFragmentInitializer 
+//----------------------------------------------------------------------//
+UMassNetworkIDFragmentInitializer::UMassNetworkIDFragmentInitializer()
+{
+	ExecutionFlags = (int32)(EProcessorExecutionFlags::Standalone | EProcessorExecutionFlags::Server);
+	FragmentType = FMassNetworkIDFragment::StaticStruct();
+}
+
+void UMassNetworkIDFragmentInitializer::ConfigureQueries()
+{
+	EntityQuery.AddRequirement<FMassNetworkIDFragment>(EMassFragmentAccess::ReadWrite);
+}
+
+void UMassNetworkIDFragmentInitializer::Initialize(UObject& Owner)
+{
+	Super::Initialize(Owner);
+#if UE_REPLICATION_COMPILE_SERVER_CODE
+	ReplicationManager = UWorld::GetSubsystem<UMassReplicationManager>(Owner.GetWorld());
+#endif //UE_REPLICATION_COMPILE_SERVER_CODE
+}
+
+void UMassNetworkIDFragmentInitializer::Execute(UMassEntitySubsystem& EntitySubsystem, FMassExecutionContext& Context)
+{
+	QUICK_SCOPE_CYCLE_COUNTER(MassProcessor_InitNetworkID_Run);
+
+	const UWorld* World = EntitySubsystem.GetWorld();
+	const ENetMode NetMode = World->GetNetMode();
+
+	if (NetMode != NM_Client)
+	{
+#if UE_REPLICATION_COMPILE_SERVER_CODE
+		check(ReplicationManager);
+
+		EntityQuery.ForEachEntityChunk(EntitySubsystem, Context, [this](FMassExecutionContext& Context)
+			{
+				const TArrayView<FMassNetworkIDFragment> NetworkIDList = Context.GetMutableFragmentView<FMassNetworkIDFragment>();
+				const int32 NumEntities = Context.GetNumEntities();
+
+				for (int32 Idx = 0; Idx < NumEntities; ++Idx)
+				{
+					NetworkIDList[Idx].NetID = ReplicationManager->GetNextAvailableMassNetID();
+				}
+			});
+#endif //UE_REPLICATION_COMPILE_SERVER_CODE
+	}
+}
diff --git a/Engine/Plugins/Runtime/MassGameplay/Source/MassReplication/Private/MassReplicationTrait.cpp b/Engine/Plugins/Runtime/MassGameplay/Source/MassReplication/Private/MassReplicationTrait.cpp
--- a/Engine/Plugins/Runtime/MassGameplay/Source/MassReplication/Private/MassReplicationTrait.cpp
+++ b/Engine/Plugins/Runtime/MassGameplay/Source/MassReplication/Private/MassReplicationTrait.cpp
@@ -6,6 +6,7 @@
 #include "MassSpawnerTypes.h"
 #include "MassSimulationLOD.h"
 #include "MassReplicationTypes.h"
+#include "MassReplicationFragments.h"
 
 
 void UMassReplicationTrait::BuildTemplate(FMassEntityTemplateBuildContext& BuildContext, UWorld& World) const
diff --git a/Engine/Plugins/Runtime/MassGameplay/Source/MassReplication/Private/MassReplicationTypes.cpp b/Engine/Plugins/Runtime/MassGameplay/Source/MassReplication/Private/MassReplicationTypes.cpp
--- a/Engine/Plugins/Runtime/MassGameplay/Source/MassReplication/Private/MassReplicationTypes.cpp
+++ b/Engine/Plugins/Runtime/MassGameplay/Source/MassReplication/Private/MassReplicationTypes.cpp
@@ -1,56 +1,5 @@
 // Copyright Epic Games, Inc. All Rights Reserved.
 
 #include "MassReplicationTypes.h"
-#include "MassReplicationManager.h"
-#include "MassSpawnerTypes.h"
-#include "Engine/World.h"
 
-DEFINE_LOG_CATEGORY(LogMassReplication);
-
-//----------------------------------------------------------------------//
-// UMassNetworkIDFragmentInitializer 
-//----------------------------------------------------------------------//
-UMassNetworkIDFragmentInitializer::UMassNetworkIDFragmentInitializer()
-{
-	ExecutionFlags = (int32)(EProcessorExecutionFlags::Standalone | EProcessorExecutionFlags::Server);
-	FragmentType = FMassNetworkIDFragment::StaticStruct();
-}
-
-void UMassNetworkIDFragmentInitializer::ConfigureQueries()
-{
-	EntityQuery.AddRequirement<FMassNetworkIDFragment>(EMassFragmentAccess::ReadWrite);
-}
-
-void UMassNetworkIDFragmentInitializer::Initialize(UObject& Owner)
-{
-	Super::Initialize(Owner);
-#if UE_REPLICATION_COMPILE_SERVER_CODE
-	ReplicationManager = UWorld::GetSubsystem<UMassReplicationManager>(Owner.GetWorld());
-#endif //UE_REPLICATION_COMPILE_SERVER_CODE
-}
-
-void UMassNetworkIDFragmentInitializer::Execute(UMassEntitySubsystem& EntitySubsystem, FMassExecutionContext& Context)
-{
-	QUICK_SCOPE_CYCLE_COUNTER(MassProcessor_InitNetworkID_Run);
-
-	const UWorld* World = EntitySubsystem.GetWorld();
-	const ENetMode NetMode = World->GetNetMode();
-
-	if (NetMode != NM_Client)
-	{
-#if UE_REPLICATION_COMPILE_SERVER_CODE
-		check(ReplicationManager);
-
-		EntityQuery.ForEachEntityChunk(EntitySubsystem, Context, [this](FMassExecutionContext& Context)
-			{
-				const TArrayView<FMassNetworkIDFragment> NetworkIDList = Context.GetMutableFragmentView<FMassNetworkIDFragment>();
-				const int32 NumEntities = Context.GetNumEntities();
-
-				for (int32 Idx = 0; Idx < NumEntities; ++Idx)
-				{
-					NetworkIDList[Idx].NetID = ReplicationManager->GetNextAvailableMassNetID();
-				}
-			});
-#endif //UE_REPLICATION_COMPILE_SERVER_CODE
-	}
-}
+DEFINE_LOG_CATEGORY(LogMassReplication);
\ No newline at end of file
diff --git a/Engine/Plugins/Runtime/MassGameplay/Source/MassReplication/Public/MassClientBubbleHandler.h b/Engine/Plugins/Runtime/MassGameplay/Source/MassReplication/Public/MassClientBubbleHandler.h
--- a/Engine/Plugins/Runtime/MassGameplay/Source/MassReplication/Public/MassClientBubbleHandler.h
+++ b/Engine/Plugins/Runtime/MassGameplay/Source/MassReplication/Public/MassClientBubbleHandler.h
@@ -9,6 +9,7 @@
 #include "MassClientBubbleSerializerBase.h"
 #include "MassSpawnerSubsystem.h"
 #include "MassSpawnerTypes.h"
+#include "MassReplicationFragments.h"
 #include "MassReplicationManager.h"
 #include "Engine/World.h"
 #include "MassEntityTemplate.h"
diff --git a/Engine/Plugins/Runtime/MassGameplay/Source/MassReplication/Public/MassReplicationFragments.h b/Engine/Plugins/Runtime/MassGameplay/Source/MassReplication/Public/MassReplicationFragments.h
new file mode 100644
--- /dev/null
+++ b/Engine/Plugins/Runtime/MassGameplay/Source/MassReplication/Public/MassReplicationFragments.h
@@ -0,0 +1,106 @@
+// Copyright Epic Games, Inc. All Rights Reserved.
+
+#pragma once
+
+#include "MassCommonTypes.h"
+#include "MassReplicationTypes.h"
+// #include "MassLODManager.h"
+// #include "IndexedHandle.h"
+// #include "MassEntityTemplate.h"
+// #include "MassObserverProcessor.h"
+// #include "MassLODLogic.h"
+
+#include "MassReplicationFragments.generated.h"
+
+/**
+ *  Fragment type for the mass network id of a mass entity
+ */
+USTRUCT()
+struct MASSREPLICATION_API FMassNetworkIDFragment : public FMassFragment
+{
+	GENERATED_BODY()
+
+	FMassNetworkID NetID;
+};
+
+struct FMassReplicatedAgentArrayData
+{
+	void Invalidate()
+	{
+		Handle.Invalidate();
+		LastUpdateTime = 0.f;
+#if UE_ALLOW_DEBUG_REPLICATION
+		LOD = EMassLOD::Off;
+#endif // UE_DO_DEBUG_REPLICATION
+	}
+
+	FMassReplicatedAgentHandle Handle;
+	float LastUpdateTime = 0.f;
+
+#if UE_ALLOW_DEBUG_REPLICATION
+	EMassLOD::Type LOD = EMassLOD::Off;
+#endif // UE_DO_DEBUG_REPLICATION
+};
+
+/**
+ * Agent handle per client, these will be at TArray indices of the Client handles indicies (used as a free list array)
+ */
+USTRUCT()
+struct MASSREPLICATION_API FMassReplicatedAgentFragment : public FMassFragment
+{
+	GENERATED_BODY()
+
+	//these will be maintained so that the TArray entries are at the same indices as Client Handles
+	TArray<FMassReplicatedAgentArrayData, TInlineAllocator<UE::Mass::Replication::MaxNumOfClients>> AgentsData;
+};
+
+/*
+ * Data fragment to store the calculated distances to viewers
+ */
+USTRUCT()
+struct MASSREPLICATION_API FMassReplicationViewerInfoFragment : public FMassFragment
+{
+	GENERATED_BODY()
+
+	/** Closest viewer distance */
+	float ClosestViewerDistanceSq;
+
+	/** Distance between each viewer and entity */
+	TArray<float> DistanceToViewerSq;
+};
+
+USTRUCT()
+struct MASSREPLICATION_API FMassReplicationLODFragment : public FMassFragment
+{
+	GENERATED_BODY()
+
+	/**LOD information */
+	TEnumAsByte<EMassLOD::Type> LOD = EMassLOD::Max;
+
+	/** Previous LOD information*/
+	TEnumAsByte<EMassLOD::Type> PrevLOD = EMassLOD::Max;
+
+	/** Per viewer LOD information */
+	TArray<TEnumAsByte<EMassLOD::Type>> LODPerViewer;
+
+	/** Per viewer previous LOD information */
+	TArray<TEnumAsByte<EMassLOD::Type>> PrevLODPerViewer;
+};
+
+UCLASS()
+class MASSREPLICATION_API UMassNetworkIDFragmentInitializer : public UMassFragmentInitializer
+{
+	GENERATED_BODY()
+public:
+	UMassNetworkIDFragmentInitializer();
+
+protected:
+	virtual void ConfigureQueries() override;
+	virtual void Initialize(UObject& Owner) override;
+	virtual void Execute(UMassEntitySubsystem& EntitySubsystem, FMassExecutionContext& Context) override;
+
+	UPROPERTY()
+		UMassReplicationManager* ReplicationManager = nullptr;
+
+	FMassEntityQuery EntityQuery;
+};
diff --git a/Engine/Plugins/Runtime/MassGameplay/Source/MassReplication/Public/MassReplicationProcessorBase.h b/Engine/Plugins/Runtime/MassGameplay/Source/MassReplication/Public/MassReplicationProcessorBase.h
--- a/Engine/Plugins/Runtime/MassGameplay/Source/MassReplication/Public/MassReplicationProcessorBase.h
+++ b/Engine/Plugins/Runtime/MassGameplay/Source/MassReplication/Public/MassReplicationProcessorBase.h
@@ -8,6 +8,7 @@
 #include "MassLODCalculator.h"
 #include "MassSimulationLOD.h"
 #include "MassReplicationTypes.h"
+#include "MassReplicationFragments.h"
 #include "MassSpawnerTypes.h"
 #include "MassLODManager.h"
 #include "MassReplicationManager.h"
diff --git a/Engine/Plugins/Runtime/MassGameplay/Source/MassReplication/Public/MassReplicationTypes.h b/Engine/Plugins/Runtime/MassGameplay/Source/MassReplication/Public/MassReplicationTypes.h
--- a/Engine/Plugins/Runtime/MassGameplay/Source/MassReplication/Public/MassReplicationTypes.h
+++ b/Engine/Plugins/Runtime/MassGameplay/Source/MassReplication/Public/MassReplicationTypes.h
@@ -117,48 +117,6 @@ struct MASSREPLICATION_API FMassBubbleInfoClassHandle : public FSimpleIndexedHan
 	};
 };
 
-/**
- *  Fragment type for the mass network id of a mass entity
- */
-USTRUCT()
-struct MASSREPLICATION_API FMassNetworkIDFragment : public FMassFragment
-{
-	GENERATED_BODY()
-
-	FMassNetworkID NetID;
-};
-
-struct FMassReplicatedAgentArrayData
-{
-	void Invalidate()
-	{
-		Handle.Invalidate();
-		LastUpdateTime = 0.f;
-#if UE_ALLOW_DEBUG_REPLICATION
-		LOD = EMassLOD::Off;
-#endif // UE_DO_DEBUG_REPLICATION
-	}
-
-	FMassReplicatedAgentHandle Handle;
-	float LastUpdateTime = 0.f;
-
-#if UE_ALLOW_DEBUG_REPLICATION
-	EMassLOD::Type LOD = EMassLOD::Off;
-#endif // UE_DO_DEBUG_REPLICATION
-};
-
-/**
- * Agent handle per client, these will be at TArray indices of the Client handles indicies (used as a free list array)
- */
-USTRUCT()
-struct MASSREPLICATION_API FMassReplicatedAgentFragment : public FMassFragment
-{
-	GENERATED_BODY()
-
-	//these will be maintained so that the TArray entries are at the same indices as Client Handles
-	TArray<FMassReplicatedAgentArrayData, TInlineAllocator<UE::Mass::Replication::MaxNumOfClients>> AgentsData;
-};
-
 USTRUCT()
 struct FReplicatedAgentBase
 {
@@ -195,57 +153,6 @@ struct FReplicationLODLogic : public FLODDefaultLogic
 	};
 };
 
-/*
- * Data fragment to store the calculated distances to viewers
- */
-USTRUCT()
-struct MASSREPLICATION_API FMassReplicationViewerInfoFragment : public FMassFragment
-{
-	GENERATED_BODY()
-
-	/** Closest viewer distance */
-	float ClosestViewerDistanceSq;
-
-	/** Distance between each viewer and entity */
-	TArray<float> DistanceToViewerSq;
-};
-
-USTRUCT()
-struct MASSREPLICATION_API FMassReplicationLODFragment : public FMassFragment
-{
-	GENERATED_BODY()
-
-	/**LOD information */
-	TEnumAsByte<EMassLOD::Type> LOD = EMassLOD::Max;
-
-	/** Previous LOD information*/
-	TEnumAsByte<EMassLOD::Type> PrevLOD = EMassLOD::Max;
-
-	/** Per viewer LOD information */
-	TArray<TEnumAsByte<EMassLOD::Type>> LODPerViewer;
-
-	/** Per viewer previous LOD information */
-	TArray<TEnumAsByte<EMassLOD::Type>> PrevLODPerViewer;
-};
-
-UCLASS()
-class MASSREPLICATION_API UMassNetworkIDFragmentInitializer : public UMassFragmentInitializer
-{
-	GENERATED_BODY()
-public:
-	UMassNetworkIDFragmentInitializer();
-
-protected:
-	virtual void ConfigureQueries() override;
-	virtual void Initialize(UObject& Owner) override;
-	virtual void Execute(UMassEntitySubsystem& EntitySubsystem, FMassExecutionContext& Context) override;
-
-	UPROPERTY()
-	UMassReplicationManager* ReplicationManager = nullptr;
-
-	FMassEntityQuery EntityQuery;
-};
-
 #if UE_REPLICATION_COMPILE_CLIENT_CODE
 /** Data that can be accessed from a FMassNetworkID on a client */
 struct FMassReplicationEntityInfo
