Rename MassReplicationManager to MassReplicationSubsystem to be consistent with the rest of the mass code

#rnx

diff --git a/Engine/Plugins/AI/MassCrowd/Source/MassCrowd/Private/MassCrowdBubble.cpp b/Engine/Plugins/AI/MassCrowd/Source/MassCrowd/Private/MassCrowdBubble.cpp
--- a/Engine/Plugins/AI/MassCrowd/Source/MassCrowd/Private/MassCrowdBubble.cpp
+++ b/Engine/Plugins/AI/MassCrowd/Source/MassCrowd/Private/MassCrowdBubble.cpp
@@ -100,8 +100,8 @@ void FMassCrowdClientBubbleHandler::DebugValidateBubbleOnClient()
 		UMassEntitySubsystem* EntitySystem = Serializer->GetEntitySystem();
 		check(EntitySystem);
 
-		UMassReplicationManager* ReplicationManager = Serializer->GetReplicationManager();
-		check(ReplicationManager);
+		UMassReplicationSubsystem* ReplicationSubsystem = Serializer->GetReplicationSubsystem();
+		check(ReplicationSubsystem);
 
 		// @todo cap at MaxAgentsDraw for now
 		const int32 MaxAgentsDraw = FMath::Min(UE::Mass::Crowd::MaxAgentsDraw, (*Agents).Num());
@@ -110,7 +110,7 @@ void FMassCrowdClientBubbleHandler::DebugValidateBubbleOnClient()
 		{
 			const FCrowdFastArrayItem& CrowdItem = (*Agents)[Idx];
 
-			const FMassReplicationEntityInfo* EntityInfo = ReplicationManager->FindMassEntityInfo(CrowdItem.Agent.GetNetID());
+			const FMassReplicationEntityInfo* EntityInfo = ReplicationSubsystem->FindMassEntityInfo(CrowdItem.Agent.GetNetID());
 
 			check(EntityInfo->Entity.IsSet());
 
diff --git a/Engine/Plugins/AI/MassCrowd/Source/MassCrowd/Private/MassCrowdReplicationProcessor.cpp b/Engine/Plugins/AI/MassCrowd/Source/MassCrowd/Private/MassCrowdReplicationProcessor.cpp
--- a/Engine/Plugins/AI/MassCrowd/Source/MassCrowd/Private/MassCrowdReplicationProcessor.cpp
+++ b/Engine/Plugins/AI/MassCrowd/Source/MassCrowd/Private/MassCrowdReplicationProcessor.cpp
@@ -1,7 +1,7 @@
 // Copyright Epic Games, Inc. All Rights Reserved.
 
 #include "MassCrowdReplicationProcessor.h"
-#include "MassReplicationManager.h"
+#include "MassReplicationSubsystem.h"
 #include "MassReplicationTypes.h"
 #include "MassClientBubbleHandler.h"
 #include "MassCrowdBubble.h"
@@ -45,9 +45,9 @@ void UMassCrowdReplicationProcessor::Initialize(UObject& Owner)
 {
 	Super::Initialize(Owner);
 
-	check(ReplicationManager);
+	check(ReplicationSubsystem);
 
-	BubbleInfoClassHandle = ReplicationManager->GetBubbleInfoClassHandle(AMassCrowdClientBubbleInfo::StaticClass());
+	BubbleInfoClassHandle = ReplicationSubsystem->GetBubbleInfoClassHandle(AMassCrowdClientBubbleInfo::StaticClass());
 }
 
 void UMassCrowdReplicationProcessor::Execute(UMassEntitySubsystem& EntitySubsystem, FMassExecutionContext& Context)
diff --git a/Engine/Plugins/AI/MassCrowd/Source/MassCrowd/Private/MassCrowdSubsystem.cpp b/Engine/Plugins/AI/MassCrowd/Source/MassCrowd/Private/MassCrowdSubsystem.cpp
--- a/Engine/Plugins/AI/MassCrowd/Source/MassCrowd/Private/MassCrowdSubsystem.cpp
+++ b/Engine/Plugins/AI/MassCrowd/Source/MassCrowd/Private/MassCrowdSubsystem.cpp
@@ -5,7 +5,7 @@
 #include "ZoneGraphData.h"
 #include "MassSimulationSubsystem.h"
 #include "MassCrowdBubble.h"
-#include "MassReplicationManager.h"
+#include "MassReplicationSubsystem.h"
 #include "MassCrowdFragments.h"
 #include "MassCrowdSettings.h"
 #include "ZoneGraphAnnotationSubsystem.h"
@@ -106,10 +106,10 @@ void UMassCrowdSubsystem::PostInitialize()
 {
 	Super::PostInitialize();
 
-	UMassReplicationManager* ReplicationManager = UWorld::GetSubsystem<UMassReplicationManager>(GetWorld());
+	UMassReplicationSubsystem* ReplicationSubsystem = UWorld::GetSubsystem<UMassReplicationSubsystem>(GetWorld());
 
-	check(ReplicationManager);
-	ReplicationManager->RegisterBubbleInfoClass(AMassCrowdClientBubbleInfo::StaticClass());
+	check(ReplicationSubsystem);
+	ReplicationSubsystem->RegisterBubbleInfoClass(AMassCrowdClientBubbleInfo::StaticClass());
 }
 
 void UMassCrowdSubsystem::PostZoneGraphDataAdded(const AZoneGraphData* ZoneGraphData)
diff --git a/Engine/Plugins/AI/MassCrowd/Source/MassCrowd/Public/MassCrowdReplicationProcessor.h b/Engine/Plugins/AI/MassCrowd/Source/MassCrowd/Public/MassCrowdReplicationProcessor.h
--- a/Engine/Plugins/AI/MassCrowd/Source/MassCrowd/Public/MassCrowdReplicationProcessor.h
+++ b/Engine/Plugins/AI/MassCrowd/Source/MassCrowd/Public/MassCrowdReplicationProcessor.h
@@ -6,7 +6,7 @@
 
 #include "MassCrowdReplicationProcessor.generated.h"
 
-class UMassReplicationManager;
+class UMassReplicationSubsystem;
 class MassLODManager;
 class AMassCrowdClientBubbleInfo;
 class UWorld;
diff --git a/Engine/Plugins/Runtime/MassGameplay/Source/MassActors/Private/MassAgentComponent.cpp b/Engine/Plugins/Runtime/MassGameplay/Source/MassActors/Private/MassAgentComponent.cpp
--- a/Engine/Plugins/Runtime/MassGameplay/Source/MassActors/Private/MassAgentComponent.cpp
+++ b/Engine/Plugins/Runtime/MassGameplay/Source/MassActors/Private/MassAgentComponent.cpp
@@ -10,7 +10,7 @@
 #include "MassActorSubsystem.h"
 #include "MassSpawner.h"
 #include "Net/UnrealNetwork.h"
-#include "MassReplicationManager.h"
+#include "MassReplicationSubsystem.h"
 #include "MassReplicationFragments.h"
 #include "MassMovementFragments.h"
 #include "GameFramework/CharacterMovementComponent.h"
diff --git a/Engine/Plugins/Runtime/MassGameplay/Source/MassActors/Private/MassAgentSubsystem.cpp b/Engine/Plugins/Runtime/MassGameplay/Source/MassActors/Private/MassAgentSubsystem.cpp
--- a/Engine/Plugins/Runtime/MassGameplay/Source/MassActors/Private/MassAgentSubsystem.cpp
+++ b/Engine/Plugins/Runtime/MassGameplay/Source/MassActors/Private/MassAgentSubsystem.cpp
@@ -12,7 +12,7 @@
 #include "MassEntityTemplateRegistry.h"
 #include "VisualLogger/VisualLogger.h"
 #include "MassEntityView.h"
-#include "MassReplicationManager.h"
+#include "MassReplicationSubsystem.h"
 #include "Engine/NetDriver.h"
 
 
@@ -42,7 +42,7 @@ void UMassAgentSubsystem::Initialize(FSubsystemCollectionBase& Collection)
 	Collection.InitializeDependency<UMassSimulationSubsystem>();
 	Collection.InitializeDependency<UMassSpawnerSubsystem>();
 #if UE_REPLICATION_COMPILE_CLIENT_CODE
-	Collection.InitializeDependency<UMassReplicationManager>();
+	Collection.InitializeDependency<UMassReplicationSubsystem>();
 #endif // UE_REPLICATION_COMPILE_CLIENT_CODE
 	
 	UWorld* World = GetWorld();
@@ -53,11 +53,11 @@ void UMassAgentSubsystem::Initialize(FSubsystemCollectionBase& Collection)
 	SimulationSystem->GetOnProcessingPhaseStarted(EMassProcessingPhase::PrePhysics).AddUObject(this, &UMassAgentSubsystem::OnProcessingPhaseStarted, EMassProcessingPhase::PrePhysics);
 
 #if UE_REPLICATION_COMPILE_CLIENT_CODE
-	ReplicationManager = UWorld::GetSubsystem<UMassReplicationManager>(World);
-	check(ReplicationManager);
+	ReplicationSubsystem = UWorld::GetSubsystem<UMassReplicationSubsystem>(World);
+	check(ReplicationSubsystem);
 
-	ReplicationManager->GetOnMassAgentAdded().AddUObject(this, &UMassAgentSubsystem::OnMassAgentAddedToReplication);
-	ReplicationManager->GetOnRemovingMassAgent().AddUObject(this, &UMassAgentSubsystem::OnMassAgentRemovedFromReplication);
+	ReplicationSubsystem->GetOnMassAgentAdded().AddUObject(this, &UMassAgentSubsystem::OnMassAgentAddedToReplication);
+	ReplicationSubsystem->GetOnRemovingMassAgent().AddUObject(this, &UMassAgentSubsystem::OnMassAgentRemovedFromReplication);
 #endif // UE_REPLICATION_COMPILE_CLIENT_CODE
 }
 
@@ -379,11 +379,11 @@ void UMassAgentSubsystem::NotifyMassAgentComponentReplicated(UMassAgentComponent
 	UWorld* World = GetWorld();
 	if (World && ensureMsgf(World->IsNetMode(NM_Client), TEXT("%s: Expecting to only be called in network game on the client"), ANSI_TO_TCHAR(__FUNCTION__)))
 	{
-		check(ReplicationManager);
+		check(ReplicationSubsystem);
 		check(AgentComp.GetNetID().IsValid());
 		check(!ReplicatedAgentComponents.Find(AgentComp.GetNetID()));
 		ReplicatedAgentComponents.Add(AgentComp.GetNetID(), &AgentComp);
-		const FMassEntityHandle Entity = ReplicationManager->FindEntity(AgentComp.GetNetID());
+		const FMassEntityHandle Entity = ReplicationSubsystem->FindEntity(AgentComp.GetNetID());
 
 		// If not found, the NotifyMassAgentAddedToReplication will link it later once replicated.
 		if (Entity.IsSet())
diff --git a/Engine/Plugins/Runtime/MassGameplay/Source/MassActors/Public/MassAgentSubsystem.h b/Engine/Plugins/Runtime/MassGameplay/Source/MassActors/Public/MassAgentSubsystem.h
--- a/Engine/Plugins/Runtime/MassGameplay/Source/MassActors/Public/MassAgentSubsystem.h
+++ b/Engine/Plugins/Runtime/MassGameplay/Source/MassActors/Public/MassAgentSubsystem.h
@@ -15,7 +15,7 @@ class UMassAgentComponent;
 class UMassSimulationSubsystem;
 class UMassAgentSubsystem;
 class UMassAgentComponent;
-class UMassReplicationManager;
+class UMassReplicationSubsystem;
 
 namespace UE::MassActor
 {
@@ -137,7 +137,7 @@ protected:
 	TMap<FMassEntityTemplateID, FMassAgentInitializationQueue> PendingPuppets;
 
 	UPROPERTY()
-	UMassReplicationManager* ReplicationManager;
+	UMassReplicationSubsystem* ReplicationSubsystem;
 
 	UPROPERTY()
 	TMap<FMassNetworkID, UMassAgentComponent*> ReplicatedAgentComponents;
diff --git a/Engine/Plugins/Runtime/MassGameplay/Source/MassReplication/Private/MassClientBubbleSerializerBase.cpp b/Engine/Plugins/Runtime/MassGameplay/Source/MassReplication/Private/MassClientBubbleSerializerBase.cpp
--- a/Engine/Plugins/Runtime/MassGameplay/Source/MassReplication/Private/MassClientBubbleSerializerBase.cpp
+++ b/Engine/Plugins/Runtime/MassGameplay/Source/MassReplication/Private/MassClientBubbleSerializerBase.cpp
@@ -10,7 +10,7 @@ void FMassClientBubbleSerializerBase::InitializeForWorld(UWorld& InWorld)
 	SpawnerSubsystem = InWorld.GetSubsystem<UMassSpawnerSubsystem>();
 	check(SpawnerSubsystem);
 
-	ReplicationManager = InWorld.GetSubsystem<UMassReplicationManager>();
+	ReplicationSubsystem = InWorld.GetSubsystem<UMassReplicationSubsystem>();
 
 	EntitySystem = UWorld::GetSubsystem<UMassEntitySubsystem>(World);
 	check(EntitySystem);
diff --git a/Engine/Plugins/Runtime/MassGameplay/Source/MassReplication/Private/MassReplicationFragments.cpp b/Engine/Plugins/Runtime/MassGameplay/Source/MassReplication/Private/MassReplicationFragments.cpp
--- a/Engine/Plugins/Runtime/MassGameplay/Source/MassReplication/Private/MassReplicationFragments.cpp
+++ b/Engine/Plugins/Runtime/MassGameplay/Source/MassReplication/Private/MassReplicationFragments.cpp
@@ -1,7 +1,7 @@
 // Copyright Epic Games, Inc. All Rights Reserved.
 
 #include "MassReplicationFragments.h"
-#include "MassReplicationManager.h"
+#include "MassReplicationSubsystem.h"
 #include "Engine/World.h"
 
 //----------------------------------------------------------------------//
@@ -23,7 +23,7 @@ void UMassNetworkIDFragmentInitializer::Initialize(UObject& Owner)
 {
 	Super::Initialize(Owner);
 #if UE_REPLICATION_COMPILE_SERVER_CODE
-	ReplicationManager = UWorld::GetSubsystem<UMassReplicationManager>(Owner.GetWorld());
+	ReplicationSubsystem = UWorld::GetSubsystem<UMassReplicationSubsystem>(Owner.GetWorld());
 #endif //UE_REPLICATION_COMPILE_SERVER_CODE
 }
 
@@ -37,7 +37,7 @@ void UMassNetworkIDFragmentInitializer::Execute(UMassEntitySubsystem& EntitySubs
 	if (NetMode != NM_Client)
 	{
 #if UE_REPLICATION_COMPILE_SERVER_CODE
-		check(ReplicationManager);
+		check(ReplicationSubsystem);
 
 		EntityQuery.ForEachEntityChunk(EntitySubsystem, Context, [this](FMassExecutionContext& Context)
 			{
@@ -46,7 +46,7 @@ void UMassNetworkIDFragmentInitializer::Execute(UMassEntitySubsystem& EntitySubs
 
 				for (int32 Idx = 0; Idx < NumEntities; ++Idx)
 				{
-					NetworkIDList[Idx].NetID = ReplicationManager->GetNextAvailableMassNetID();
+					NetworkIDList[Idx].NetID = ReplicationSubsystem->GetNextAvailableMassNetID();
 				}
 			});
 #endif //UE_REPLICATION_COMPILE_SERVER_CODE
diff --git a/Engine/Plugins/Runtime/MassGameplay/Source/MassReplication/Private/MassReplicationManager.cpp b/Engine/Plugins/Runtime/MassGameplay/Source/MassReplication/Private/MassReplicationManager.cpp
deleted file mode 100644
diff --git a/Engine/Plugins/Runtime/MassGameplay/Source/MassReplication/Private/MassReplicationProcessorBase.cpp b/Engine/Plugins/Runtime/MassGameplay/Source/MassReplication/Private/MassReplicationProcessorBase.cpp
--- a/Engine/Plugins/Runtime/MassGameplay/Source/MassReplication/Private/MassReplicationProcessorBase.cpp
+++ b/Engine/Plugins/Runtime/MassGameplay/Source/MassReplication/Private/MassReplicationProcessorBase.cpp
@@ -55,11 +55,11 @@ void UMassReplicationProcessorBase::Initialize(UObject& Owner)
 #if UE_REPLICATION_COMPILE_SERVER_CODE
 	LODCalculator.Initialize(LODDistance, BufferHysteresisOnDistancePercentage / 100.0f, LODMaxCount, LODMaxCountPerViewer);
 
-	ReplicationManager = UWorld::GetSubsystem<UMassReplicationManager>(World);
+	ReplicationSubsystem = UWorld::GetSubsystem<UMassReplicationSubsystem>(World);
 
-	check(ReplicationManager);
+	check(ReplicationSubsystem);
 
-	//BubbleInfoClassHandle = ReplicationManager->GetBubbleInfoClassHandle(AMassCrowdClientBubbleInfo::StaticClass());
+	//BubbleInfoClassHandle = ReplicationSubsystem->GetBubbleInfoClassHandle(AMassCrowdClientBubbleInfo::StaticClass());
 #endif //UE_REPLICATION_COMPILE_SERVER_CODE
 }
 
@@ -67,10 +67,10 @@ void UMassReplicationProcessorBase::PrepareExecution()
 {
 #if UE_REPLICATION_COMPILE_SERVER_CODE
 
-	check(ReplicationManager);
+	check(ReplicationSubsystem);
 
 	//first synchronize clients and viewers
-	ReplicationManager->SynchronizeClientsAndViewers();
+	ReplicationSubsystem->SynchronizeClientsAndViewers();
 
 	//then call base class
 	check(LODManager);
@@ -78,7 +78,7 @@ void UMassReplicationProcessorBase::PrepareExecution()
 	LODCollector.PrepareExecution(Viewers);
 	LODCalculator.PrepareExecution(Viewers);
 
-	const TArray<FMassClientHandle>& CurrentClientHandles = ReplicationManager->GetClientReplicationHandles();
+	const TArray<FMassClientHandle>& CurrentClientHandles = ReplicationSubsystem->GetClientReplicationHandles();
 	const int32 MinNumHandles = FMath::Min(CachedClientHandles.Num(), CurrentClientHandles.Num());
 
 	//check to see if we don't have enough cached client handles
@@ -95,7 +95,7 @@ void UMassReplicationProcessorBase::PrepareExecution()
 			CachedClientHandles.Add(CurrentClientHandle);
 			bBubbleChanged.Add(true);
 			AMassClientBubbleInfoBase* Info = CurrentClientHandle.IsValid() ?
-				ReplicationManager->GetClientBubbleChecked(BubbleInfoClassHandle, CurrentClientHandle) :
+				ReplicationSubsystem->GetClientBubbleChecked(BubbleInfoClassHandle, CurrentClientHandle) :
 				nullptr;
 
 			check(Info);
@@ -123,7 +123,7 @@ void UMassReplicationProcessorBase::PrepareExecution()
 		if (bChanged)
 		{
 			AMassClientBubbleInfoBase* Info = CurrentClientHandle.IsValid() ?
-				ReplicationManager->GetClientBubbleChecked(BubbleInfoClassHandle, CurrentClientHandle) :
+				ReplicationSubsystem->GetClientBubbleChecked(BubbleInfoClassHandle, CurrentClientHandle) :
 				nullptr;
 
 			check(Info);
diff --git a/Engine/Plugins/Runtime/MassGameplay/Source/MassReplication/Private/MassReplicationSubsystem.cpp b/Engine/Plugins/Runtime/MassGameplay/Source/MassReplication/Private/MassReplicationSubsystem.cpp
new file mode 100644
--- /dev/null
+++ b/Engine/Plugins/Runtime/MassGameplay/Source/MassReplication/Private/MassReplicationSubsystem.cpp
@@ -0,0 +1,697 @@
+// Copyright Epic Games, Inc. All Rights Reserved.
+
+#include "MassReplicationSubsystem.h"
+#include "Engine/World.h"
+#include "Engine/ChildConnection.h"
+#include "GameFramework/GameModeBase.h"
+#include "MassCommonTypes.h"
+#include "MassEntitySubsystem.h"
+#include "MassClientBubbleHandler.h"
+#include "MassClientBubbleInfoBase.h"
+
+
+uint32 UMassReplicationSubsystem::CurrentNetMassCounter = 0;
+
+void UMassReplicationSubsystem::Initialize(FSubsystemCollectionBase& Collection)
+{
+	Super::Initialize(Collection);
+
+	World = GetWorld();
+
+	check(World);
+
+	MassLODManager = Collection.InitializeDependency<UMassLODManager>();
+
+	check(MassLODManager);
+
+	EntitySystem = UWorld::GetSubsystem<UMassEntitySubsystem>(GetWorld());
+}
+
+void UMassReplicationSubsystem::Deinitialize()
+{
+	// remove all Clients
+	for (const FMassClientHandle& Handle : ClientHandleManager.GetHandles())
+	{
+		if (Handle.IsValid())
+		{
+			RemoveClient(Handle);
+		}
+	}
+
+	// make sure all the other data is reset
+	ClientHandleManager.Reset();
+	BubbleInfoArray.Reset();
+	ClientsViewerHandles.Reset();
+	ClientToViewerHandleArray.Reset();
+	ViewerToClientHandleArray.Reset();
+
+	World = nullptr;
+	MassLODManager = nullptr;
+	EntitySystem = nullptr;
+}
+
+FMassNetworkID UMassReplicationSubsystem::GetNetIDFromHandle(const FMassEntityHandle Handle) const
+{
+	check(EntitySystem);
+	const FMassNetworkIDFragment& Data = EntitySystem->GetFragmentDataChecked<FMassNetworkIDFragment>(Handle);
+	return Data.NetID;
+}
+
+namespace UE { namespace Mass { namespace Replication {
+// Checks the parent net connection has APlayerController OwningActor, unfortunately GetParentConnection() prevents UChildConnection being const here.
+APlayerController* GetValidParentNetConnection(UChildConnection& ChildConnection)
+{
+	return (ChildConnection.GetParentConnection() != nullptr) ? Cast<APlayerController>(ChildConnection.GetParentConnection()->OwningActor) : nullptr;
+}
+
+enum class ENetConnectionType : uint8
+{
+	None,
+	Parent,
+	Child,
+};
+
+// Checks Controllers net connection is a parent net connection
+bool HasParentNetConnection(const APlayerController* Controller)
+{
+	return (Controller != nullptr) && Controller->NetConnection && (Cast<UChildConnection>(Controller->NetConnection) == nullptr);
+}
+
+// If Controller has a UChildConnection then this gets the parent net connection's APlayerController OwningActor otherwise nullptr
+APlayerController* GetParentControllerFromChildNetConnection(const APlayerController* Controller)
+{
+	//ChildConnection can't be const
+	UChildConnection* ChildConnection = Controller ? Cast<UChildConnection>(Controller->NetConnection) : nullptr;
+
+	return ChildConnection ? GetValidParentNetConnection(*ChildConnection) : nullptr;
+}
+
+// Gets the ENetConnectionType for Controller. Tests if Controller's net connection is a parent or if Controller's net connection is a child and its parent net connection has a APlayerController OwningActor, 
+ENetConnectionType HasParentOrChildWithValidParentNetConnection(const APlayerController* Controller)
+{
+	ENetConnectionType ConnectionType = ENetConnectionType::None;
+
+	if (Controller && Controller->NetConnection)
+	{
+		UChildConnection* ChildConnection = Cast<UChildConnection>(Controller->NetConnection);
+
+		// if the NetConnection is either a parent connection or a child connection with APlayerController owning actor its valid 
+		if (ChildConnection == nullptr)
+		{
+			ConnectionType = ENetConnectionType::Parent;
+		}
+		else
+		{
+			ConnectionType = (GetValidParentNetConnection(*ChildConnection) != nullptr) ? ENetConnectionType::Child : ENetConnectionType::None;
+		}		
+	}
+
+	return ConnectionType;
+}
+}}};
+
+bool UMassReplicationSubsystem::SynchronizeClients(const TArray<FViewerInfo>& Viewers)
+{
+	typedef TMap<FMassViewerHandle, FMassClientHandle> FControllerMap;
+	struct FClientAddData
+	{
+		FClientAddData(FMassViewerHandle InHandle, APlayerController* InController)
+			: Handle(InHandle)
+			, Controller(InController)
+		{}
+
+		FMassViewerHandle Handle;
+		APlayerController* Controller = nullptr;
+	};
+
+	bool bNeedShrinking = false;
+
+	//we can only replicate if we have some BubbleInfos to use for replication
+	if (BubbleInfoArray.Num() == 0)
+	{
+		return bNeedShrinking;
+	}
+
+	// Arbitrarily use index 0, if there are more items the corresponding InfoData, Bubbles will be the same length and corresponding indices between the Bubbles
+	// will have the same player controller owner.
+	const FMassClientBubbleInfoData& InfoData = BubbleInfoArray[0];
+
+	check(MassLODManager);
+
+	FControllerMap ClientConnectionMap;
+
+	TArray<FClientAddData> ClientsToAdd;
+
+	{
+		// Go through the stored Clients, add valid ones to the ClientConnectionMap and remove invalids.
+		// Note this is only valid until we next add a client handle.
+		const TArray<FMassClientHandle>& ClientHandles = ClientHandleManager.GetHandles();
+
+		for (int32 Idx = 0; Idx < ClientHandles.Num(); ++Idx)
+		{
+			const FMassClientHandle& ClientHandle = ClientHandles[Idx];
+
+			if (ClientHandle.IsValid())
+			{
+				check(ClientToViewerHandleArray.IsValidIndex(ClientHandle.GetIndex()));
+
+				const FViewerClientPair ViewerClientPair = ClientToViewerHandleArray[ClientHandle.GetIndex()];
+
+				check(ViewerClientPair.ViewerHandle.IsValid());
+				check(ViewerClientPair.ClientHandle == ClientHandle);
+				check(InfoData.Bubbles.IsValidIndex(ClientHandle.GetIndex()));
+
+				const AMassClientBubbleInfoBase* ClientBubble = InfoData.Bubbles[ClientHandle.GetIndex()];
+				const APlayerController* Controller = ClientBubble ? Cast<APlayerController>(ClientBubble->GetOwner()) : nullptr;
+
+//no need to check netconnection if UE_ALLOW_DEBUG_REPLICATION_BUBBLES_STANDALONE
+#if !UE_ALLOW_DEBUG_REPLICATION_BUBBLES_STANDALONE
+				// must have a valid non child UNetConnection
+				if (UE::Mass::Replication::HasParentNetConnection(Controller))
+#endif
+				{
+					ClientConnectionMap.Add(ViewerClientPair.ViewerHandle, ClientHandle);
+				}
+#if !UE_ALLOW_DEBUG_REPLICATION_BUBBLES_STANDALONE
+				else
+				{
+					// this is safe as Clients and their handles are free list arrays
+					RemoveClient(ClientHandle);
+					bNeedShrinking |= Idx == ClientHandles.Num() - 1;
+				}
+#endif
+			}
+		}
+
+		// now go through all current Viewers and add if they do not exist
+		for (const FViewerInfo& Viewer : Viewers)
+		{
+			// must have valid Viewer and parent net connection
+			if (Viewer.Handle.IsValid())
+			{
+//no need to check netconnection if UE_ALLOW_DEBUG_REPLICATION_BUBBLES_STANDALONE
+#if !UE_ALLOW_DEBUG_REPLICATION_BUBBLES_STANDALONE
+				if (UE::Mass::Replication::HasParentNetConnection(Viewer.PlayerController))
+#endif
+				{
+					// check if the controller already exists by trying to remove it from the map which was filled up with controllers we were tracking
+					if (ClientConnectionMap.Remove(Viewer.Handle) == 0)
+					{
+						// If not add it to ClientsToAdd. Its important AddClient isn't called until necessary Clients are removed, as we may reuse 
+						// array indices that are already in use.
+						ClientsToAdd.Emplace(Viewer.Handle, Viewer.PlayerController);
+					}
+				}
+			}
+		}
+
+		// anything left in the map needs to be removed from the list
+		for (FControllerMap::TIterator Itr = ClientConnectionMap.CreateIterator(); Itr; ++Itr)
+		{
+			const int32 ViewerIdx = Itr->Value.GetIndex();
+
+			RemoveClient(Itr->Value);
+
+			//InfoData must be valid here as ClientConnectionMap can only have items if ClientHandleManager.GetHandles() was greater than 0 
+			//when we entered the function.
+			bNeedShrinking |= (ViewerIdx == (InfoData.Bubbles.Num() - 1));
+		}
+	}
+
+	// resize the ViewerToClientHandleArray array to match the viewer handles array (for consistency)
+	if (ViewerToClientHandleArray.Num() > Viewers.Num())
+	{
+		ViewerToClientHandleArray.RemoveAt(Viewers.Num(), ViewerToClientHandleArray.Num() - Viewers.Num(), /* bAllowShrinking */ false);
+	}
+
+	for (FClientAddData& ClientAdd : ClientsToAdd)
+	{
+		//TODO make AddClients (ie plural) functionality
+		AddClient(ClientAdd.Handle, *ClientAdd.Controller);
+	}
+
+	return bNeedShrinking;
+}
+
+// synchronize the ClientViewers
+void UMassReplicationSubsystem::SynchronizeClientViewers(const TArray<FViewerInfo>& Viewers)
+{
+	check(MassLODManager);
+
+	struct FMassClientViewerHandle
+	{
+		FMassClientViewerHandle(FMassClientHandle InClientHandle, FMassViewerHandle InViewerHandle)
+			: ClientHandle(InClientHandle)
+			, ViewerHandle(InViewerHandle)
+		{}
+
+		FMassClientHandle ClientHandle;
+		FMassViewerHandle ViewerHandle;
+	};
+
+	// go through the ClientsViewerHandles and check validity and store the valid ones into a map
+	typedef TMap<APlayerController*, FMassClientViewerHandle> FViewerMap;
+
+	FViewerMap ClientViewerMap;
+	{
+		const TArray<FMassClientHandle> ClientHandles = ClientHandleManager.GetHandles();
+
+		for (const FMassClientHandle& ClientHandle : ClientHandles)
+		{
+			//as this is a fresh handle then we only need to check against !IsInvalid
+			if (ClientHandle.IsValid())
+			{
+				FClientViewerHandles& ClientViewers = ClientsViewerHandles[ClientHandle.GetIndex()];
+
+				int32 ViewerIdx = 0;
+
+				while (ViewerIdx < ClientViewers.Handles.Num())
+				{
+					const FMassViewerHandle& ClientViewer = ClientViewers.Handles[ViewerIdx];
+
+					APlayerController* Controller = MassLODManager->GetPlayerControllerFromViewerHandle(ClientViewer);
+
+					const UE::Mass::Replication::ENetConnectionType ConnectionType = UE::Mass::Replication::HasParentOrChildWithValidParentNetConnection(Controller);
+
+					if (ConnectionType != UE::Mass::Replication::ENetConnectionType::None)
+					{
+						++ViewerIdx;
+
+						// we don't verify or remove Parent Connections here as that was done when we synchronized the client bubbles
+						if (ConnectionType == UE::Mass::Replication::ENetConnectionType::Child)
+						{
+							ClientViewerMap.Add(Controller, FMassClientViewerHandle(ClientHandle, ClientViewer));
+						}
+					}
+					else //remove invalid ClientViewer, but dont increment the ViewerIdx
+					{
+						ClientViewers.Handles.RemoveAt(ViewerIdx, 1, /* bAllowShrinking */ false);
+					}
+				}
+			}
+		}
+
+		// now go through all current Viewers and add if they are a valid ClientViewer with a child UNetconnection if they do not exist
+		for (const FViewerInfo& Viewer : Viewers)
+		{
+			// we are only interested in valid Viewers
+			if (Viewer.Handle.IsValid())
+			{
+				// we are only processing child UNetConnections that have a valid APlayerController OwningActor
+				const APlayerController* ParentController = UE::Mass::Replication::GetParentControllerFromChildNetConnection(Viewer.PlayerController);
+
+				// check if the parent controller is valid and already exists
+				if (ParentController && (ClientViewerMap.Find(Viewer.PlayerController) == nullptr))
+				{
+					FMassViewerHandle ParentViewerHandle = MassLODManager->GetViewerHandleFromPlayerController(ParentController);
+
+					if (ensureMsgf(ParentViewerHandle.IsValid(), TEXT("MassLODManager handles are out of sync with PlayerController NetConnections!")))
+					{
+						// note all clients (parent NetConnections) should already be set up so we would expect valid handles here
+						check(ViewerToClientHandleArray.IsValidIndex(ParentViewerHandle.GetIndex()));
+
+						const FViewerClientPair& ParentViewerClientPair = ViewerToClientHandleArray[ParentViewerHandle.GetIndex()];
+
+						check(MassLODManager->IsValidViewer(ParentViewerClientPair.ViewerHandle));
+						check(ClientHandleManager.IsValidHandle(ParentViewerClientPair.ClientHandle));
+
+						// remove APlayerController from the ClientViewerMap and Add the viewer to the ClientsViewerHandles
+						ClientViewerMap.Remove(Viewer.PlayerController);
+
+						FClientViewerHandles& ClientViewers = ClientsViewerHandles[ParentViewerClientPair.ClientHandle.GetIndex()];
+
+						ClientViewers.Handles.Add(Viewer.Handle);
+					}
+				}
+			}
+		}
+	}
+
+	{
+		// anything left in the map needs to be removed from the list
+		for (FViewerMap::TIterator Itr = ClientViewerMap.CreateIterator(); Itr; ++Itr)
+		{
+			const FMassClientViewerHandle& HandleData = Itr->Value;
+
+			FClientViewerHandles& ClientViewers = ClientsViewerHandles[HandleData.ClientHandle.GetIndex()];
+
+			ClientViewers.Handles.RemoveSingle(HandleData.ViewerHandle);
+		}
+	}
+}
+
+void UMassReplicationSubsystem::SynchronizeClientsAndViewers()
+{
+	// only execute this code at most once per frame
+	if (LastSynchronizedFrame == GFrameCounter)
+	{
+		return;
+	}
+	LastSynchronizedFrame = GFrameCounter;
+
+	if (MassLODManager == nullptr)
+	{
+		checkNoEntry();
+		return;
+	}
+
+	// makes sure the LOD manager Viewers are synced before we process them
+	const TArray<FViewerInfo>& Viewers = MassLODManager->GetSynchronizedViewers();
+
+	const bool bNeedShrinking = SynchronizeClients(Viewers);
+
+	//only synchronize the client viewers outside of the debug functionality for now
+#if !UE_ALLOW_DEBUG_REPLICATION_BUBBLES_STANDALONE
+	SynchronizeClientViewers(Viewers);
+#endif //UE_ALLOW_REPLICATION_DEBUG_BUBBLES_IN_STANDALONE
+
+	if (bNeedShrinking)
+	{
+		const int32 NumItems = ClientHandleManager.ShrinkHandles();
+		 
+		ClientsViewerHandles.RemoveAt(NumItems, ClientsViewerHandles.Num() - NumItems, /* bAllowShrinking */ false);
+		ClientToViewerHandleArray.RemoveAt(NumItems, ClientToViewerHandleArray.Num() - NumItems, /* bAllowShrinking */ false);
+
+		for (FMassClientBubbleInfoData& InfoData : BubbleInfoArray)
+		{
+			InfoData.Bubbles.RemoveAt(NumItems, InfoData.Bubbles.Num() - NumItems, /* bAllowShrinking */ false);
+		}
+	}
+}
+
+FMassBubbleInfoClassHandle UMassReplicationSubsystem::RegisterBubbleInfoClass(const TSubclassOf<AMassClientBubbleInfoBase>& BubbleInfoClass)
+{
+	checkf(BubbleInfoClass.Get() != nullptr, TEXT("BubbleInfoClass must have been set!"));
+
+	if (ClientHandleManager.GetHandles().Num() > 0)
+	{
+		checkf(false, TEXT("RegisterBubbleInfoClass() must not be called after AddClient, BubbleInfoClass has not been registered"));
+		return FMassBubbleInfoClassHandle();
+	}
+
+	const int32 IdxFound = BubbleInfoArray.IndexOfByPredicate([BubbleInfoClass](const FMassClientBubbleInfoData& Data)
+		{
+			return BubbleInfoClass == Data.BubbleClass;
+		});
+
+	if (IdxFound != INDEX_NONE)
+	{
+		UE_LOG(LogMassReplication, Log, TEXT("UMassReplicationSubsystem: Trying to RegisterBubbleInfoClass() twice with the same BubbleInfoClass, Only one BubbleInfoClass will be registered for this type"));
+		return FMassBubbleInfoClassHandle(IdxFound);
+	}
+
+	const int32 Idx = BubbleInfoArray.Emplace(BubbleInfoClass);
+
+	return FMassBubbleInfoClassHandle(Idx);
+}
+
+FMassBubbleInfoClassHandle UMassReplicationSubsystem::GetBubbleInfoClassHandle(const TSubclassOf<AMassClientBubbleInfoBase>& BubbleInfoClass) const
+{
+	FMassBubbleInfoClassHandle Handle;
+
+	for (int32 Idx = 0; Idx < BubbleInfoArray.Num(); ++Idx)
+	{
+		const FMassClientBubbleInfoData& BubbleData = BubbleInfoArray[Idx];
+
+		if (BubbleData.BubbleClass == BubbleInfoClass)
+		{
+			Handle.SetIndex(Idx);
+			break;
+		}
+	}
+	return Handle;
+}
+
+#if UE_REPLICATION_COMPILE_CLIENT_CODE
+void UMassReplicationSubsystem::SetEntity(const FMassNetworkID NetworkID, const FMassEntityHandle Entity)
+{
+	FMassReplicationEntityInfo* EntityInfo = FindMassEntityInfoMutable(NetworkID);
+	check(EntityInfo && !EntityInfo->Entity.IsSet());
+	EntityInfo->Entity = Entity;
+
+	OnMassAgentAdded.Broadcast(NetworkID, Entity);
+}
+#endif // UE_REPLICATION_COMPILE_CLIENT_CODE
+
+#if UE_REPLICATION_COMPILE_CLIENT_CODE
+FMassEntityHandle UMassReplicationSubsystem::ResetEntityIfValid(const FMassNetworkID NetworkID, int32 ReplicationID)
+{
+	FMassEntityHandle EntityReset;
+
+	FMassReplicationEntityInfo* EntityInfo = FindMassEntityInfoMutable(NetworkID);
+
+	checkf(EntityInfo, TEXT("EntityData must have been added to EntityInfoMap!"));
+	checkf(EntityInfo->ReplicationID >= ReplicationID, TEXT("We must not be removing an item we've never added!"));
+
+	// Only reset the item if its currently Set / Valid and its the most recent ReplicationID. Stale removes after more recent adds are ignored
+	// We do need to check the ReplicationID in this case
+	if (EntityInfo->Entity.IsSet() && (EntityInfo->ReplicationID == ReplicationID))
+	{
+		OnRemovingMassAgent.Broadcast(NetworkID, EntityReset);
+
+		EntityReset = EntityInfo->Entity;
+
+		//Unset the Entity handle, this indicates that its currently removed from the bubble
+		EntityInfo->Entity = FMassEntityHandle();
+	}
+
+	return EntityReset;
+}
+#endif // UE_REPLICATION_COMPILE_CLIENT_CODE
+
+#if UE_REPLICATION_COMPILE_CLIENT_CODE
+void UMassReplicationSubsystem::RemoveFromEntityInfoMap(const FMassNetworkID NetworkID)
+{
+	check(NetworkID.IsValid());
+	FMassReplicationEntityInfo Info;
+	ensureMsgf(EntityInfoMap.RemoveAndCopyValue(NetworkID, Info), TEXT("Removing a non-existant NetworkID(%s)!"), *NetworkID.Describe());
+}
+#endif // UE_REPLICATION_COMPILE_CLIENT_CODE
+
+#if UE_REPLICATION_COMPILE_CLIENT_CODE
+UMassReplicationSubsystem::EFindOrAddMassEntityInfo UMassReplicationSubsystem::FindAndUpdateOrAddMassEntityInfo(const FMassNetworkID NetworkID, int32 ReplicationID, const FMassReplicationEntityInfo*& OutMassEntityInfo)
+{
+	FMassReplicationEntityInfo* MassEntityInfo = FindMassEntityInfoMutable(NetworkID);
+	EFindOrAddMassEntityInfo FindOrAddStatus = EFindOrAddMassEntityInfo::FoundOlderReplicationID;
+
+	if (MassEntityInfo)
+	{
+		// Currently we don't think this should be needed, but are leaving it in for bomb proofing
+		if (ensure(MassEntityInfo->ReplicationID < ReplicationID))
+		{
+			// Update the replication ID to the latest
+			MassEntityInfo->ReplicationID = ReplicationID;
+		}
+		else
+		{
+			FindOrAddStatus = EFindOrAddMassEntityInfo::FoundNewerReplicationID;
+		}
+	}
+	else
+	{
+		MassEntityInfo = &EntityInfoMap.Add(NetworkID, FMassReplicationEntityInfo(FMassEntityHandle(), ReplicationID));
+		FindOrAddStatus = EFindOrAddMassEntityInfo::Added;
+	}
+
+	OutMassEntityInfo = MassEntityInfo;
+	return FindOrAddStatus;
+}
+#endif // UE_REPLICATION_COMPILE_CLIENT_CODE
+
+#if UE_REPLICATION_COMPILE_CLIENT_CODE
+const FMassReplicationEntityInfo* UMassReplicationSubsystem::FindMassEntityInfo(const FMassNetworkID NetworkID) const
+{
+	check(NetworkID.IsValid());
+	return EntityInfoMap.Find(NetworkID);
+}
+#endif // UE_REPLICATION_COMPILE_CLIENT_CODE
+
+#if UE_REPLICATION_COMPILE_CLIENT_CODE
+FMassReplicationEntityInfo* UMassReplicationSubsystem::FindMassEntityInfoMutable(const FMassNetworkID NetworkID)
+{
+	check(NetworkID.IsValid());
+	return EntityInfoMap.Find(NetworkID);
+}
+#endif // UE_REPLICATION_COMPILE_CLIENT_CODE
+
+#if UE_REPLICATION_COMPILE_CLIENT_CODE
+FMassEntityHandle UMassReplicationSubsystem::FindEntity(const FMassNetworkID NetworkID) const
+{
+	check(NetworkID.IsValid());
+	const FMassReplicationEntityInfo* Info = EntityInfoMap.Find(NetworkID);
+	return Info ? Info->Entity : FMassEntityHandle();
+}
+#endif // UE_REPLICATION_COMPILE_CLIENT_CODE
+
+void UMassReplicationSubsystem::DebugCheckArraysAreInSync()
+{
+#if UE_DEBUG_REPLICATION
+
+	checkf((ClientToViewerHandleArray.Num() == ClientsViewerHandles.Num()), TEXT("Client arrays out of sync with each other!"));
+
+	const int32 NumEntries = (BubbleInfoArray.Num() > 0) ? BubbleInfoArray[0].Bubbles.Num() : 0;
+
+	for (int32 IdxOuter = 0; IdxOuter < BubbleInfoArray.Num(); ++IdxOuter)
+	{
+		FMassClientBubbleInfoData& InfoDataOuter = BubbleInfoArray[IdxOuter];
+
+		checkf((InfoDataOuter.Bubbles.Num() == ClientsViewerHandles.Num()), TEXT("BubbleInfoArray arrays out of sync with ClientsViewerHandles!"));
+		checkf(InfoDataOuter.Bubbles.Num() == NumEntries, TEXT("Bubbles have different numbers of items!"));
+
+		for (int32 IdxInner = IdxOuter + 1; IdxInner < BubbleInfoArray.Num(); ++IdxInner)
+		{
+			FMassClientBubbleInfoData& InfoDataInner = BubbleInfoArray[IdxInner];
+
+			for (int32 IdxBubble = 0; IdxBubble < InfoDataOuter.Bubbles.Num(); ++IdxBubble)
+			{
+				AMassClientBubbleInfoBase* InfoOuter = InfoDataOuter.Bubbles[IdxBubble];
+				AMassClientBubbleInfoBase* InfoInner = InfoDataInner.Bubbles[IdxBubble];
+
+				const TArray<FMassClientHandle>& Handles = ClientHandleManager.GetHandles();
+
+				if (ClientHandleManager.IsValidHandle(Handles[IdxBubble]))
+				{				
+					check(InfoOuter && InfoInner);
+
+					APlayerController* OuterController = Cast<APlayerController>(InfoOuter->GetOwner());
+					APlayerController* InnerController = Cast<APlayerController>(InfoInner->GetOwner());
+
+					checkf((OuterController != nullptr) && (InnerController != nullptr), TEXT("Controller owners must be valid in BubbleInfoArray"));
+					checkf(OuterController == InnerController, TEXT("Owner controllers at the same indices in different BubbleInfoArray items must be equal!"));
+				}
+				else
+				{
+					check(!InfoOuter && !InfoInner);
+				}
+			}
+		}
+	}
+
+ #endif // UE_DEBUG_REPLICATION
+};
+
+void UMassReplicationSubsystem::AddClient(FMassViewerHandle ViewerHandle, APlayerController& InController)
+{
+	check(World);
+	check(MassLODManager);
+	checkf(MassLODManager->IsValidViewer(ViewerHandle), TEXT("ViewerHandle must be valid"));
+
+#if !UE_ALLOW_DEBUG_REPLICATION_BUBBLES_STANDALONE
+	checkf(UE::Mass::Replication::HasParentNetConnection(&InController), TEXT("InController must have a parent net connection or replication will not occur!"));
+#endif
+
+	if (!ensureMsgf(BubbleInfoArray.Num() > 0, TEXT("BubbleInfoClass has not been set Client will not be added to UMassReplicationSubsystem!")))
+	{
+		return;
+	}
+
+	if (ViewerHandle.GetIndex() >= ViewerToClientHandleArray.Num())
+	{
+		// making no assumptions about the order that ViewerHandles are added, make sure we add enough space to accomodate the new ViewerHandle index
+		ViewerToClientHandleArray.AddDefaulted(ViewerHandle.GetIndex() - ViewerToClientHandleArray.Num() + 1);
+	}
+	else 
+	{
+		checkf(ViewerToClientHandleArray[ViewerHandle.GetIndex()].ViewerHandle.IsValid() == false, TEXT("Adding a Client without removing the previous with the same Index!"));
+	}
+
+	FActorSpawnParameters SpawnParams;
+	SpawnParams.Owner = &InController;
+
+	const FMassClientHandle ClientHandle = ClientHandleManager.GetNextHandle();
+
+	const FViewerClientPair ViewerClientPair(ViewerHandle, ClientHandle);
+
+	FViewerClientPair& ViewerToClientHandleItem = ViewerToClientHandleArray[ViewerHandle.GetIndex()];
+
+	checkf(ViewerToClientHandleItem.IsValid() == false, TEXT("Handles should have been reset when previous Viewer in this slot was removed"));
+
+	ViewerToClientHandleItem = ViewerClientPair;
+
+	checkf(ClientHandle.GetIndex() <= ClientToViewerHandleArray.Num(), TEXT("ClientHandle out of sync with ClientToViewerHandleArray"));
+
+	// check if the handle is a new entry in the free list arrays or uses an existing entry
+	if (ClientHandle.GetIndex() == ClientToViewerHandleArray.Num())
+	{
+		ClientsViewerHandles.AddDefaulted();
+		ClientToViewerHandleArray.Emplace(ViewerHandle, ClientHandle);
+	}
+	else
+	{
+		checkf(ClientsViewerHandles[ClientHandle.GetIndex()].Handles.Num() == 0, TEXT("ClientsViewerHandles being replaced must have zero entries, they should have been removed first!"));
+
+		FViewerClientPair& ClientToViewerHandleItem = ClientToViewerHandleArray[ClientHandle.GetIndex()];
+
+		checkf(ClientToViewerHandleItem.IsValid() == false, TEXT("Handles should have been reset when previous Client in this slot was removed"));
+
+		ClientToViewerHandleItem = ViewerClientPair;
+	}
+
+	FClientViewerHandles& ClientViewers = ClientsViewerHandles[ClientHandle.GetIndex()];
+	ClientViewers.Handles.Add(ViewerHandle);
+
+	for (FMassClientBubbleInfoData& InfoData : BubbleInfoArray)
+	{
+		AMassClientBubbleInfoBase* ClientBubbleInfo = World->SpawnActor<AMassClientBubbleInfoBase>(InfoData.BubbleClass, SpawnParams);
+		ClientBubbleInfo->SetClientHandle(ClientHandle);
+
+		checkf(ClientHandle.GetIndex() <= InfoData.Bubbles.Num(), TEXT("ClientHandle out of sync with Bubbles"));
+
+		// check if the handle is a new entry in the free list arrays or uses an existing entry
+		if (ClientHandle.GetIndex() == InfoData.Bubbles.Num())
+		{
+			InfoData.Bubbles.Push(ClientBubbleInfo);
+		}
+		else
+		{
+			AMassClientBubbleInfoBase*& BubbleUpdate = InfoData.Bubbles[ClientHandle.GetIndex()];
+			checkf(BubbleUpdate == nullptr, TEXT("ClientBubble being replaced must be nullptr it should have been removed first!"));
+
+			BubbleUpdate = ClientBubbleInfo;
+		}
+	}
+	DebugCheckArraysAreInSync();
+}
+
+void UMassReplicationSubsystem::RemoveClient(FMassClientHandle ClientHandle)
+{
+	check(World);
+
+	checkf(ClientHandleManager.IsValidHandle(ClientHandle), TEXT("ClientHandle must be a valid non stale handle"));
+
+	checkf(ClientToViewerHandleArray.IsValidIndex(ClientHandle.GetIndex()), TEXT("ClientHandle is out of sync with ClientToViewerHandleArray!"));
+	FViewerClientPair& ClientToViewerHandleItem = ClientToViewerHandleArray[ClientHandle.GetIndex()];
+
+	checkf(ClientToViewerHandleItem.ViewerHandle.IsValid(), TEXT("Invalid ViewerHandle! ClientHandle is out of sync with ClientToViewerHandleArray!"));
+	checkf(ClientToViewerHandleItem.ClientHandle == ClientHandle, TEXT("ClientHandle is out of sync with ClientToViewerHandleArray!"));
+
+	FClientViewerHandles& ClientViewerHandles = ClientsViewerHandles[ClientHandle.GetIndex()];
+
+	checkf(ClientViewerHandles.Handles.Num() > 0, TEXT("There should always be atleast one client viewer handle (the parent NetConnection)"));
+	ClientViewerHandles.Handles.Reset();
+
+	checkf(ViewerToClientHandleArray.IsValidIndex(ClientToViewerHandleItem.ViewerHandle.GetIndex()), TEXT("ViewerHandle is out of sync with ViewerToClientHandleArray!"));
+	FViewerClientPair& ViewerToClientHandleItem = ViewerToClientHandleArray[ClientToViewerHandleItem.ViewerHandle.GetIndex()];
+
+	checkf(ViewerToClientHandleItem == ClientToViewerHandleItem, TEXT("ClientToViewerHandleArray and ViewerToClientHandleArray out of sync!"));
+
+	ViewerToClientHandleItem.Invalidate();
+	ClientToViewerHandleItem.Invalidate();
+
+	for (FMassClientBubbleInfoData& InfoData : BubbleInfoArray)
+	{
+		checkf(InfoData.Bubbles.IsValidIndex(ClientHandle.GetIndex()), TEXT("ClientHandle is out of sync with Bubbles!"));
+		AMassClientBubbleInfoBase*& BubbleInfoItem = InfoData.Bubbles[ClientHandle.GetIndex()];
+
+		if ((World != nullptr) && (BubbleInfoItem != nullptr))
+		{
+			World->DestroyActor(BubbleInfoItem);
+		}
+
+		BubbleInfoItem = nullptr;
+	}
+
+	ClientHandleManager.RemoveHandle(ClientHandle);
+
+	DebugCheckArraysAreInSync();
+}
diff --git a/Engine/Plugins/Runtime/MassGameplay/Source/MassReplication/Public/MassClientBubbleHandler.h b/Engine/Plugins/Runtime/MassGameplay/Source/MassReplication/Public/MassClientBubbleHandler.h
--- a/Engine/Plugins/Runtime/MassGameplay/Source/MassReplication/Public/MassClientBubbleHandler.h
+++ b/Engine/Plugins/Runtime/MassGameplay/Source/MassReplication/Public/MassClientBubbleHandler.h
@@ -10,7 +10,7 @@
 #include "MassSpawnerSubsystem.h"
 #include "MassSpawnerTypes.h"
 #include "MassReplicationFragments.h"
-#include "MassReplicationManager.h"
+#include "MassReplicationSubsystem.h"
 #include "Engine/World.h"
 #include "MassEntityTemplate.h"
 #include "MassEntityView.h"
@@ -354,8 +354,8 @@ void TClientBubbleHandlerBase<AgentArrayItem>::UpdateAgentsToRemove()
 
 	check(Serializer->GetWorld());
 
-	UMassReplicationManager* ReplicationManager = Serializer->GetReplicationManager();
-	check(ReplicationManager);
+	UMassReplicationSubsystem* ReplicationSubsystem = Serializer->GetReplicationSubsystem();
+	check(ReplicationSubsystem);
 
 	const float TimeRemove = Serializer->GetWorld()->GetRealTimeSeconds() - UE::Mass::Replication::AgentRemoveInterval;
 
@@ -371,7 +371,7 @@ void TClientBubbleHandlerBase<AgentArrayItem>::UpdateAgentsToRemove()
 		{
 			const FMassNetworkID NetID = (*Iter).Key;
 
-			ReplicationManager->RemoveFromEntityInfoMap(NetID);
+			ReplicationSubsystem->RemoveFromEntityInfoMap(NetID);
 			Iter.RemoveCurrent();
 		}
 	}
@@ -390,14 +390,14 @@ void TClientBubbleHandlerBase<AgentArrayItem>::PreReplicatedRemove(const TArrayV
 	UMassSpawnerSubsystem* SpawnerSubsystem = Serializer->GetSpawnerSubsystem();
 	check(SpawnerSubsystem);
 
-	UMassReplicationManager* ReplicationManager = Serializer->GetReplicationManager();
-	check(ReplicationManager);
+	UMassReplicationSubsystem* ReplicationSubsystem = Serializer->GetReplicationSubsystem();
+	check(ReplicationSubsystem);
 
 	for (int32 Idx : RemovedIndices)
 	{
 		const AgentArrayItem& RemovedItem = (*Agents)[Idx];
 
-		FMassEntityHandle Entity = ReplicationManager->ResetEntityIfValid(RemovedItem.Agent.GetNetID(), RemovedItem.ReplicationID);
+		FMassEntityHandle Entity = ReplicationSubsystem->ResetEntityIfValid(RemovedItem.Agent.GetNetID(), RemovedItem.ReplicationID);
 
 		// Only remove the item if its currently Set / Valid and its the most recent ReplicationID. Stale removes after more recent adds are ignored
 		// We do need to check the ReplicationID in this case
@@ -431,8 +431,8 @@ void TClientBubbleHandlerBase<AgentArrayItem>::PostReplicatedAddHelper(const TAr
 	UMassEntitySubsystem* EntitySystem = Serializer->GetEntitySystem();
 	check(EntitySystem);
 
-	UMassReplicationManager* ReplicationManager = Serializer->GetReplicationManager();
-	check(ReplicationManager);
+	UMassReplicationSubsystem* ReplicationSubsystem = Serializer->GetReplicationSubsystem();
+	check(ReplicationSubsystem);
 
 	TMap<FMassNetworkID, int32> AgentsToAddMap; //NetID to index in AgentsToAddArray
 	AgentsToAddMap.Reserve(AddedIndices.Num());
@@ -446,9 +446,9 @@ void TClientBubbleHandlerBase<AgentArrayItem>::PostReplicatedAddHelper(const TAr
 	{
 		const AgentArrayItem& AddedItem = (*Agents)[Idx];
 
-		switch (ReplicationManager->FindAndUpdateOrAddMassEntityInfo(AddedItem.Agent.GetNetID(), AddedItem.ReplicationID, EntityInfo))
+		switch (ReplicationSubsystem->FindAndUpdateOrAddMassEntityInfo(AddedItem.Agent.GetNetID(), AddedItem.ReplicationID, EntityInfo))
 		{
-		case UMassReplicationManager::EFindOrAddMassEntityInfo::FoundOlderReplicationID:
+		case UMassReplicationSubsystem::EFindOrAddMassEntityInfo::FoundOlderReplicationID:
 		{
 			AgentsRemoveDataMap.Remove(AddedItem.Agent.GetNetID());
 
@@ -478,7 +478,7 @@ void TClientBubbleHandlerBase<AgentArrayItem>::PostReplicatedAddHelper(const TAr
 		}
 		break;
 
-		case UMassReplicationManager::EFindOrAddMassEntityInfo::Added:
+		case UMassReplicationSubsystem::EFindOrAddMassEntityInfo::Added:
 		{
 			check(AgentsToAddMap.Find(AddedItem.Agent.GetNetID()) == nullptr);
 
@@ -487,7 +487,7 @@ void TClientBubbleHandlerBase<AgentArrayItem>::PostReplicatedAddHelper(const TAr
 		}
 		break;
 
-		case UMassReplicationManager::EFindOrAddMassEntityInfo::FoundNewerReplicationID:
+		case UMassReplicationSubsystem::EFindOrAddMassEntityInfo::FoundNewerReplicationID:
 			break;
 
 		default:
@@ -513,8 +513,8 @@ void TClientBubbleHandlerBase<AgentArrayItem>::PostReplicatedAddEntitiesHelper(c
 	UMassSpawnerSubsystem* SpawnerSubsystem = Serializer->GetSpawnerSubsystem();
 	check(SpawnerSubsystem);
 
-	UMassReplicationManager* ReplicationManager = Serializer->GetReplicationManager();
-	check(ReplicationManager);
+	UMassReplicationSubsystem* ReplicationSubsystem = Serializer->GetReplicationSubsystem();
+	check(ReplicationSubsystem);
 
 	TMap<FMassEntityTemplateID, TArray<typename AgentArrayItem::FReplicatedAgentType*>> AgentsSpawnMap;
 
@@ -552,7 +552,7 @@ void TClientBubbleHandlerBase<AgentArrayItem>::PostReplicatedAddEntitiesHelper(c
 		int32 AgentsSpawnIdx = 0;
 
 		Query.ForEachEntityChunk(FMassArchetypeSubChunks(ArchetypeHandle, Entities, FMassArchetypeSubChunks::NoDuplicates)
-								, *EntitySystem, ExecContext, [&AgentsSpawn, &AgentsSpawnIdx, this, ReplicationManager, &ExecContext, &CacheFragmentViewsForSpawnQuery, &SetSpawnedEntityData, &EntitySystem](FMassExecutionContext& Context)
+								, *EntitySystem, ExecContext, [&AgentsSpawn, &AgentsSpawnIdx, this, ReplicationSubsystem, &ExecContext, &CacheFragmentViewsForSpawnQuery, &SetSpawnedEntityData, &EntitySystem](FMassExecutionContext& Context)
 			{
 				CacheFragmentViewsForSpawnQuery(ExecContext);
 
@@ -566,7 +566,7 @@ void TClientBubbleHandlerBase<AgentArrayItem>::PostReplicatedAddEntitiesHelper(c
 					FMassNetworkIDFragment& NetIDFragment = NetworkIDList[i];
 
 					NetIDFragment.NetID = AgentSpawn.GetNetID();
-					ReplicationManager->SetEntity(NetIDFragment.NetID, Entity);
+					ReplicationSubsystem->SetEntity(NetIDFragment.NetID, Entity);
 
 					FMassEntityView EntityView(*EntitySystem, Entity);
 					SetSpawnedEntityData(EntityView, AgentSpawn, i);
@@ -585,15 +585,15 @@ void TClientBubbleHandlerBase<AgentArrayItem>::PostReplicatedChangeHelper(const
 	UMassEntitySubsystem* EntitySystem = Serializer->GetEntitySystem();
 	check(EntitySystem);
 
-	UMassReplicationManager* ReplicationManager = Serializer->GetReplicationManager();
-	check(ReplicationManager);
+	UMassReplicationSubsystem* ReplicationSubsystem = Serializer->GetReplicationSubsystem();
+	check(ReplicationSubsystem);
 
 	// Go through the changed Entities and update their Mass data
 	for (int32 Idx : ChangedIndices)
 	{
 		const AgentArrayItem& ChangedItem = (*Agents)[Idx];
 
-		const FMassReplicationEntityInfo* EntityInfo = ReplicationManager->FindMassEntityInfo(ChangedItem.Agent.GetNetID());
+		const FMassReplicationEntityInfo* EntityInfo = ReplicationSubsystem->FindMassEntityInfo(ChangedItem.Agent.GetNetID());
 
 		checkf(EntityInfo, TEXT("EntityInfo must be valid if the Agent has already been added (which it must have been to get PostReplicatedChange"));
 		checkf(EntityInfo->ReplicationID >= ChangedItem.ReplicationID, TEXT("ReplicationID out of sync, this should never happen!"));
@@ -653,8 +653,8 @@ void TClientBubbleHandlerBase<AgentArrayItem>::DebugValidateBubbleOnClient()
 	UMassEntitySubsystem* EntitySystem = Serializer->GetEntitySystem();
 	check(EntitySystem);
 
-	UMassReplicationManager* ReplicationManager = Serializer->GetReplicationManager();
-	check(ReplicationManager);
+	UMassReplicationSubsystem* ReplicationSubsystem = Serializer->GetReplicationSubsystem();
+	check(ReplicationSubsystem);
 
 	for (int32 Idx = 0; Idx < (*Agents).Num(); ++Idx)
 	{
@@ -664,7 +664,7 @@ void TClientBubbleHandlerBase<AgentArrayItem>::DebugValidateBubbleOnClient()
 		check(Agent.GetTemplateID().IsValid());
 		check(Agent.GetNetID().IsValid());
 
-		const FMassReplicationEntityInfo* EntityInfo = ReplicationManager->FindMassEntityInfo(Agent.GetNetID());
+		const FMassReplicationEntityInfo* EntityInfo = ReplicationSubsystem->FindMassEntityInfo(Agent.GetNetID());
 		checkf(EntityInfo, TEXT("There should always be an EntityInfoMap entry for Agents that are in the Agents array!"));
 
 		if (EntityInfo)
@@ -685,7 +685,7 @@ void TClientBubbleHandlerBase<AgentArrayItem>::DebugValidateBubbleOnClient()
 	}
 
 #if UE_ALLOW_DEBUG_SLOW_REPLICATION
-	const TMap<FMassNetworkID, FMassReplicationEntityInfo>& EntityInfoMap = ReplicationManager->GetEntityInfoMap();
+	const TMap<FMassNetworkID, FMassReplicationEntityInfo>& EntityInfoMap = ReplicationSubsystem->GetEntityInfoMap();
 
 	for (TMap<FMassNetworkID, FMassReplicationEntityInfo>::TConstIterator Iter = EntityInfoMap.CreateConstIterator(); Iter; ++Iter)
 	{
diff --git a/Engine/Plugins/Runtime/MassGameplay/Source/MassReplication/Public/MassClientBubbleSerializerBase.h b/Engine/Plugins/Runtime/MassGameplay/Source/MassReplication/Public/MassClientBubbleSerializerBase.h
--- a/Engine/Plugins/Runtime/MassGameplay/Source/MassReplication/Public/MassClientBubbleSerializerBase.h
+++ b/Engine/Plugins/Runtime/MassGameplay/Source/MassReplication/Public/MassClientBubbleSerializerBase.h
@@ -34,7 +34,7 @@ public:
 
 	UWorld* GetWorld() const { return World; }
 	UMassSpawnerSubsystem* GetSpawnerSubsystem() const { return SpawnerSubsystem; }
-	UMassReplicationManager* GetReplicationManager() const { return ReplicationManager; }
+	UMassReplicationSubsystem* GetReplicationSubsystem() const { return ReplicationSubsystem; }
 	UMassEntitySubsystem* GetEntitySystem() const { return EntitySystem; }
 
 	void SetClientHandler(IClientBubbleHandlerInterface& InClientHandler) { ClientHandler = &InClientHandler; }
@@ -51,7 +51,7 @@ private:
 	UMassEntitySubsystem* EntitySystem = nullptr;
 
 	UPROPERTY(Transient)
-	UMassReplicationManager* ReplicationManager = nullptr;
+	UMassReplicationSubsystem* ReplicationSubsystem = nullptr;
 
 	/** Pointer to the IClientBubbleHandlerInterface derived class in the class derived from this one */
 	IClientBubbleHandlerInterface* ClientHandler = nullptr;
diff --git a/Engine/Plugins/Runtime/MassGameplay/Source/MassReplication/Public/MassReplicationFragments.h b/Engine/Plugins/Runtime/MassGameplay/Source/MassReplication/Public/MassReplicationFragments.h
--- a/Engine/Plugins/Runtime/MassGameplay/Source/MassReplication/Public/MassReplicationFragments.h
+++ b/Engine/Plugins/Runtime/MassGameplay/Source/MassReplication/Public/MassReplicationFragments.h
@@ -100,7 +100,7 @@ protected:
 	virtual void Execute(UMassEntitySubsystem& EntitySubsystem, FMassExecutionContext& Context) override;
 
 	UPROPERTY()
-		UMassReplicationManager* ReplicationManager = nullptr;
+	UMassReplicationSubsystem* ReplicationSubsystem = nullptr;
 
 	FMassEntityQuery EntityQuery;
 };
diff --git a/Engine/Plugins/Runtime/MassGameplay/Source/MassReplication/Public/MassReplicationManager.h b/Engine/Plugins/Runtime/MassGameplay/Source/MassReplication/Public/MassReplicationManager.h
deleted file mode 100644
diff --git a/Engine/Plugins/Runtime/MassGameplay/Source/MassReplication/Public/MassReplicationProcessorBase.h b/Engine/Plugins/Runtime/MassGameplay/Source/MassReplication/Public/MassReplicationProcessorBase.h
--- a/Engine/Plugins/Runtime/MassGameplay/Source/MassReplication/Public/MassReplicationProcessorBase.h
+++ b/Engine/Plugins/Runtime/MassGameplay/Source/MassReplication/Public/MassReplicationProcessorBase.h
@@ -11,10 +11,10 @@
 #include "MassReplicationFragments.h"
 #include "MassSpawnerTypes.h"
 #include "MassLODManager.h"
-#include "MassReplicationManager.h"
+#include "MassReplicationSubsystem.h"
 #include "MassReplicationProcessorBase.generated.h"
 
-class UMassReplicationManager;
+class UMassReplicationSubsystem;
 class AMassClientBubbleInfoBase;
 class UWorld;
 
@@ -50,7 +50,7 @@ protected:
 protected:
 
 	UPROPERTY()
-	UMassReplicationManager* ReplicationManager;
+	UMassReplicationSubsystem* ReplicationSubsystem;
 
 	FMassBubbleInfoClassHandle BubbleInfoClassHandle;
 
@@ -95,7 +95,7 @@ void UMassReplicationProcessorBase::CalculateClientReplication(FMassExecutionCon
 {
 #if UE_REPLICATION_COMPILE_SERVER_CODE
 
-	check(ReplicationManager);
+	check(ReplicationSubsystem);
 	check(LODManager);
 	check(World);
 
@@ -142,7 +142,7 @@ void UMassReplicationProcessorBase::CalculateClientReplication(FMassExecutionCon
 
 				//now we need to see what our highest viewer LOD is on this client (split screen etc), we can use the unsafe version as we have already checked all the
 				//CachedClientHandles for validity
-				const FClientViewerHandles& ClientViewers = ReplicationManager->GetClientViewersChecked(ClientHandle);
+				const FClientViewerHandles& ClientViewers = ReplicationSubsystem->GetClientViewersChecked(ClientHandle);
 
 				EMassLOD::Type HighestLOD = EMassLOD::Off;
 
diff --git a/Engine/Plugins/Runtime/MassGameplay/Source/MassReplication/Public/MassReplicationSubsystem.h b/Engine/Plugins/Runtime/MassGameplay/Source/MassReplication/Public/MassReplicationSubsystem.h
new file mode 100644
--- /dev/null
+++ b/Engine/Plugins/Runtime/MassGameplay/Source/MassReplication/Public/MassReplicationSubsystem.h
@@ -0,0 +1,294 @@
+// Copyright Epic Games, Inc. All Rights Reserved.
+
+#pragma once
+
+#include "MassCommonTypes.h"
+#include "MassEntityTypes.h"
+#include "MassReplicationTypes.h"
+#include "Subsystems/WorldSubsystem.h"
+
+#include "MassReplicationSubsystem.generated.h"
+
+class UMassEntitySubsystem;
+class AMassCrowdClientBubbleInfo;
+class UMassLODManager;
+class AMassClientBubbleInfoBase;
+
+struct FClientViewerHandles
+{
+	TArray<FMassViewerHandle> Handles;
+};
+
+struct FViewerClientPair
+{
+	FViewerClientPair() = default;
+
+	FViewerClientPair(FMassViewerHandle InViewerHandle, FMassClientHandle InClientHandle)
+		: ViewerHandle(InViewerHandle)
+		, ClientHandle(InClientHandle)
+	{}
+
+	bool operator== (FViewerClientPair Other) const { return ViewerHandle == Other.ViewerHandle && ClientHandle == Other.ClientHandle; }
+
+	bool IsValid() const { return ViewerHandle.IsValid() && ClientHandle.IsValid(); }
+	void Invalidate() { ViewerHandle.Invalidate(); ClientHandle.Invalidate(); }
+
+	FMassViewerHandle ViewerHandle;
+	FMassClientHandle ClientHandle;
+};
+
+USTRUCT()
+struct FMassClientBubbleInfoData
+{
+	GENERATED_BODY()
+
+	FMassClientBubbleInfoData() = default;
+
+	FMassClientBubbleInfoData(TSubclassOf<AMassClientBubbleInfoBase> InBubbleClass)
+		: BubbleClass(InBubbleClass)
+	{}
+
+	/** A free list array of AMassClientBubbleInfos. This is organised so the index is that of the client FMassClientHandle */
+	UPROPERTY(Transient)
+	TArray<AMassClientBubbleInfoBase*> Bubbles;
+
+	UPROPERTY(Transient)
+	TSubclassOf<AMassClientBubbleInfoBase> BubbleClass;
+};
+
+namespace UE::MassReplication
+{
+	DECLARE_MULTICAST_DELEGATE_TwoParams(FMassAgentDelegate, FMassNetworkID, FMassEntityHandle);
+
+} // UE::MassReplication
+
+
+/**
+ *  Manages the creation of NetworkIDs, ClientBubbles and ClientViewers.
+ *  NetworkIDs are per replicated Agent Entity and are unique and replicated between server and clients.
+ *  ClientBubbles relate to the player controller that owns the parent UNetConnection to a Client machine.
+ *  ClientViewers relate to all the player controllers that have a parent or child UNetConnection to a single Client machine (split screen etc).
+ */
+UCLASS()
+class MASSREPLICATION_API UMassReplicationSubsystem : public UWorldSubsystem
+{
+	GENERATED_BODY()
+
+public:
+
+	enum class EFindOrAddMassEntityInfo : uint8
+	{
+		FoundOlderReplicationID,
+		FoundNewerReplicationID,
+		Added,
+	};
+
+	bool IsBubbleClassHandleValid(FMassBubbleInfoClassHandle Handle) const
+	{
+		return BubbleInfoArray.IsValidIndex(Handle.GetIndex());
+	}
+
+	/** Retrieve network id from Mass handle. */
+	FMassNetworkID GetNetIDFromHandle(const FMassEntityHandle Handle) const;
+
+	static inline FMassNetworkID GetNextAvailableMassNetID() { return FMassNetworkID(++CurrentNetMassCounter); }
+
+	const TArray<FMassClientHandle>& GetClientReplicationHandles() const { return ClientHandleManager.GetHandles(); }
+
+	/** Gets the client bubble safely */
+	AMassClientBubbleInfoBase* GetClientBubble(FMassBubbleInfoClassHandle BubbleClassHandle, FMassClientHandle ClientHandle) const
+	{
+		AMassClientBubbleInfoBase* ClientBubble = nullptr;
+
+		if (IsBubbleClassHandleValid(BubbleClassHandle) && ClientHandleManager.IsValidHandle(ClientHandle))
+		{
+			ClientBubble = BubbleInfoArray[BubbleClassHandle.GetIndex()].Bubbles[ClientHandle.GetIndex()];
+			check(ClientBubble);
+		}
+		return ClientBubble;
+	}
+
+	/** Get the client bubble. Faster version using check()s */
+	AMassClientBubbleInfoBase* GetClientBubbleChecked(FMassBubbleInfoClassHandle BubbleClassHandle, FMassClientHandle ClientHandle) const
+	{
+		check(IsBubbleClassHandleValid(BubbleClassHandle));
+		check(ClientHandleManager.IsValidHandle(ClientHandle));
+		AMassClientBubbleInfoBase* ClientBubble = BubbleInfoArray[BubbleClassHandle.GetIndex()].Bubbles[ClientHandle.GetIndex()];
+
+		check(ClientBubble);
+		return ClientBubble;
+	}
+
+	/** Gets the client bubble safely returning the template type */
+	template<typename TType>
+	TType* GetTypedClientBubble(FMassBubbleInfoClassHandle BubbleClassHandle, FMassClientHandle ClientHandle) const
+	{
+		return Cast<TType>(GetClientBubble(BubbleClassHandle, ClientHandle));
+	}
+
+	/** Gets the client bubble returning the template type. Faster version using check()s */
+	template<typename TType>
+	TType* GetTypedClientBubbleChecked(FMassBubbleInfoClassHandle BubbleClassHandle, FMassClientHandle ClientHandle) const
+	{
+		AMassClientBubbleInfoBase* ClientBubble = GetClientBubbleChecked(BubbleClassHandle, ClientHandle);
+		checkSlow(Cast<TType>(ClientBubble) != nullptr);
+
+		return static_cast<TType *>(ClientBubble);
+	}
+
+	/** Gets the client viewers safely */
+	const FClientViewerHandles* GetClientViewers(FMassClientHandle Handle) const
+	{
+		return ClientHandleManager.IsValidHandle(Handle) ? &(ClientsViewerHandles[Handle.GetIndex()]) : nullptr;
+	}
+
+	/** Gets the client viewers. Faster version using check()s */
+	const FClientViewerHandles& GetClientViewersChecked(FMassClientHandle Handle) const
+	{
+		check(ClientHandleManager.IsValidHandle(Handle));
+
+		return ClientsViewerHandles[Handle.GetIndex()];
+	}
+
+	void SynchronizeClientsAndViewers();
+
+	/** 
+	 * Registers BubbleInfoClass with the system. These are created one per client as clients join. This function can be called multiple times with
+	 * different BubbleInfoClass types (but not the same type), in that case each client will get multiple BubbleInfoClasses. This can be useful for 
+	 * replicating different Entity types in different bubbles, although its also possible to have multiple TClientBubbleHandlerBase derived class
+	 * instances per BubbleInfoClass. This must not be called after AddClient or SynchronizeClientsAndViewers is called
+	 * @return FMassBubbleInfoClassHandle Handle to the BubbleInfoClass, this will be an invalid handle if its not been created
+	 */
+	FMassBubbleInfoClassHandle RegisterBubbleInfoClass(const TSubclassOf<AMassClientBubbleInfoBase>& BubbleInfoClass);
+
+	/** @return FMassBubbleInfoClassHandle Handle to the BubbleInfoClass, this will be an invalid handle if BubbleInfoClass can not be found */
+	FMassBubbleInfoClassHandle GetBubbleInfoClassHandle(const TSubclassOf<AMassClientBubbleInfoBase>& BubbleInfoClass) const;
+
+#if UE_REPLICATION_COMPILE_CLIENT_CODE
+
+	/** 
+	 * Must be called immediately after an Entity is added to the simulation and after FindAndUpdateOrAddMassEntityInfo has been called
+	 */
+	void SetEntity(const FMassNetworkID NetworkID, const FMassEntityHandle Entity);
+
+	/** 
+	 * Resets the item with NetworkID in the  EntityInfoMap if its currently Set / Valid and its the most recent ReplicationID. In this case the FMassEntityHandle that was reset is returned.
+	 * Otherwise FMassEntityHandle::IsSet() will be false on the returned FMassEntityHandle
+	 * Must be called after SetEntity and just before the Entity is removed on the client. This sets the Entity member of the AddToEntityInfoMap item to FMassEntityHandle so it will be not IsSet() 
+	 */
+	FMassEntityHandle ResetEntityIfValid(const FMassNetworkID NetworkID, int32 ReplicationID);
+
+	void RemoveFromEntityInfoMap(const FMassNetworkID NetworkID);
+
+	/** 
+	 * Finds or adds an FMassReplicationEntityInfo to the EntityInfoMap for an Entity with NetWorkID. This will update the ReplicationID if an older one is found but it does not change the 
+	 * FMassReplicationEntityInfo::Entity member. Returns whether we added an entry in EntityInfoMap, or if the found item has an older or newer ReplicationID than that that is passed in. 
+	 * We shouldn't actually get an equal ReplicationID, but in this case we'll treat the existing EntityInfoMap entry as newer. SetEntity() must be called soon after this if the entity handle is !IsSet()
+	 * @param NetworkID ID of the entity we are looking up OutMassEntityInfo for
+	 * @param ReplicationID this is the ReplicationID of the FastArray item
+	 * @param OutMassEntityInfo retrieved or added FMassReplicationEntityInfo
+	 * @return EFindOrAddMassEntityInfo, result of the function call
+	 */
+	EFindOrAddMassEntityInfo FindAndUpdateOrAddMassEntityInfo(const FMassNetworkID NetworkID, int32 ReplicationID, const FMassReplicationEntityInfo*& OutMassEntityInfo);
+
+	/** 
+	 * Finds the FMassReplicationEntityInfo for the Entity with NetworkID. NetworkID must be valid and nullptr is returned if we can't find it. Note FMassReplicationEntityInfo::IsSet() may be false,
+	 * in which case this Entity has been removed from the client simulation (but we are still storing the ReplicationID between Removes and Adds!
+	 */
+	const FMassReplicationEntityInfo* FindMassEntityInfo(const FMassNetworkID NetworkID) const;
+
+	FMassEntityHandle FindEntity(const FMassNetworkID NetworkID) const;
+
+	const TMap<FMassNetworkID, FMassReplicationEntityInfo>& GetEntityInfoMap() const { return EntityInfoMap; }
+
+	UE::MassReplication::FMassAgentDelegate& GetOnMassAgentAdded()
+	{
+		return OnMassAgentAdded;
+	}
+
+	UE::MassReplication::FMassAgentDelegate& GetOnRemovingMassAgent()
+	{
+		return OnRemovingMassAgent;
+	}
+
+#endif // UE_REPLICATION_COMPILE_CLIENT_CODE
+
+protected:
+	// USubsystem BEGIN
+	virtual void Initialize(FSubsystemCollectionBase& Collection) override;
+	virtual void Deinitialize() override;
+	// USubsystem END
+	
+#if UE_REPLICATION_COMPILE_CLIENT_CODE
+	FMassReplicationEntityInfo* FindMassEntityInfoMutable(const FMassNetworkID NetworkID);
+#endif // UE_REPLICATION_COMPILE_CLIENT_CODE
+
+	/** Adds a Client and sets up all the data relevant to the bubble. Where there are multiple viewers for
+	 *  one NetConnection only the parent net connection should be added here
+	 *  @param ViewerHandle handle to viewer to add as a client. This must be UMassLODManager->IsValidViewer()
+	 *  @param InController associated with the viewer
+	 */
+	void AddClient(FMassViewerHandle ViewerHandle, APlayerController& InController);
+
+	/** Removes a Client and removes up all the data relevant to the bubble
+	 *  @param ClientHandle handle to client to remove. This must be ClientHandleManager.IsValidHandle()
+	 */
+	void RemoveClient(FMassClientHandle ClientHandle);
+
+	/** @return true if we should shrink the number of handles. */
+	bool SynchronizeClients(const TArray<FViewerInfo>& Viewers);
+	void SynchronizeClientViewers(const TArray<FViewerInfo>& Viewers);
+
+	static void DebugCheckArraysAreInSync();
+
+protected:
+	static uint32 CurrentNetMassCounter;
+
+	UPROPERTY()
+	UMassEntitySubsystem* EntitySystem;
+
+	/** Clients free list FMassClientHandle manager, handles will to the indices of FMassClientReplicationData::ClientBubbles */
+	FMassClientHandleManager ClientHandleManager;
+
+	UPROPERTY(Transient)
+	TArray<FMassClientBubbleInfoData> BubbleInfoArray;
+
+	/** An Array of each Clients viewer handles (split screen players sharing the same client connections). 
+	 *  This will include both the parent and child NetConnections per client.
+	 *	The array is organised so the array index is the same as the index of the client FMassClientHandle.
+	 */
+	TArray<FClientViewerHandles> ClientsViewerHandles;
+
+	/** A free list array of FViewerClientPairs. This is organised so the index is that of the FMassViewerHandle for fast lookup of the related FMassClientHandle. 
+	 *  This only contains viewers that are also clients (ie only parent NetConnections not child ones).
+	 */
+	TArray<FViewerClientPair> ViewerToClientHandleArray;
+
+	/** A free list array of FViewerClientPairs. This is organised so the index is that of the FMassClientHandle.
+	 *  For fast lookup of the related FMassViewerHandle. This only contains viewers that are also clients (ie only parent NetConnections not child ones).
+	 */
+	TArray<FViewerClientPair> ClientToViewerHandleArray;
+
+	UPROPERTY()
+	UWorld* World;
+
+	UPROPERTY()
+	UMassLODManager* MassLODManager;
+
+#if UE_REPLICATION_COMPILE_CLIENT_CODE
+	TMap<FMassNetworkID, FMassReplicationEntityInfo> EntityInfoMap;
+
+	// @todo MassReplication consider batching these
+	/** Broadcast just after a mass agent is added to the client simulation */
+	UE::MassReplication::FMassAgentDelegate OnMassAgentAdded;
+
+	/** Broadcast just before a mass agent is removed from the client simulation */
+	UE::MassReplication::FMassAgentDelegate OnRemovingMassAgent;
+#endif // UE_REPLICATION_COMPILE_CLIENT_CODE
+
+	/**
+	 * Used to make sure the FMassEntityHandle_ClientReplications are synchronized immediately before they are needed
+	 * @todo this comment is no longer accurate, needs fixing
+	 */
+	uint64 LastSynchronizedFrame = 0;
+};
\ No newline at end of file
diff --git a/Engine/Plugins/Runtime/MassGameplay/Source/MassReplication/Public/MassReplicationTypes.h b/Engine/Plugins/Runtime/MassGameplay/Source/MassReplication/Public/MassReplicationTypes.h
--- a/Engine/Plugins/Runtime/MassGameplay/Source/MassReplication/Public/MassReplicationTypes.h
+++ b/Engine/Plugins/Runtime/MassGameplay/Source/MassReplication/Public/MassReplicationTypes.h
@@ -60,7 +60,7 @@
 #define UE_ALLOW_DEBUG_REPLICATION 0
 #endif
 
-class UMassReplicationManager;
+class UMassReplicationSubsystem;
 
 namespace UE { namespace Mass { namespace Replication
 {
@@ -166,7 +166,7 @@ struct FMassReplicationEntityInfo
 	/** If this is not IsSet() then the entity has been removed from the client simulation */
 	FMassEntityHandle Entity; 
 
-	/** This is stored between removes and adds, however this item in the UMassReplicationManager::EntityInfoMap will eventually get cleaned up if Entity.IsSet() == false for a fairly substantial length of time. */
+	/** This is stored between removes and adds, however this item in the UMassReplicationSubsystem::EntityInfoMap will eventually get cleaned up if Entity.IsSet() == false for a fairly substantial length of time. */
 	int32 ReplicationID = INDEX_NONE;
 };
 #endif //UE_REPLICATION_COMPILE_CLIENT_CODE
\ No newline at end of file
