Fixed some uninitialized struct members
#rnx

diff --git a/Engine/Plugins/Runtime/MassGameplay/Source/MassRepresentation/Public/MassRepresentationFragments.h b/Engine/Plugins/Runtime/MassGameplay/Source/MassRepresentation/Public/MassRepresentationFragments.h
--- a/Engine/Plugins/Runtime/MassGameplay/Source/MassRepresentation/Public/MassRepresentationFragments.h
+++ b/Engine/Plugins/Runtime/MassGameplay/Source/MassRepresentation/Public/MassRepresentationFragments.h
@@ -75,7 +75,7 @@ struct FMassRepresentationConfig : public FMassSharedFragment
 
 	/** What should be the representation of this entity for each specific LOD */
 	UPROPERTY(EditAnywhere, Category = "Mass|Representation", config)
-	EMassRepresentationType LODRepresentation[EMassLOD::Max];
+	EMassRepresentationType LODRepresentation[EMassLOD::Max] = {};
 
 	/** If true, LowRes actors will be kept around, disabled, whilst StaticMeshInstance representation is active */
 	UPROPERTY(EditAnywhere, Category = "Mass|Representation", config)
diff --git a/Engine/Plugins/Runtime/StateTree/Source/StateTreeModule/Public/Blueprint/StateTreeItemBlueprintBase.h b/Engine/Plugins/Runtime/StateTree/Source/StateTreeModule/Public/Blueprint/StateTreeItemBlueprintBase.h
--- a/Engine/Plugins/Runtime/StateTree/Source/StateTreeModule/Public/Blueprint/StateTreeItemBlueprintBase.h
+++ b/Engine/Plugins/Runtime/StateTree/Source/StateTreeModule/Public/Blueprint/StateTreeItemBlueprintBase.h
@@ -34,7 +34,7 @@ struct STATETREEMODULE_API FStateTreeBlueprintPropertyInfo
 	FName PropertyName;
 
 	UPROPERTY()
-	EStateTreeBlueprintPropertyCategory Category;
+	EStateTreeBlueprintPropertyCategory Category = EStateTreeBlueprintPropertyCategory::NotSet;
 };
 
 /** Struct use to copy external data to the Blueprint item instance, resolved during StateTree linking. */
