Added OnSpawningFinishedEvent and OnDespawningFinishedEvent  BP events to AMassSpawner to allow for simple chaining of manual spawning via level BP etc

#rnx

diff --git a/Engine/Plugins/Runtime/MassGameplay/Source/MassSpawner/Private/MassSpawner.cpp b/Engine/Plugins/Runtime/MassGameplay/Source/MassSpawner/Private/MassSpawner.cpp
--- a/Engine/Plugins/Runtime/MassGameplay/Source/MassSpawner/Private/MassSpawner.cpp
+++ b/Engine/Plugins/Runtime/MassGameplay/Source/MassSpawner/Private/MassSpawner.cpp
@@ -471,6 +471,8 @@ void AMassSpawner::SpawnGeneratedEntities(TConstArrayView<FMassEntitySpawnDataGe
 		FMassProcessingContext ProcessingContext(*EntitySystem, /*TimeDelta=*/0.0f);
 		UE::Mass::Executor::RunProcessorsView(Processors, ProcessingContext);
 	}
+
+	OnSpawningFinishedEvent.Broadcast();
 }
 
 void AMassSpawner::DoDespawning()
@@ -492,6 +494,8 @@ void AMassSpawner::DoDespawning()
 		SpawnerSystem->DestroyEntities(SpawnedEntities.TemplateID, SpawnedEntities.Entities);
 	}
 	AllSpawnedEntities.Reset();
+
+	OnDespawningFinishedEvent.Broadcast();
 }
 
 bool AMassSpawner::DespawnEntity(const FMassEntityHandle Entity)
diff --git a/Engine/Plugins/Runtime/MassGameplay/Source/MassSpawner/Public/MassSpawner.h b/Engine/Plugins/Runtime/MassGameplay/Source/MassSpawner/Public/MassSpawner.h
--- a/Engine/Plugins/Runtime/MassGameplay/Source/MassSpawner/Public/MassSpawner.h
+++ b/Engine/Plugins/Runtime/MassGameplay/Source/MassSpawner/Public/MassSpawner.h
@@ -11,6 +11,9 @@
 class UMassSchematic;
 struct FStreamableHandle;
 
+DECLARE_DYNAMIC_MULTICAST_DELEGATE(FMassSpawnerOnSpawningFinishedEvent);
+DECLARE_DYNAMIC_MULTICAST_DELEGATE(FMassSpawnerOnDespawningFinishedEvent);
+
 /** A spawner you can put on a map and configure it to spawn different things */
 UCLASS(hidecategories = (Object, Actor, Input, Rendering, LOD, Cooking, Collision, HLOD, Partition))
 class MASSSPAWNER_API AMassSpawner : public AActor
@@ -76,6 +79,14 @@ public:
 	UFUNCTION(BlueprintCallable, Category = "Spawning")
 	void ScaleSpawningCount(float Scale) { SpawningCountScale = Scale; }
 
+	/** Called once DoSpawning completes and all entities have been spawned. */
+	UPROPERTY(BlueprintAssignable)
+	FMassSpawnerOnSpawningFinishedEvent OnSpawningFinishedEvent;
+
+	/** Called once DoDespawning completes and all mass agents spawned by this spawner have been despawned. */
+	UPROPERTY(BlueprintAssignable)
+	FMassSpawnerOnDespawningFinishedEvent OnDespawningFinishedEvent; 
+
 protected:
 	void OnPostWorldInit(UWorld* World, const UWorld::InitializationValues);
 	void SpawnGeneratedEntities(TConstArrayView<FMassEntitySpawnDataGeneratorResult> Results);
