Made MassVisualizer's root component static to avoid component attachment warning

#rnx

diff --git a/Engine/Plugins/Runtime/MassGameplay/Source/MassRepresentation/Private/MassVisualizer.cpp b/Engine/Plugins/Runtime/MassGameplay/Source/MassRepresentation/Private/MassVisualizer.cpp
--- a/Engine/Plugins/Runtime/MassGameplay/Source/MassRepresentation/Private/MassVisualizer.cpp
+++ b/Engine/Plugins/Runtime/MassGameplay/Source/MassRepresentation/Private/MassVisualizer.cpp
@@ -6,6 +6,7 @@
 AMassVisualizer::AMassVisualizer()
 {
 	RootComponent = CreateDefaultSubobject<USceneComponent>(TEXT("SceneComp"));
+	RootComponent->SetMobility(EComponentMobility::Static);
 	VisComponent = CreateDefaultSubobject<UMassVisualizationComponent>(TEXT("VisualizerComponent"));
 }
 
