Use shared fragments for replication settings in prevision of removing processor overrides

#rnx

diff --git a/Engine/Plugins/AI/MassCrowd/Source/MassCrowd/Private/MassCrowdReplicationProcessor.cpp b/Engine/Plugins/AI/MassCrowd/Source/MassCrowd/Private/MassCrowdReplicationProcessor.cpp
--- a/Engine/Plugins/AI/MassCrowd/Source/MassCrowd/Private/MassCrowdReplicationProcessor.cpp
+++ b/Engine/Plugins/AI/MassCrowd/Source/MassCrowd/Private/MassCrowdReplicationProcessor.cpp
@@ -1,7 +1,6 @@
 // Copyright Epic Games, Inc. All Rights Reserved.
 
 #include "MassCrowdReplicationProcessor.h"
-#include "MassReplicationSubsystem.h"
 #include "MassReplicationTypes.h"
 #include "MassClientBubbleHandler.h"
 #include "MassCrowdBubble.h"
@@ -38,18 +37,10 @@ void UMassCrowdReplicationProcessor::ConfigureQueries()
 
 	CollectViewerInfoQuery.AddTagRequirement<FTagFragment_MassCrowd>(EMassFragmentPresence::All);
 	CalculateLODQuery.AddTagRequirement<FTagFragment_MassCrowd>(EMassFragmentPresence::All);
+	AdjustLODDistancesQuery.AddTagRequirement<FTagFragment_MassCrowd>(EMassFragmentPresence::All);
 	EntityQuery.AddTagRequirement<FTagFragment_MassCrowd>(EMassFragmentPresence::All);
 }
 
-void UMassCrowdReplicationProcessor::Initialize(UObject& Owner)
-{
-	Super::Initialize(Owner);
-
-	check(ReplicationSubsystem);
-
-	BubbleInfoClassHandle = ReplicationSubsystem->GetBubbleInfoClassHandle(AMassCrowdClientBubbleInfo::StaticClass());
-}
-
 void UMassCrowdReplicationProcessor::Execute(UMassEntitySubsystem& EntitySubsystem, FMassExecutionContext& Context)
 {
 	Super::Execute(EntitySubsystem, Context);
@@ -64,7 +55,8 @@ void UMassCrowdReplicationProcessor::Execute(UMassEntitySubsystem& EntitySubsyst
 			{
 				const TConstArrayView<FDataFragment_Transform> TransformList = Context.GetFragmentView<FDataFragment_Transform>();
 				const TConstArrayView<FMassReplicationLODFragment> ViewerLODList = Context.GetFragmentView<FMassReplicationLODFragment>();
-				LODCalculator.DebugDisplayLOD(Context, ViewerLODList, TransformList, World);
+				FMassReplicationSharedFragment& RepSharedFragment = Context.GetMutableSharedFragment<FMassReplicationSharedFragment>();
+				RepSharedFragment.LODCalculator.DebugDisplayLOD(Context, ViewerLODList, TransformList, World);
 			});
 	}
 #endif // WITH_MASSGAMEPLAY_DEBUG
@@ -76,16 +68,19 @@ void UMassCrowdReplicationProcessor::ProcessClientReplication(UMassEntitySubsyst
 
 	FMassReplicationProcessorPathHandler PathHandler;
 	FMassReplicationProcessorPositionYawHandler PositionYawHandler;
+	FMassReplicationSharedFragment* RepSharedFrag = nullptr;
 
-	auto CacheViewsCallback = [&Context, &PathHandler, &PositionYawHandler]()
+	auto CacheViewsCallback = [&RepSharedFrag, &PathHandler, &PositionYawHandler](FMassExecutionContext& Context)
 	{
 		PathHandler.CacheFragmentViews(Context);
 		PositionYawHandler.CacheFragmentViews(Context);
+		RepSharedFrag = &Context.GetMutableSharedFragment<FMassReplicationSharedFragment>();
+		check(RepSharedFrag);
 	};
 
-	auto AddEntityCallback = [this, &Context, &PathHandler, &PositionYawHandler](const int32 EntityIdx, FReplicatedCrowdAgent& InReplicatedAgent, const FMassClientHandle ClientHandle)->FMassReplicatedAgentHandle
+	auto AddEntityCallback = [&RepSharedFrag, &PathHandler, &PositionYawHandler](FMassExecutionContext& Context,const int32 EntityIdx, FReplicatedCrowdAgent& InReplicatedAgent, const FMassClientHandle ClientHandle)->FMassReplicatedAgentHandle
 	{
-		AMassCrowdClientBubbleInfo& CrowdBubbleInfo = GetTypedClientBubbleInfoChecked<AMassCrowdClientBubbleInfo>(ClientHandle);
+		AMassCrowdClientBubbleInfo& CrowdBubbleInfo = RepSharedFrag->GetTypedClientBubbleInfoChecked<AMassCrowdClientBubbleInfo>(ClientHandle);
 
 		PathHandler.AddEntity(EntityIdx, InReplicatedAgent.GetReplicatedPathDataMutable());
 		PositionYawHandler.AddEntity(EntityIdx, InReplicatedAgent.GetReplicatedPositionYawDataMutable());
@@ -93,9 +88,9 @@ void UMassCrowdReplicationProcessor::ProcessClientReplication(UMassEntitySubsyst
 		return CrowdBubbleInfo.GetCrowdSerializer().Bubble.AddAgent(Context.GetEntity(EntityIdx), InReplicatedAgent);
 	};
 
-	auto ModifyEntityCallback = [this, &PathHandler](const int32 EntityIdx, const EMassLOD::Type LOD, const float Time, const FMassReplicatedAgentHandle Handle, const FMassClientHandle ClientHandle)
+	auto ModifyEntityCallback = [&RepSharedFrag, &PathHandler](FMassExecutionContext& Context, const int32 EntityIdx, const EMassLOD::Type LOD, const float Time, const FMassReplicatedAgentHandle Handle, const FMassClientHandle ClientHandle)
 	{
-		AMassCrowdClientBubbleInfo& CrowdBubbleInfo = GetTypedClientBubbleInfoChecked<AMassCrowdClientBubbleInfo>(ClientHandle);
+		AMassCrowdClientBubbleInfo& CrowdBubbleInfo = RepSharedFrag->GetTypedClientBubbleInfoChecked<AMassCrowdClientBubbleInfo>(ClientHandle);
 
 		FMassCrowdClientBubbleHandler& Bubble = CrowdBubbleInfo.GetCrowdSerializer().Bubble;
 
@@ -104,9 +99,9 @@ void UMassCrowdReplicationProcessor::ProcessClientReplication(UMassEntitySubsyst
 		// Don't call the PositionYawHandler here as we currently only replicate the position and yaw when we add an entity to Mass
 	};
 
-	auto RemoveEntityCallback = [this](const FMassReplicatedAgentHandle Handle, const FMassClientHandle ClientHandle)
+	auto RemoveEntityCallback = [&RepSharedFrag](FMassExecutionContext& Context, const FMassReplicatedAgentHandle Handle, const FMassClientHandle ClientHandle)
 	{
-		AMassCrowdClientBubbleInfo& CrowdBubbleInfo = GetTypedClientBubbleInfoChecked<AMassCrowdClientBubbleInfo>(ClientHandle);
+		AMassCrowdClientBubbleInfo& CrowdBubbleInfo = RepSharedFrag->GetTypedClientBubbleInfoChecked<AMassCrowdClientBubbleInfo>(ClientHandle);
 
 		CrowdBubbleInfo.GetCrowdSerializer().Bubble.RemoveAgentChecked(Handle);
 	};
diff --git a/Engine/Plugins/AI/MassCrowd/Source/MassCrowd/Public/MassCrowdReplicationProcessor.h b/Engine/Plugins/AI/MassCrowd/Source/MassCrowd/Public/MassCrowdReplicationProcessor.h
--- a/Engine/Plugins/AI/MassCrowd/Source/MassCrowd/Public/MassCrowdReplicationProcessor.h
+++ b/Engine/Plugins/AI/MassCrowd/Source/MassCrowd/Public/MassCrowdReplicationProcessor.h
@@ -21,7 +21,6 @@ public:
 
 protected:
 	virtual void ConfigureQueries() override;
-	virtual void Initialize(UObject& Owner) override;
 	virtual void Execute(UMassEntitySubsystem& EntitySubsystem, FMassExecutionContext& Context) override;
 
 	void ProcessClientReplication(UMassEntitySubsystem& EntitySubsystem, FMassExecutionContext& Context);
diff --git a/Engine/Plugins/Runtime/MassGameplay/Source/MassReplication/Private/MassReplicationFragments.cpp b/Engine/Plugins/Runtime/MassGameplay/Source/MassReplication/Private/MassReplicationFragments.cpp
--- a/Engine/Plugins/Runtime/MassGameplay/Source/MassReplication/Private/MassReplicationFragments.cpp
+++ b/Engine/Plugins/Runtime/MassGameplay/Source/MassReplication/Private/MassReplicationFragments.cpp
@@ -52,3 +52,39 @@ void UMassNetworkIDFragmentInitializer::Execute(UMassEntitySubsystem& EntitySubs
 #endif //UE_REPLICATION_COMPILE_SERVER_CODE
 	}
 }
+
+//----------------------------------------------------------------------//
+//  FMassReplicationParameters
+//----------------------------------------------------------------------//
+
+FMassReplicationParameters::FMassReplicationParameters()
+{
+	LODDistance[EMassLOD::High] = 0.f;
+	LODDistance[EMassLOD::Medium] = 1000.f;
+	LODDistance[EMassLOD::Low] = 2500.f;
+	LODDistance[EMassLOD::Off] = 5000.f;
+
+	LODMaxCount[EMassLOD::High] = 1600;
+	LODMaxCount[EMassLOD::Medium] = 3200;
+	LODMaxCount[EMassLOD::Low] = 48000;
+	LODMaxCount[EMassLOD::Off] = 0;
+
+	LODMaxCountPerViewer[EMassLOD::High] = 100;
+	LODMaxCountPerViewer[EMassLOD::Medium] = 200;
+	LODMaxCountPerViewer[EMassLOD::Low] = 300;
+	LODMaxCountPerViewer[EMassLOD::Off] = 0;
+
+	UpdateInterval[EMassLOD::High] = 0.1f;
+	UpdateInterval[EMassLOD::Medium] = 0.2f;
+	UpdateInterval[EMassLOD::Low] = 0.3f;
+	UpdateInterval[EMassLOD::Off] = 0.5f;
+}
+
+//----------------------------------------------------------------------//
+//  FMassReplicationSharedFragment
+//----------------------------------------------------------------------//
+FMassReplicationSharedFragment::FMassReplicationSharedFragment(UMassReplicationSubsystem& ReplicationSubsystem, const FMassReplicationParameters& Params)
+{
+	LODCalculator.Initialize(Params.LODDistance, Params.BufferHysteresisOnDistancePercentage / 100.0f, Params.LODMaxCount, Params.LODMaxCountPerViewer);
+	BubbleInfoClassHandle = ReplicationSubsystem.GetBubbleInfoClassHandle(Params.BubbleInfoClass);
+}
\ No newline at end of file
diff --git a/Engine/Plugins/Runtime/MassGameplay/Source/MassReplication/Private/MassReplicationProcessorBase.cpp b/Engine/Plugins/Runtime/MassGameplay/Source/MassReplication/Private/MassReplicationProcessorBase.cpp
--- a/Engine/Plugins/Runtime/MassGameplay/Source/MassReplication/Private/MassReplicationProcessorBase.cpp
+++ b/Engine/Plugins/Runtime/MassGameplay/Source/MassReplication/Private/MassReplicationProcessorBase.cpp
@@ -6,30 +6,10 @@
 #include "MassCommonFragments.h"
 
 //----------------------------------------------------------------------//
-//  UMassProcessor_Replication
+//  UMassReplicationProcessorBase
 //----------------------------------------------------------------------//
 UMassReplicationProcessorBase::UMassReplicationProcessorBase()
 {
-	LODDistance[EMassLOD::High] = 0.f;
-	LODDistance[EMassLOD::Medium] = 1000.f;
-	LODDistance[EMassLOD::Low] = 2500.f;
-	LODDistance[EMassLOD::Off] = 5000.f;
-	
-	LODMaxCount[EMassLOD::High] = 1600;
-	LODMaxCount[EMassLOD::Medium] = 3200;
-	LODMaxCount[EMassLOD::Low] = 48000;
-	LODMaxCount[EMassLOD::Off] = 0;
-
-	LODMaxCountPerViewer[EMassLOD::High] = 100;
-	LODMaxCountPerViewer[EMassLOD::Medium] = 200;
-	LODMaxCountPerViewer[EMassLOD::Low] = 300;
-	LODMaxCountPerViewer[EMassLOD::Off] = 0;
-	
-	UpdateInterval[EMassLOD::High] = 0.1f;
-	UpdateInterval[EMassLOD::Medium] = 0.2f;
-	UpdateInterval[EMassLOD::Low] = 0.3f;
-	UpdateInterval[EMassLOD::Off] = 0.5f;
-
 	ExecutionFlags = int32(EProcessorExecutionFlags::None);
 	ProcessingPhase = EMassProcessingPhase::PostPhysics;
 }
@@ -38,14 +18,28 @@ void UMassReplicationProcessorBase::ConfigureQueries()
 {
 	CollectViewerInfoQuery.AddRequirement<FDataFragment_Transform>(EMassFragmentAccess::ReadOnly);
 	CollectViewerInfoQuery.AddRequirement<FMassReplicationViewerInfoFragment>(EMassFragmentAccess::ReadOnly);
+	CollectViewerInfoQuery.AddSharedRequirement<FMassReplicationSharedFragment>(EMassFragmentAccess::ReadWrite);
 
 	CalculateLODQuery.AddRequirement<FMassReplicationViewerInfoFragment>(EMassFragmentAccess::ReadOnly);
 	CalculateLODQuery.AddRequirement<FMassReplicationLODFragment>(EMassFragmentAccess::ReadWrite);
+	CalculateLODQuery.AddConstSharedRequirement<FMassReplicationParameters>();
+	CalculateLODQuery.AddSharedRequirement<FMassReplicationSharedFragment>(EMassFragmentAccess::ReadWrite);
+
+	AdjustLODDistancesQuery.AddRequirement<FMassReplicationViewerInfoFragment>(EMassFragmentAccess::ReadOnly);
+	AdjustLODDistancesQuery.AddRequirement<FMassReplicationLODFragment>(EMassFragmentAccess::ReadWrite);
+	AdjustLODDistancesQuery.AddSharedRequirement<FMassReplicationSharedFragment>(EMassFragmentAccess::ReadWrite);
+	AdjustLODDistancesQuery.SetArchetypeFilter([](const FMassExecutionContext& Context)
+	{
+		const FMassReplicationSharedFragment& LODSharedFragment = Context.GetSharedFragment<FMassReplicationSharedFragment>();
+		return LODSharedFragment.bHasAdjustedDistancesFromCount;
+	});
 
 	EntityQuery.AddRequirement<FMassNetworkIDFragment>(EMassFragmentAccess::ReadOnly);
 	EntityQuery.AddRequirement<FDataFragment_ReplicationTemplateID>(EMassFragmentAccess::ReadOnly);
 	EntityQuery.AddRequirement<FMassReplicationLODFragment>(EMassFragmentAccess::ReadWrite);
 	EntityQuery.AddRequirement<FMassReplicatedAgentFragment>(EMassFragmentAccess::ReadWrite);
+	EntityQuery.AddConstSharedRequirement<FMassReplicationParameters>();
+	EntityQuery.AddSharedRequirement<FMassReplicationSharedFragment>(EMassFragmentAccess::ReadWrite);
 }
 
 void UMassReplicationProcessorBase::Initialize(UObject& Owner)
@@ -53,17 +47,13 @@ void UMassReplicationProcessorBase::Initialize(UObject& Owner)
 	Super::Initialize(Owner);
 
 #if UE_REPLICATION_COMPILE_SERVER_CODE
-	LODCalculator.Initialize(LODDistance, BufferHysteresisOnDistancePercentage / 100.0f, LODMaxCount, LODMaxCountPerViewer);
-
 	ReplicationSubsystem = UWorld::GetSubsystem<UMassReplicationSubsystem>(World);
 
 	check(ReplicationSubsystem);
-
-	//BubbleInfoClassHandle = ReplicationSubsystem->GetBubbleInfoClassHandle(AMassCrowdClientBubbleInfo::StaticClass());
 #endif //UE_REPLICATION_COMPILE_SERVER_CODE
 }
 
-void UMassReplicationProcessorBase::PrepareExecution()
+void UMassReplicationProcessorBase::PrepareExecution(UMassEntitySubsystem& EntitySubsystem)
 {
 #if UE_REPLICATION_COMPILE_SERVER_CODE
 
@@ -72,66 +62,68 @@ void UMassReplicationProcessorBase::PrepareExecution()
 	//first synchronize clients and viewers
 	ReplicationSubsystem->SynchronizeClientsAndViewers();
 
-	//then call base class
 	check(LODManager);
 	const TArray<FViewerInfo>& Viewers = LODManager->GetViewers();
-	LODCollector.PrepareExecution(Viewers);
-	LODCalculator.PrepareExecution(Viewers);
-
-	const TArray<FMassClientHandle>& CurrentClientHandles = ReplicationSubsystem->GetClientReplicationHandles();
-	const int32 MinNumHandles = FMath::Min(CachedClientHandles.Num(), CurrentClientHandles.Num());
-
-	//check to see if we don't have enough cached client handles
-	if (CachedClientHandles.Num() < CurrentClientHandles.Num())
+	EntitySubsystem.ForEachSharedFragment<FMassReplicationSharedFragment>([this, &Viewers](FMassReplicationSharedFragment& RepSharedFragment)
 	{
-		CachedClientHandles.Reserve(CurrentClientHandles.Num());
-		bBubbleChanged.Reserve(CurrentClientHandles.Num());
-		BubbleInfos.Reserve(CurrentClientHandles.Num());
+		RepSharedFragment.LODCollector.PrepareExecution(Viewers);
+		RepSharedFragment.LODCalculator.PrepareExecution(Viewers);
+
+		const TArray<FMassClientHandle>& CurrentClientHandles = ReplicationSubsystem->GetClientReplicationHandles();
+		const int32 MinNumHandles = FMath::Min(RepSharedFragment.CachedClientHandles.Num(), CurrentClientHandles.Num());
 
-		for (int32 Idx = CachedClientHandles.Num(); Idx < CurrentClientHandles.Num(); ++Idx)
+		//check to see if we don't have enough cached client handles
+		if (RepSharedFragment.CachedClientHandles.Num() < CurrentClientHandles.Num())
 		{
-			const FMassClientHandle& CurrentClientHandle = CurrentClientHandles[Idx];
+			RepSharedFragment.CachedClientHandles.Reserve(CurrentClientHandles.Num());
+			RepSharedFragment.bBubbleChanged.Reserve(CurrentClientHandles.Num());
+			RepSharedFragment.BubbleInfos.Reserve(CurrentClientHandles.Num());
 
-			CachedClientHandles.Add(CurrentClientHandle);
-			bBubbleChanged.Add(true);
-			AMassClientBubbleInfoBase* Info = CurrentClientHandle.IsValid() ?
-				ReplicationSubsystem->GetClientBubbleChecked(BubbleInfoClassHandle, CurrentClientHandle) :
-				nullptr;
+			for (int32 Idx = RepSharedFragment.CachedClientHandles.Num(); Idx < CurrentClientHandles.Num(); ++Idx)
+			{
+				const FMassClientHandle& CurrentClientHandle = CurrentClientHandles[Idx];
 
-			check(Info);
+				RepSharedFragment.CachedClientHandles.Add(CurrentClientHandle);
+				RepSharedFragment.bBubbleChanged.Add(true);
+				AMassClientBubbleInfoBase* Info = CurrentClientHandle.IsValid() ?
+					ReplicationSubsystem->GetClientBubbleChecked(RepSharedFragment.BubbleInfoClassHandle, CurrentClientHandle) :
+					nullptr;
 
-			BubbleInfos.Add(Info);
-		}
-	}
-	//check to see if we have too many cached client handles
-	else if (CachedClientHandles.Num() > CurrentClientHandles.Num())
-	{
-		const int32 NumRemove = CachedClientHandles.Num() - CurrentClientHandles.Num();
+				check(Info);
 
-		CachedClientHandles.RemoveAt(CurrentClientHandles.Num(), NumRemove, /* bAllowShrinking */ false);
-		bBubbleChanged.RemoveAt(CurrentClientHandles.Num(), NumRemove, /* bAllowShrinking */ false);
-		BubbleInfos.RemoveAt(CurrentClientHandles.Num(), NumRemove, /* bAllowShrinking */ false);
-	}
+				RepSharedFragment.BubbleInfos.Add(Info);
+			}
+		}
+		//check to see if we have too many cached client handles
+		else if (RepSharedFragment.CachedClientHandles.Num() > CurrentClientHandles.Num())
+		{
+			const int32 NumRemove = RepSharedFragment.CachedClientHandles.Num() - CurrentClientHandles.Num();
 
-	//check to see if any cached client handles have changed, if they have set the BubbleInfo[] appropriately
-	for (int32 Idx = 0; Idx < MinNumHandles; ++Idx)
-	{
-		const FMassClientHandle& CurrentClientHandle = CurrentClientHandles[Idx];
-		FMassClientHandle& CachedClientHandle = CachedClientHandles[Idx];
+			RepSharedFragment.CachedClientHandles.RemoveAt(CurrentClientHandles.Num(), NumRemove, /* bAllowShrinking */ false);
+			RepSharedFragment.bBubbleChanged.RemoveAt(CurrentClientHandles.Num(), NumRemove, /* bAllowShrinking */ false);
+			RepSharedFragment.BubbleInfos.RemoveAt(CurrentClientHandles.Num(), NumRemove, /* bAllowShrinking */ false);
+		}
 
-		const bool bChanged = (bBubbleChanged[Idx] = (CurrentClientHandle != CachedClientHandle));
-		if (bChanged)
+		//check to see if any cached client handles have changed, if they have set the BubbleInfo[] appropriately
+		for (int32 Idx = 0; Idx < MinNumHandles; ++Idx)
 		{
-			AMassClientBubbleInfoBase* Info = CurrentClientHandle.IsValid() ?
-				ReplicationSubsystem->GetClientBubbleChecked(BubbleInfoClassHandle, CurrentClientHandle) :
-				nullptr;
+			const FMassClientHandle& CurrentClientHandle = CurrentClientHandles[Idx];
+			FMassClientHandle& CachedClientHandle = RepSharedFragment.CachedClientHandles[Idx];
+
+			const bool bChanged = (RepSharedFragment.bBubbleChanged[Idx] = (CurrentClientHandle != CachedClientHandle));
+			if (bChanged)
+			{
+				AMassClientBubbleInfoBase* Info = CurrentClientHandle.IsValid() ?
+					ReplicationSubsystem->GetClientBubbleChecked(RepSharedFragment.BubbleInfoClassHandle, CurrentClientHandle) :
+					nullptr;
 
-			check(Info);
+				check(Info);
 
-			BubbleInfos[Idx] = Info;
-			CachedClientHandle = CurrentClientHandle;
+				RepSharedFragment.BubbleInfos[Idx] = Info;
+				CachedClientHandle = CurrentClientHandle;
+			}
 		}
-	}
+	});
 
 #endif //UE_REPLICATION_COMPILE_SERVER_CODE
 }
@@ -141,30 +133,36 @@ void UMassReplicationProcessorBase::Execute(UMassEntitySubsystem& EntitySubsyste
 #if UE_REPLICATION_COMPILE_SERVER_CODE
 	check(World);
 
-	PrepareExecution();
+	PrepareExecution(EntitySubsystem);
 
 	CollectViewerInfoQuery.ForEachEntityChunk(EntitySubsystem, Context, [this](FMassExecutionContext& Context)
 	{
 		const TConstArrayView<FDataFragment_Transform> LocationList = Context.GetFragmentView<FDataFragment_Transform>();
 		const TArrayView<FMassReplicationViewerInfoFragment> ViewersInfoList = Context.GetMutableFragmentView<FMassReplicationViewerInfoFragment>();
-		LODCollector.CollectLODInfo(Context, LocationList, ViewersInfoList, ViewersInfoList);
+		FMassReplicationSharedFragment& RepSharedFragment = Context.GetMutableSharedFragment<FMassReplicationSharedFragment>();
+		RepSharedFragment.LODCollector.CollectLODInfo(Context, LocationList, ViewersInfoList, ViewersInfoList);
 	});
 
 	CalculateLODQuery.ForEachEntityChunk(EntitySubsystem, Context, [this](FMassExecutionContext& Context)
 	{
 		const TConstArrayView<FMassReplicationViewerInfoFragment> ViewersInfoList = Context.GetFragmentView<FMassReplicationViewerInfoFragment>();
 		const TArrayView<FMassReplicationLODFragment> ViewerLODList = Context.GetMutableFragmentView<FMassReplicationLODFragment>();
-		LODCalculator.CalculateLOD(Context, ViewersInfoList, ViewerLODList, ViewersInfoList);
+		FMassReplicationSharedFragment& RepSharedFragment = Context.GetMutableSharedFragment<FMassReplicationSharedFragment>();
+		RepSharedFragment.LODCalculator.CalculateLOD(Context, ViewersInfoList, ViewerLODList, ViewersInfoList);
 	});
 
-	if (LODCalculator.AdjustDistancesFromCount())
+	EntitySubsystem.ForEachSharedFragment<FMassReplicationSharedFragment>([](FMassReplicationSharedFragment& RepSharedFragment)
 	{
-		CalculateLODQuery.ForEachEntityChunk(EntitySubsystem, Context, [this](FMassExecutionContext& Context)
-		{
-			const TConstArrayView<FMassReplicationViewerInfoFragment> ViewersInfoList = Context.GetFragmentView<FMassReplicationViewerInfoFragment>();
-			const TArrayView<FMassReplicationLODFragment> ViewerLODList = Context.GetMutableFragmentView<FMassReplicationLODFragment>();
-			LODCalculator.AdjustLODFromCount(Context, ViewersInfoList, ViewerLODList, ViewersInfoList);
-		});
-	}
+		RepSharedFragment.bHasAdjustedDistancesFromCount = RepSharedFragment.LODCalculator.AdjustDistancesFromCount();
+	});
+
+	AdjustLODDistancesQuery.ForEachEntityChunk(EntitySubsystem, Context, [this](FMassExecutionContext& Context)
+	{
+		const TConstArrayView<FMassReplicationViewerInfoFragment> ViewersInfoList = Context.GetFragmentView<FMassReplicationViewerInfoFragment>();
+		const TArrayView<FMassReplicationLODFragment> ViewerLODList = Context.GetMutableFragmentView<FMassReplicationLODFragment>();
+		FMassReplicationSharedFragment& RepSharedFragment = Context.GetMutableSharedFragment<FMassReplicationSharedFragment>();
+		RepSharedFragment.LODCalculator.AdjustLODFromCount(Context, ViewersInfoList, ViewerLODList, ViewersInfoList);
+	});
+
 #endif //UE_REPLICATION_COMPILE_SERVER_CODE
-}
+}
\ No newline at end of file
diff --git a/Engine/Plugins/Runtime/MassGameplay/Source/MassReplication/Private/MassReplicationTrait.cpp b/Engine/Plugins/Runtime/MassGameplay/Source/MassReplication/Private/MassReplicationTrait.cpp
--- a/Engine/Plugins/Runtime/MassGameplay/Source/MassReplication/Private/MassReplicationTrait.cpp
+++ b/Engine/Plugins/Runtime/MassGameplay/Source/MassReplication/Private/MassReplicationTrait.cpp
@@ -6,7 +6,7 @@
 #include "MassSpawnerTypes.h"
 #include "MassSimulationLOD.h"
 #include "MassReplicationTypes.h"
-#include "MassReplicationFragments.h"
+#include "MassReplicationSubsystem.h"
 
 
 void UMassReplicationTrait::BuildTemplate(FMassEntityTemplateBuildContext& BuildContext, UWorld& World) const
@@ -23,4 +23,16 @@ void UMassReplicationTrait::BuildTemplate(FMassEntityTemplateBuildContext& Build
 	BuildContext.AddFragment<FMassReplicatedAgentFragment>();
 	BuildContext.AddFragment<FMassReplicationViewerInfoFragment>();
 	BuildContext.AddFragment<FMassReplicationLODFragment>();
+
+	UMassEntitySubsystem* EntitySubsystem = UWorld::GetSubsystem<UMassEntitySubsystem>(&World);
+	check(EntitySubsystem);
+
+	UMassReplicationSubsystem* ReplicationSubsystem = UWorld::GetSubsystem<UMassReplicationSubsystem>(&World);
+	check(ReplicationSubsystem);
+
+	uint32 ParamsHash = UE::StructUtils::GetStructCrc32(FConstStructView::Make(Params));
+	FConstSharedStruct ParamsFragment = EntitySubsystem->GetOrCreateConstSharedFragment(ParamsHash, Params);
+	BuildContext.AddConstSharedFragment(ParamsFragment);
+	FSharedStruct SharedFragment = EntitySubsystem->GetOrCreateSharedFragment<FMassReplicationSharedFragment>(ParamsHash, *ReplicationSubsystem, Params);
+	BuildContext.AddSharedFragment(SharedFragment);
 }
\ No newline at end of file
diff --git a/Engine/Plugins/Runtime/MassGameplay/Source/MassReplication/Public/MassReplicationFragments.h b/Engine/Plugins/Runtime/MassGameplay/Source/MassReplication/Public/MassReplicationFragments.h
--- a/Engine/Plugins/Runtime/MassGameplay/Source/MassReplication/Public/MassReplicationFragments.h
+++ b/Engine/Plugins/Runtime/MassGameplay/Source/MassReplication/Public/MassReplicationFragments.h
@@ -4,14 +4,13 @@
 
 #include "MassCommonTypes.h"
 #include "MassReplicationTypes.h"
-// #include "MassLODManager.h"
-// #include "IndexedHandle.h"
-// #include "MassEntityTemplate.h"
-// #include "MassObserverProcessor.h"
-// #include "MassLODLogic.h"
+#include "MassLODCollector.h"
+#include "MassLODCalculator.h"
 
 #include "MassReplicationFragments.generated.h"
 
+class AMassClientBubbleInfoBase;
+
 /**
  *  Fragment type for the mass network id of a mass entity
  */
@@ -104,3 +103,69 @@ protected:
 
 	FMassEntityQuery EntityQuery;
 };
+
+USTRUCT()
+struct FMassReplicationParameters : public FMassSharedFragment
+{
+	GENERATED_BODY()
+public:
+	FMassReplicationParameters();
+
+	/** Distance where each LOD becomes relevant */
+	UPROPERTY(EditAnywhere, Category = "Mass|LOD", config)
+	float LODDistance[EMassLOD::Max];
+
+	/** Hysteresis percentage on delta between the LOD distances */
+	UPROPERTY(EditAnywhere, Category = "Mass|LOD", meta = (ClampMin = "0.0", UIMin = "0.0"), config)
+	float BufferHysteresisOnDistancePercentage = 10.0f;
+
+	/** Maximum limit of entity per LOD */
+	UPROPERTY(EditAnywhere, Category = "Mass|LOD", config)
+	int32 LODMaxCount[EMassLOD::Max];
+
+	/** Maximum limit of entity per LOD per viewer */
+	UPROPERTY(EditAnywhere, Category = "Mass|LOD", config)
+	int32 LODMaxCountPerViewer[EMassLOD::Max];
+
+	/** Distance where each LOD becomes relevant */
+	UPROPERTY(EditAnywhere, Category = "Mass|LOD", config)
+	float UpdateInterval[EMassLOD::Max];
+
+	UPROPERTY(EditAnywhere, Category = "Mass|Replication", config)
+	TSubclassOf<AMassClientBubbleInfoBase> BubbleInfoClass;
+};
+
+USTRUCT()
+struct FMassReplicationSharedFragment : public FMassSharedFragment
+{
+	GENERATED_BODY()
+public:
+	FMassReplicationSharedFragment() = default;
+	FMassReplicationSharedFragment(UMassReplicationSubsystem& ReplicationSubsystem, const FMassReplicationParameters& Params);
+
+	FMassBubbleInfoClassHandle BubbleInfoClassHandle;
+
+	TArray<FMassClientHandle, TInlineAllocator<UE::Mass::Replication::MaxNumOfClients>> CachedClientHandles;
+	TArray<bool, TInlineAllocator<UE::Mass::Replication::MaxNumOfClients>> bBubbleChanged;
+
+	//TODO review if we need to have this as a UPROPERTY at all and also if we can make this use a TInlineAllocator
+	//Can not use TInlineAllocator with UPROPERTY()
+	UPROPERTY(Transient)
+	TArray<AMassClientBubbleInfoBase*> BubbleInfos;
+
+	TMassLODCollector<FReplicationLODLogic> LODCollector;
+	TMassLODCalculator<FReplicationLODLogic> LODCalculator;
+	bool bHasAdjustedDistancesFromCount = false;
+
+	template<typename T>
+	T& GetTypedClientBubbleInfoChecked(FMassClientHandle Handle)
+	{
+		checkSlow(BubbleInfos.IsValidIndex(Handle.GetIndex()));
+
+		AMassClientBubbleInfoBase* BubbleInfo = BubbleInfos[Handle.GetIndex()];
+
+		checkSlow(BubbleInfo && Cast<T>(BubbleInfo) != nullptr);
+
+		return *static_cast<T*>(BubbleInfo);
+	}
+};
\ No newline at end of file
diff --git a/Engine/Plugins/Runtime/MassGameplay/Source/MassReplication/Public/MassReplicationProcessorBase.h b/Engine/Plugins/Runtime/MassGameplay/Source/MassReplication/Public/MassReplicationProcessorBase.h
--- a/Engine/Plugins/Runtime/MassGameplay/Source/MassReplication/Public/MassReplicationProcessorBase.h
+++ b/Engine/Plugins/Runtime/MassGameplay/Source/MassReplication/Public/MassReplicationProcessorBase.h
@@ -5,7 +5,6 @@
 #include "Engine/World.h"
 #include "MassProcessor.h"
 #include "MassCommonTypes.h"
-#include "MassLODCalculator.h"
 #include "MassSimulationLOD.h"
 #include "MassReplicationTypes.h"
 #include "MassReplicationFragments.h"
@@ -35,7 +34,7 @@ protected:
 	virtual void Initialize(UObject& Owner) override;
 	virtual void Execute(UMassEntitySubsystem& EntitySubsystem, FMassExecutionContext& Context) override;
 
-	void PrepareExecution();
+	void PrepareExecution(UMassEntitySubsystem& EntitySubsystem);
 
 	/** 
 	 *  Implemented as straight template callbacks as when profiled this was faster than TFunctionRef. Its probably easier to pass Lamdas in to these
@@ -44,49 +43,16 @@ protected:
 	template<typename AgentArrayItem, typename CacheViewsCallback, typename AddEntityCallback, typename ModifyEntityCallback, typename RemoveEntityCallback>
 	void CalculateClientReplication(FMassExecutionContext& Context, CacheViewsCallback&& CacheViews, AddEntityCallback&& AddEntity, ModifyEntityCallback&& ModifyEntity, RemoveEntityCallback&& RemoveEntity);
 
-	template<typename T>
-	T& GetTypedClientBubbleInfoChecked(FMassClientHandle Handle);
+
 
 protected:
 
 	UPROPERTY()
 	UMassReplicationSubsystem* ReplicationSubsystem;
 
-	FMassBubbleInfoClassHandle BubbleInfoClassHandle;
-
-	/** Distance where each LOD becomes relevant */
-	UPROPERTY(EditAnywhere, Category = "Mass|LOD", config)
-	float LODDistance[EMassLOD::Max];
-
-	/** Hysteresis percentage on delta between the LOD distances */
-	UPROPERTY(EditAnywhere, Category = "Mass|LOD", meta = (ClampMin = "0.0", UIMin = "0.0"), config)
-	float BufferHysteresisOnDistancePercentage = 10.0f;
-
-	/** Maximum limit of entity per LOD */
-	UPROPERTY(EditAnywhere, Category = "Mass|LOD", config)
-	int32 LODMaxCount[EMassLOD::Max];
-
-	/** Maximum limit of entity per LOD per viewer */
-	UPROPERTY(EditAnywhere, Category = "Mass|LOD", config)
-	int32 LODMaxCountPerViewer[EMassLOD::Max];
-
-	/** Distance where each LOD becomes relevant */
-	UPROPERTY(EditAnywhere, Category = "Mass|LOD", config)
-	float UpdateInterval[EMassLOD::Max];
-
-	TArray<FMassClientHandle, TInlineAllocator<UE::Mass::Replication::MaxNumOfClients>> CachedClientHandles;
-	TArray<bool, TInlineAllocator<UE::Mass::Replication::MaxNumOfClients>> bBubbleChanged;
-
-	//TODO review if we need to have this as a UPROPERTY at all and also if we can make this use a TInlineAllocator
-	//Can not use TInlineAllocator with UPROPERTY()
-	UPROPERTY()
-	TArray<AMassClientBubbleInfoBase*> BubbleInfos;
-
-	TMassLODCollector<FReplicationLODLogic> LODCollector;
-	TMassLODCalculator<FReplicationLODLogic> LODCalculator;
-
 	FMassEntityQuery CollectViewerInfoQuery;
 	FMassEntityQuery CalculateLODQuery;
+	FMassEntityQuery AdjustLODDistancesQuery;
 	FMassEntityQuery EntityQuery;
 };
 
@@ -105,8 +71,9 @@ void UMassReplicationProcessorBase::CalculateClientReplication(FMassExecutionCon
 	TArrayView<FMassReplicationLODFragment> ViewerLODList = Context.GetMutableFragmentView<FMassReplicationLODFragment>();
 	TArrayView<FMassReplicatedAgentFragment> ReplicatedAgentList = Context.GetMutableFragmentView<FMassReplicatedAgentFragment>();
 	TConstArrayView<FDataFragment_ReplicationTemplateID> TemplateIDList = Context.GetFragmentView<FDataFragment_ReplicationTemplateID>();
+	FMassReplicationSharedFragment& RepSharedFragment = Context.GetMutableSharedFragment<FMassReplicationSharedFragment>();
 
-	CacheViews();
+	CacheViews(Context);
 
 	const float Time = World->GetRealTimeSeconds();
 
@@ -114,28 +81,28 @@ void UMassReplicationProcessorBase::CalculateClientReplication(FMassExecutionCon
 	{
 		FMassReplicatedAgentFragment& AgentFragment = ReplicatedAgentList[EntityIdx];
 
-		if (AgentFragment.AgentsData.Num() < CachedClientHandles.Num())
+		if (AgentFragment.AgentsData.Num() < RepSharedFragment.CachedClientHandles.Num())
 		{
-			AgentFragment.AgentsData.AddDefaulted(CachedClientHandles.Num() - AgentFragment.AgentsData.Num());
+			AgentFragment.AgentsData.AddDefaulted(RepSharedFragment.CachedClientHandles.Num() - AgentFragment.AgentsData.Num());
 		}
-		else if (AgentFragment.AgentsData.Num() > CachedClientHandles.Num())
+		else if (AgentFragment.AgentsData.Num() > RepSharedFragment.CachedClientHandles.Num())
 		{
-			AgentFragment.AgentsData.RemoveAt(CachedClientHandles.Num(), AgentFragment.AgentsData.Num() - CachedClientHandles.Num(), /* bAllowShrinking */ false);
+			AgentFragment.AgentsData.RemoveAt(RepSharedFragment.CachedClientHandles.Num(), AgentFragment.AgentsData.Num() - RepSharedFragment.CachedClientHandles.Num(), /* bAllowShrinking */ false);
 		}
 
-		for (int32 ClientIdx = 0; ClientIdx < CachedClientHandles.Num(); ++ClientIdx)
+		for (int32 ClientIdx = 0; ClientIdx < RepSharedFragment.CachedClientHandles.Num(); ++ClientIdx)
 		{
-			const FMassClientHandle& ClientHandle = CachedClientHandles[ClientIdx];
+			const FMassClientHandle& ClientHandle = RepSharedFragment.CachedClientHandles[ClientIdx];
 
 			if (ClientHandle.IsValid())
 			{
-				checkSlow(BubbleInfos[ClientHandle.GetIndex()] != nullptr);
+				checkSlow(RepSharedFragment.BubbleInfos[ClientHandle.GetIndex()] != nullptr);
 
 				FMassReplicatedAgentArrayData& AgentData = AgentFragment.AgentsData[ClientIdx];
 
 				//if the bubble has changed, set the handle Invalid. We will set this to something valid if the agent is going to replicate to the bubble
 				//When a bubble has changed the existing AMassCrowdClientBubbleInfo will reset all the data associated with it.
-				if (bBubbleChanged[ClientIdx])
+				if (RepSharedFragment.bBubbleChanged[ClientIdx])
 				{
 					AgentData.Invalidate();
 				}
@@ -177,13 +144,13 @@ void UMassReplicationProcessorBase::CalculateClientReplication(FMassExecutionCon
 						ReplicatedAgent.SetNetID(NetIDFragment.NetID);
 						ReplicatedAgent.SetTemplateID(TemplateIDFragment.ID);
 
-						AgentData.Handle = AddEntity(EntityIdx, ReplicatedAgent, ClientHandle);
+						AgentData.Handle = AddEntity(Context, EntityIdx, ReplicatedAgent, ClientHandle);
 
 						AgentData.LastUpdateTime = Time;
 					}
 					else
 					{
-						ModifyEntity(EntityIdx, HighestLOD, Time, AgentData.Handle, ClientHandle);
+						ModifyEntity(Context, EntityIdx, HighestLOD, Time, AgentData.Handle, ClientHandle);
 					}
 				}
 				else
@@ -191,7 +158,7 @@ void UMassReplicationProcessorBase::CalculateClientReplication(FMassExecutionCon
 					// as this is a fresh handle, if its valid then we can use the unsafe remove function
 					if (AgentData.Handle.IsValid())
 					{
-						RemoveEntity(AgentData.Handle, ClientHandle);
+						RemoveEntity(Context, AgentData.Handle, ClientHandle);
 						AgentData.Invalidate();
 					}
 				}
@@ -200,15 +167,3 @@ void UMassReplicationProcessorBase::CalculateClientReplication(FMassExecutionCon
 	}
 #endif //UE_REPLICATION_COMPILE_SERVER_CODE
 }
-
-template<typename T>
-T& UMassReplicationProcessorBase::GetTypedClientBubbleInfoChecked(FMassClientHandle Handle)
-{
-	checkSlow(BubbleInfos.IsValidIndex(Handle.GetIndex()));
-
-	AMassClientBubbleInfoBase* BubbleInfo = BubbleInfos[Handle.GetIndex()];
-
-	checkSlow(BubbleInfo && Cast<T>(BubbleInfo) != nullptr);
-
-	return *static_cast<T*>(BubbleInfo);
-}
diff --git a/Engine/Plugins/Runtime/MassGameplay/Source/MassReplication/Public/MassReplicationTrait.h b/Engine/Plugins/Runtime/MassGameplay/Source/MassReplication/Public/MassReplicationTrait.h
--- a/Engine/Plugins/Runtime/MassGameplay/Source/MassReplication/Public/MassReplicationTrait.h
+++ b/Engine/Plugins/Runtime/MassGameplay/Source/MassReplication/Public/MassReplicationTrait.h
@@ -3,6 +3,7 @@
 #pragma once 
 
 #include "MassEntityTraitBase.h"
+#include "MassReplicationFragments.h"
 #include "MassReplicationTrait.generated.h"
 
 
@@ -11,7 +12,10 @@ class MASSREPLICATION_API UMassReplicationTrait : public UMassEntityTraitBase
 {
 	GENERATED_BODY()
 
-protected:
+public:
 
 	virtual void BuildTemplate(FMassEntityTemplateBuildContext& BuildContext, UWorld& World) const override;
+
+	UPROPERTY(EditAnywhere, Category = "Mass|Replication")
+	FMassReplicationParameters Params;
 };
