Added Transient flag to FMassSpawnedEntityType::EntityConfigPtr.

This fixes an issue where duplicating a MassSpawner and changing the entity config assets referenced would have no affect as the underlying EntityConfigPtr would still have the old reference.

#rnx

diff --git a/Engine/Plugins/Runtime/MassGameplay/Source/MassSpawner/Public/MassSpawnerTypes.h b/Engine/Plugins/Runtime/MassGameplay/Source/MassSpawner/Public/MassSpawnerTypes.h
--- a/Engine/Plugins/Runtime/MassGameplay/Source/MassSpawner/Public/MassSpawnerTypes.h
+++ b/Engine/Plugins/Runtime/MassGameplay/Source/MassSpawner/Public/MassSpawnerTypes.h
@@ -50,7 +50,7 @@ struct FMassSpawnedEntityType
 	UMassEntityConfigAsset* GetEntityConfig();
 
 private:
-	UPROPERTY()
+	UPROPERTY(Transient)
 	mutable UMassEntityConfigAsset* EntityConfigPtr = nullptr;
 };
 
