Remove replication processor inheritance

#rnx

diff --git a/Engine/Plugins/AI/MassCrowd/Source/MassCrowd/Private/MassCrowdReplicationProcessor.cpp b/Engine/Plugins/AI/MassCrowd/Source/MassCrowd/Private/MassCrowdReplicationProcessor.cpp
--- a/Engine/Plugins/AI/MassCrowd/Source/MassCrowd/Private/MassCrowdReplicationProcessor.cpp
+++ b/Engine/Plugins/AI/MassCrowd/Source/MassCrowd/Private/MassCrowdReplicationProcessor.cpp
@@ -10,59 +10,16 @@
 #include "MassReplicationPathHandlers.h"
 #include "MassReplicationTransformHandlers.h"
 
-namespace UE { namespace Mass { namespace Crowd
-{
-	int32 bDebugReplicationViewerLOD = 0;
-	FAutoConsoleVariableRef CVarDebugReplicationViewerLOD(TEXT("ai.debug.CrowdReplicationViewerLOD"), bDebugReplicationViewerLOD, TEXT("Crowd Debug Replication Viewer LOD"), ECVF_Cheat);
-}}}
-
 //----------------------------------------------------------------------//
-//  UMassProcessor_Replication
+//  UMassCrowdReplicator
 //----------------------------------------------------------------------//
-UMassCrowdReplicationProcessor::UMassCrowdReplicationProcessor()
-{
-#if !UE_ALLOW_DEBUG_REPLICATION_BUBBLES_STANDALONE
-	ExecutionFlags = int32(EProcessorExecutionFlags::Server);
-#else
-	ExecutionFlags = int32(EProcessorExecutionFlags::All);
-#endif // UE_ALLOW_DEBUG_REPLICATION_BUBBLES_STANDALONE
-}
-
-void UMassCrowdReplicationProcessor::ConfigureQueries()
+void UMassCrowdReplicator::AddRequirements(FMassEntityQuery& EntityQuery)
 {
-	Super::ConfigureQueries();
-
 	FMassReplicationProcessorPositionYawHandler::AddRequirements(EntityQuery);
 	FMassReplicationProcessorPathHandler::AddRequirements(EntityQuery);
-
-	CollectViewerInfoQuery.AddTagRequirement<FMassCrowdTag>(EMassFragmentPresence::All);
-	CalculateLODQuery.AddTagRequirement<FMassCrowdTag>(EMassFragmentPresence::All);
-	AdjustLODDistancesQuery.AddTagRequirement<FMassCrowdTag>(EMassFragmentPresence::All);
-	EntityQuery.AddTagRequirement<FMassCrowdTag>(EMassFragmentPresence::All);
 }
 
-void UMassCrowdReplicationProcessor::Execute(UMassEntitySubsystem& EntitySubsystem, FMassExecutionContext& Context)
-{
-	Super::Execute(EntitySubsystem, Context);
-
-	ProcessClientReplication(EntitySubsystem, Context);
-
-#if WITH_MASSGAMEPLAY_DEBUG
-	// Optional debug display
-	if (UE::Mass::Crowd::bDebugReplicationViewerLOD)
-	{
-		EntityQuery.ForEachEntityChunk(EntitySubsystem, Context, [this](FMassExecutionContext& Context)
-			{
-				const TConstArrayView<FTransformFragment> TransformList = Context.GetFragmentView<FTransformFragment>();
-				const TConstArrayView<FMassReplicationLODFragment> ViewerLODList = Context.GetFragmentView<FMassReplicationLODFragment>();
-				FMassReplicationSharedFragment& RepSharedFragment = Context.GetMutableSharedFragment<FMassReplicationSharedFragment>();
-				RepSharedFragment.LODCalculator.DebugDisplayLOD(Context, ViewerLODList, TransformList, World);
-			});
-	}
-#endif // WITH_MASSGAMEPLAY_DEBUG
-}
-
-void UMassCrowdReplicationProcessor::ProcessClientReplication(UMassEntitySubsystem& EntitySubsystem, FMassExecutionContext& Context)
+void UMassCrowdReplicator::ProcessClientReplication(FMassExecutionContext& Context, FMassReplicationContext& ReplicationContext)
 {
 #if UE_REPLICATION_COMPILE_SERVER_CODE
 
@@ -106,11 +63,6 @@ void UMassCrowdReplicationProcessor::ProcessClientReplication(UMassEntitySubsyst
 		CrowdBubbleInfo.GetCrowdSerializer().Bubble.RemoveAgentChecked(Handle);
 	};
 
-	QUICK_SCOPE_CYCLE_COUNTER(UMassCrowdReplicationProcessor_ProcessClientReplication);
-
-	EntityQuery.ForEachEntityChunk(EntitySubsystem, Context, [this, &CacheViewsCallback, &AddEntityCallback, &ModifyEntityCallback, &RemoveEntityCallback](FMassExecutionContext& Context)
-	{
-		CalculateClientReplication<FCrowdFastArrayItem>(Context, CacheViewsCallback, AddEntityCallback, ModifyEntityCallback, RemoveEntityCallback);
-	});
+	CalculateClientReplication<FCrowdFastArrayItem>(Context, ReplicationContext, CacheViewsCallback, AddEntityCallback, ModifyEntityCallback, RemoveEntityCallback);
 #endif // UE_REPLICATION_COMPILE_SERVER_CODE
 }
diff --git a/Engine/Plugins/AI/MassCrowd/Source/MassCrowd/Public/MassCrowdReplicationProcessor.h b/Engine/Plugins/AI/MassCrowd/Source/MassCrowd/Public/MassCrowdReplicationProcessor.h
--- a/Engine/Plugins/AI/MassCrowd/Source/MassCrowd/Public/MassCrowdReplicationProcessor.h
+++ b/Engine/Plugins/AI/MassCrowd/Source/MassCrowd/Public/MassCrowdReplicationProcessor.h
@@ -6,22 +6,22 @@
 
 #include "MassCrowdReplicationProcessor.generated.h"
 
-class UMassReplicationSubsystem;
-class MassLODManager;
-class AMassCrowdClientBubbleInfo;
-class UWorld;
-
-/** Processor that handles replication and only runs on the server. It queries Mass entity fragments and sets those values when appropriate using the MassClientBubbleHandler. */
+/** Class that handles replication and only runs on the server. It queries Mass entity fragments and sets those values when appropriate using the MassClientBubbleHandler. */
 UCLASS()
-class MASSCROWD_API UMassCrowdReplicationProcessor : public UMassReplicationProcessorBase
+class MASSCROWD_API UMassCrowdReplicator : public UMassReplicatorBase
 {
 	GENERATED_BODY()
-public:
-	UMassCrowdReplicationProcessor();
 
-protected:
-	virtual void ConfigureQueries() override;
-	virtual void Execute(UMassEntitySubsystem& EntitySubsystem, FMassExecutionContext& Context) override;
-
-	void ProcessClientReplication(UMassEntitySubsystem& EntitySubsystem, FMassExecutionContext& Context);
+public:
+	/**
+	 * Overridden to add specific entity query requirements for replication.
+	 * Usually we add replication processor handler requirements.
+	 */
+	void AddRequirements(FMassEntityQuery& EntityQuery) override;
+	
+	/**
+	 * Overridden to process the client replication.
+	 * This methods should call CalculateClientReplication with the appropriate callback implementation.
+	 */
+	void ProcessClientReplication(FMassExecutionContext& Context, FMassReplicationContext& ReplicationContext) override;
 };
\ No newline at end of file
diff --git a/Engine/Plugins/Runtime/MassGameplay/Source/MassReplication/Private/MassReplicationFragments.cpp b/Engine/Plugins/Runtime/MassGameplay/Source/MassReplication/Private/MassReplicationFragments.cpp
--- a/Engine/Plugins/Runtime/MassGameplay/Source/MassReplication/Private/MassReplicationFragments.cpp
+++ b/Engine/Plugins/Runtime/MassGameplay/Source/MassReplication/Private/MassReplicationFragments.cpp
@@ -2,6 +2,7 @@
 
 #include "MassReplicationFragments.h"
 #include "MassReplicationSubsystem.h"
+#include "MassReplicationProcessorBase.h"
 #include "Engine/World.h"
 
 //----------------------------------------------------------------------//
@@ -87,4 +88,7 @@ FMassReplicationSharedFragment::FMassReplicationSharedFragment(UMassReplicationS
 {
 	LODCalculator.Initialize(Params.LODDistance, Params.BufferHysteresisOnDistancePercentage / 100.0f, Params.LODMaxCount, Params.LODMaxCountPerViewer);
 	BubbleInfoClassHandle = ReplicationSubsystem.GetBubbleInfoClassHandle(Params.BubbleInfoClass);
+
+	CachedReplicator = Params.ReplicatorClass.GetDefaultObject();
+	checkf(CachedReplicator, TEXT("Expecting a valid replicator class"))
 }
\ No newline at end of file
diff --git a/Engine/Plugins/Runtime/MassGameplay/Source/MassReplication/Private/MassReplicationProcessorBase.cpp b/Engine/Plugins/Runtime/MassGameplay/Source/MassReplication/Private/MassReplicationProcessorBase.cpp
--- a/Engine/Plugins/Runtime/MassGameplay/Source/MassReplication/Private/MassReplicationProcessorBase.cpp
+++ b/Engine/Plugins/Runtime/MassGameplay/Source/MassReplication/Private/MassReplicationProcessorBase.cpp
@@ -5,16 +5,27 @@
 #include "MassLODManager.h"
 #include "MassCommonFragments.h"
 
+namespace UE::Mass::Replication
+{
+	int32 bDebugReplicationLOD = 0;
+	FAutoConsoleVariableRef CVarDebugReplicationViewerLOD(TEXT("ai.debug.ReplicationLOD"), bDebugReplicationLOD, TEXT("Debug Replication LOD"), ECVF_Cheat);
+} // UE::Mass::Crowd
+
 //----------------------------------------------------------------------//
-//  UMassReplicationProcessorBase
+//  UMassReplicationProcessor
 //----------------------------------------------------------------------//
-UMassReplicationProcessorBase::UMassReplicationProcessorBase()
+UMassReplicationProcessor::UMassReplicationProcessor()
 {
-	ExecutionFlags = int32(EProcessorExecutionFlags::None);
+#if !UE_ALLOW_DEBUG_REPLICATION_BUBBLES_STANDALONE
+	ExecutionFlags = int32(EProcessorExecutionFlags::Server);
+#else
+	ExecutionFlags = int32(EProcessorExecutionFlags::All);
+#endif // UE_ALLOW_DEBUG_REPLICATION_BUBBLES_STANDALONE
+
 	ProcessingPhase = EMassProcessingPhase::PostPhysics;
 }
 
-void UMassReplicationProcessorBase::ConfigureQueries()
+void UMassReplicationProcessor::ConfigureQueries()
 {
 	CollectViewerInfoQuery.AddRequirement<FTransformFragment>(EMassFragmentAccess::ReadOnly);
 	CollectViewerInfoQuery.AddRequirement<FMassReplicationViewerInfoFragment>(EMassFragmentAccess::ReadOnly);
@@ -42,7 +53,7 @@ void UMassReplicationProcessorBase::ConfigureQueries()
 	EntityQuery.AddSharedRequirement<FMassReplicationSharedFragment>(EMassFragmentAccess::ReadWrite);
 }
 
-void UMassReplicationProcessorBase::Initialize(UObject& Owner)
+void UMassReplicationProcessor::Initialize(UObject& Owner)
 {
 	Super::Initialize(Owner);
 
@@ -53,7 +64,7 @@ void UMassReplicationProcessorBase::Initialize(UObject& Owner)
 #endif //UE_REPLICATION_COMPILE_SERVER_CODE
 }
 
-void UMassReplicationProcessorBase::PrepareExecution(UMassEntitySubsystem& EntitySubsystem)
+void UMassReplicationProcessor::PrepareExecution(UMassEntitySubsystem& EntitySubsystem)
 {
 #if UE_REPLICATION_COMPILE_SERVER_CODE
 
@@ -66,6 +77,18 @@ void UMassReplicationProcessorBase::PrepareExecution(UMassEntitySubsystem& Entit
 	const TArray<FViewerInfo>& Viewers = LODManager->GetViewers();
 	EntitySubsystem.ForEachSharedFragment<FMassReplicationSharedFragment>([this, &Viewers](FMassReplicationSharedFragment& RepSharedFragment)
 	{
+		if (!RepSharedFragment.bEntityQueryInitialized)
+		{
+			RepSharedFragment.EntityQuery = EntityQuery;
+			RepSharedFragment.EntityQuery.SetArchetypeFilter([&RepSharedFragment](const FMassExecutionContext& Context)
+			{
+				const FMassReplicationSharedFragment& CurRepSharedFragment = Context.GetSharedFragment<FMassReplicationSharedFragment>();
+				return &CurRepSharedFragment == &RepSharedFragment;
+			});
+			RepSharedFragment.CachedReplicator->AddRequirements(RepSharedFragment.EntityQuery);
+			RepSharedFragment.bEntityQueryInitialized = true;
+		}
+
 		RepSharedFragment.LODCollector.PrepareExecution(Viewers);
 		RepSharedFragment.LODCalculator.PrepareExecution(Viewers);
 
@@ -128,41 +151,72 @@ void UMassReplicationProcessorBase::PrepareExecution(UMassEntitySubsystem& Entit
 #endif //UE_REPLICATION_COMPILE_SERVER_CODE
 }
 
-void UMassReplicationProcessorBase::Execute(UMassEntitySubsystem& EntitySubsystem, FMassExecutionContext& Context)
+void UMassReplicationProcessor::Execute(UMassEntitySubsystem& EntitySubsystem, FMassExecutionContext& Context)
 {
 #if UE_REPLICATION_COMPILE_SERVER_CODE
 	check(World);
+	check(LODManager);
+	check(ReplicationSubsystem);
 
-	PrepareExecution(EntitySubsystem);
-
-	CollectViewerInfoQuery.ForEachEntityChunk(EntitySubsystem, Context, [this](FMassExecutionContext& Context)
 	{
-		const TConstArrayView<FTransformFragment> LocationList = Context.GetFragmentView<FTransformFragment>();
-		const TArrayView<FMassReplicationViewerInfoFragment> ViewersInfoList = Context.GetMutableFragmentView<FMassReplicationViewerInfoFragment>();
-		FMassReplicationSharedFragment& RepSharedFragment = Context.GetMutableSharedFragment<FMassReplicationSharedFragment>();
-		RepSharedFragment.LODCollector.CollectLODInfo(Context, LocationList, ViewersInfoList, ViewersInfoList);
-	});
+		QUICK_SCOPE_CYCLE_COUNTER(UMassReplicationProcessor_Preperation);
+		PrepareExecution(EntitySubsystem);
+	}
 
-	CalculateLODQuery.ForEachEntityChunk(EntitySubsystem, Context, [this](FMassExecutionContext& Context)
 	{
-		const TConstArrayView<FMassReplicationViewerInfoFragment> ViewersInfoList = Context.GetFragmentView<FMassReplicationViewerInfoFragment>();
-		const TArrayView<FMassReplicationLODFragment> ViewerLODList = Context.GetMutableFragmentView<FMassReplicationLODFragment>();
-		FMassReplicationSharedFragment& RepSharedFragment = Context.GetMutableSharedFragment<FMassReplicationSharedFragment>();
-		RepSharedFragment.LODCalculator.CalculateLOD(Context, ViewersInfoList, ViewerLODList, ViewersInfoList);
-	});
+		QUICK_SCOPE_CYCLE_COUNTER(UMassReplicationProcessor_LODCaculation);
+		CollectViewerInfoQuery.ForEachEntityChunk(EntitySubsystem, Context, [this](FMassExecutionContext& Context)
+		{
+			const TConstArrayView<FTransformFragment> LocationList = Context.GetFragmentView<FTransformFragment>();
+			const TArrayView<FMassReplicationViewerInfoFragment> ViewersInfoList = Context.GetMutableFragmentView<FMassReplicationViewerInfoFragment>();
+			FMassReplicationSharedFragment& RepSharedFragment = Context.GetMutableSharedFragment<FMassReplicationSharedFragment>();
+			RepSharedFragment.LODCollector.CollectLODInfo(Context, LocationList, ViewersInfoList, ViewersInfoList);
+		});
 
-	EntitySubsystem.ForEachSharedFragment<FMassReplicationSharedFragment>([](FMassReplicationSharedFragment& RepSharedFragment)
-	{
-		RepSharedFragment.bHasAdjustedDistancesFromCount = RepSharedFragment.LODCalculator.AdjustDistancesFromCount();
-	});
+		CalculateLODQuery.ForEachEntityChunk(EntitySubsystem, Context, [this](FMassExecutionContext& Context)
+		{
+			const TConstArrayView<FMassReplicationViewerInfoFragment> ViewersInfoList = Context.GetFragmentView<FMassReplicationViewerInfoFragment>();
+			const TArrayView<FMassReplicationLODFragment> ViewerLODList = Context.GetMutableFragmentView<FMassReplicationLODFragment>();
+			FMassReplicationSharedFragment& RepSharedFragment = Context.GetMutableSharedFragment<FMassReplicationSharedFragment>();
+			RepSharedFragment.LODCalculator.CalculateLOD(Context, ViewersInfoList, ViewerLODList, ViewersInfoList);
+		});
+
+		EntitySubsystem.ForEachSharedFragment<FMassReplicationSharedFragment>([](FMassReplicationSharedFragment& RepSharedFragment)
+		{
+			RepSharedFragment.bHasAdjustedDistancesFromCount = RepSharedFragment.LODCalculator.AdjustDistancesFromCount();
+		});
+
+		AdjustLODDistancesQuery.ForEachEntityChunk(EntitySubsystem, Context, [this](FMassExecutionContext& Context)
+		{
+			const TConstArrayView<FMassReplicationViewerInfoFragment> ViewersInfoList = Context.GetFragmentView<FMassReplicationViewerInfoFragment>();
+			const TArrayView<FMassReplicationLODFragment> ViewerLODList = Context.GetMutableFragmentView<FMassReplicationLODFragment>();
+			FMassReplicationSharedFragment& RepSharedFragment = Context.GetMutableSharedFragment<FMassReplicationSharedFragment>();
+			RepSharedFragment.LODCalculator.AdjustLODFromCount(Context, ViewersInfoList, ViewerLODList, ViewersInfoList);
+		});
+	}
 
-	AdjustLODDistancesQuery.ForEachEntityChunk(EntitySubsystem, Context, [this](FMassExecutionContext& Context)
 	{
-		const TConstArrayView<FMassReplicationViewerInfoFragment> ViewersInfoList = Context.GetFragmentView<FMassReplicationViewerInfoFragment>();
-		const TArrayView<FMassReplicationLODFragment> ViewerLODList = Context.GetMutableFragmentView<FMassReplicationLODFragment>();
-		FMassReplicationSharedFragment& RepSharedFragment = Context.GetMutableSharedFragment<FMassReplicationSharedFragment>();
-		RepSharedFragment.LODCalculator.AdjustLODFromCount(Context, ViewersInfoList, ViewerLODList, ViewersInfoList);
-	});
+		QUICK_SCOPE_CYCLE_COUNTER(UMassReplicationProcessor_ProcessClientReplication);
+		FMassReplicationContext ReplicationContext(*World, *LODManager, *ReplicationSubsystem);
+		EntitySubsystem.ForEachSharedFragment<FMassReplicationSharedFragment>([this, &EntitySubsystem, &Context, &ReplicationContext](FMassReplicationSharedFragment& RepSharedFragment)
+		{
+			RepSharedFragment.EntityQuery.ForEachEntityChunk(EntitySubsystem, Context, [&ReplicationContext, &RepSharedFragment](FMassExecutionContext& Context)
+			{
+				RepSharedFragment.CachedReplicator->ProcessClientReplication(Context, ReplicationContext);
+			});
+		});
+	}
 
+	// Optional debug display
+	if (UE::Mass::Replication::bDebugReplicationLOD)
+	{
+		EntityQuery.ForEachEntityChunk(EntitySubsystem, Context, [this](FMassExecutionContext& Context)
+		{
+			const TConstArrayView<FTransformFragment> TransformList = Context.GetFragmentView<FTransformFragment>();
+			const TConstArrayView<FMassReplicationLODFragment> ViewerLODList = Context.GetFragmentView<FMassReplicationLODFragment>();
+			FMassReplicationSharedFragment& RepSharedFragment = Context.GetMutableSharedFragment<FMassReplicationSharedFragment>();
+			RepSharedFragment.LODCalculator.DebugDisplayLOD(Context, ViewerLODList, TransformList, World);
+		});
+	}
 #endif //UE_REPLICATION_COMPILE_SERVER_CODE
 }
\ No newline at end of file
diff --git a/Engine/Plugins/Runtime/MassGameplay/Source/MassReplication/Public/MassReplicationFragments.h b/Engine/Plugins/Runtime/MassGameplay/Source/MassReplication/Public/MassReplicationFragments.h
--- a/Engine/Plugins/Runtime/MassGameplay/Source/MassReplication/Public/MassReplicationFragments.h
+++ b/Engine/Plugins/Runtime/MassGameplay/Source/MassReplication/Public/MassReplicationFragments.h
@@ -10,6 +10,7 @@
 #include "MassReplicationFragments.generated.h"
 
 class AMassClientBubbleInfoBase;
+class UMassReplicatorBase;
 
 /**
  *  Fragment type for the mass network id of a mass entity
@@ -133,6 +134,10 @@ public:
 
 	UPROPERTY(EditAnywhere, Category = "Mass|Replication", config)
 	TSubclassOf<AMassClientBubbleInfoBase> BubbleInfoClass;
+
+	UPROPERTY(EditAnywhere, Category = "Mass|Replication", config)
+	TSubclassOf<UMassReplicatorBase> ReplicatorClass;
+
 };
 
 USTRUCT()
@@ -157,6 +162,12 @@ public:
 	TMassLODCalculator<FReplicationLODLogic> LODCalculator;
 	bool bHasAdjustedDistancesFromCount = false;
 
+	bool bEntityQueryInitialized = false;
+	FMassEntityQuery EntityQuery;
+
+	UPROPERTY(Transient)
+	mutable UMassReplicatorBase* CachedReplicator = nullptr;
+
 	template<typename T>
 	T& GetTypedClientBubbleInfoChecked(FMassClientHandle Handle)
 	{
diff --git a/Engine/Plugins/Runtime/MassGameplay/Source/MassReplication/Public/MassReplicationProcessorBase.h b/Engine/Plugins/Runtime/MassGameplay/Source/MassReplication/Public/MassReplicationProcessorBase.h
--- a/Engine/Plugins/Runtime/MassGameplay/Source/MassReplication/Public/MassReplicationProcessorBase.h
+++ b/Engine/Plugins/Runtime/MassGameplay/Source/MassReplication/Public/MassReplicationProcessorBase.h
@@ -22,12 +22,12 @@ class UWorld;
  *  query Mass entity fragments and set those values for replication when appropriate, using the MassClientBubbleHandler.
  */
 UCLASS()
-class MASSREPLICATION_API UMassReplicationProcessorBase : public UMassProcessor_LODBase
+class MASSREPLICATION_API UMassReplicationProcessor : public UMassProcessor_LODBase
 {
 	GENERATED_BODY()
 
 public:
-	UMassReplicationProcessorBase();
+	UMassReplicationProcessor();
 
 protected:
 	virtual void ConfigureQueries() override;
@@ -36,19 +36,10 @@ protected:
 
 	void PrepareExecution(UMassEntitySubsystem& EntitySubsystem);
 
-	/** 
-	 *  Implemented as straight template callbacks as when profiled this was faster than TFunctionRef. Its probably easier to pass Lamdas in to these
-	 *  but Functors can also be used as well as TFunctionRefs etc. Its also fairly straight forward to call member functions via some Lamda glue code
-	 */
-	template<typename AgentArrayItem, typename CacheViewsCallback, typename AddEntityCallback, typename ModifyEntityCallback, typename RemoveEntityCallback>
-	void CalculateClientReplication(FMassExecutionContext& Context, CacheViewsCallback&& CacheViews, AddEntityCallback&& AddEntity, ModifyEntityCallback&& ModifyEntity, RemoveEntityCallback&& RemoveEntity);
-
-
-
 protected:
 
-	UPROPERTY()
-	UMassReplicationSubsystem* ReplicationSubsystem;
+	UPROPERTY(Transient)
+	UMassReplicationSubsystem* ReplicationSubsystem = nullptr;
 
 	FMassEntityQuery CollectViewerInfoQuery;
 	FMassEntityQuery CalculateLODQuery;
@@ -56,15 +47,52 @@ protected:
 	FMassEntityQuery EntityQuery;
 };
 
+
+struct FMassReplicationContext
+{
+	FMassReplicationContext(UWorld& InWorld, UMassLODManager& InLODManager, UMassReplicationSubsystem& InReplicationSubsystem)
+		: World(InWorld)
+		, LODManager(InLODManager)
+		, ReplicationSubsystem(InReplicationSubsystem)
+	{}
+
+	UWorld& World;
+	UMassLODManager& LODManager;
+	UMassReplicationSubsystem& ReplicationSubsystem;
+};
+
+UCLASS(Abstract)
+class MASSREPLICATION_API UMassReplicatorBase : public UObject
+{
+	GENERATED_BODY()
+
+public:
+	/**
+	 * Must override to add specific entity query requirements for replication
+	 * Usually we add replication processor handler requirements
+	 */
+	virtual void AddRequirements(FMassEntityQuery& EntityQuery) PURE_VIRTUAL(UMassReplicatorBase::ConfigureQueries, );
+
+	/**
+	 * Must override to process the client replication
+	 * This methods should call CalculateClientReplication with the appropriate callback implementation
+	 */
+	virtual void ProcessClientReplication(FMassExecutionContext& Context, FMassReplicationContext& ReplicationContext) PURE_VIRTUAL(UMassReplicatorBase::ProcessClientReplication, );
+
+protected:
+	/**
+	 *  Implemented as straight template callbacks as when profiled this was faster than TFunctionRef. Its probably easier to pass Lamdas in to these
+	 *  but Functors can also be used as well as TFunctionRefs etc. Its also fairly straight forward to call member functions via some Lamda glue code
+	 */
+	template<typename AgentArrayItem, typename CacheViewsCallback, typename AddEntityCallback, typename ModifyEntityCallback, typename RemoveEntityCallback>
+	static void CalculateClientReplication(FMassExecutionContext& Context, FMassReplicationContext& ReplicationContext, CacheViewsCallback&& CacheViews, AddEntityCallback&& AddEntity, ModifyEntityCallback&& ModifyEntity, RemoveEntityCallback&& RemoveEntity);
+};
+
 template<typename AgentArrayItem, typename CacheViewsCallback, typename AddEntityCallback, typename ModifyEntityCallback, typename RemoveEntityCallback>
-void UMassReplicationProcessorBase::CalculateClientReplication(FMassExecutionContext& Context, CacheViewsCallback&& CacheViews, AddEntityCallback&& AddEntity, ModifyEntityCallback&& ModifyEntity, RemoveEntityCallback&& RemoveEntity)
+void UMassReplicatorBase::CalculateClientReplication(FMassExecutionContext& Context, FMassReplicationContext& ReplicationContext, CacheViewsCallback&& CacheViews, AddEntityCallback&& AddEntity, ModifyEntityCallback&& ModifyEntity, RemoveEntityCallback&& RemoveEntity)
 {
 #if UE_REPLICATION_COMPILE_SERVER_CODE
 
-	check(ReplicationSubsystem);
-	check(LODManager);
-	check(World);
-
 	const int32 NumEntities = Context.GetNumEntities();
 
 	TConstArrayView<FMassNetworkIDFragment> NetworkIDList = Context.GetFragmentView<FMassNetworkIDFragment>();
@@ -75,7 +103,7 @@ void UMassReplicationProcessorBase::CalculateClientReplication(FMassExecutionCon
 
 	CacheViews(Context);
 
-	const float Time = World->GetRealTimeSeconds();
+	const float Time = ReplicationContext.World.GetRealTimeSeconds();
 
 	for (int EntityIdx = 0; EntityIdx < NumEntities; EntityIdx++)
 	{
@@ -109,14 +137,14 @@ void UMassReplicationProcessorBase::CalculateClientReplication(FMassExecutionCon
 
 				//now we need to see what our highest viewer LOD is on this client (split screen etc), we can use the unsafe version as we have already checked all the
 				//CachedClientHandles for validity
-				const FClientViewerHandles& ClientViewers = ReplicationSubsystem->GetClientViewersChecked(ClientHandle);
+				const FClientViewerHandles& ClientViewers = ReplicationContext.ReplicationSubsystem.GetClientViewersChecked(ClientHandle);
 
 				EMassLOD::Type HighestLOD = EMassLOD::Off;
 
 				for (const FMassViewerHandle& ViewerHandle : ClientViewers.Handles)
 				{
 					//this should always we valid as we synchronized the viewers just previously
-					check(LODManager->IsValidViewer(ViewerHandle));
+					check(ReplicationContext.LODManager.IsValidViewer(ViewerHandle));
 
 					const EMassLOD::Type MassLOD = ViewerLODList[EntityIdx].LODPerViewer[ViewerHandle.GetIndex()];
 					check(MassLOD <= EMassLOD::Off);
