Change Config member variable to Params to match new structure name.

#rnx

diff --git a/Engine/Plugins/AI/MassCrowd/Config/BaseMassCrowd.ini b/Engine/Plugins/AI/MassCrowd/Config/BaseMassCrowd.ini
--- a/Engine/Plugins/AI/MassCrowd/Config/BaseMassCrowd.ini
+++ b/Engine/Plugins/AI/MassCrowd/Config/BaseMassCrowd.ini
@@ -4,3 +4,5 @@
 
 +StructRedirects=(OldName="TagFragment_MassCrowd",NewName="MassCrowdTag")
 +StructRedirects=(OldName="/Script/MassCrowd.TagFragment_MassCrowd",NewName="/Script/MassCrowd.MassCrowdTag")
+
++PropertyRedirects=(OldName="MassCrowdServerRepresentationTrait.Config", NewName="Params")
diff --git a/Engine/Plugins/AI/MassCrowd/Source/MassCrowd/Private/MassCrowdServerRepresentationTrait.cpp b/Engine/Plugins/AI/MassCrowd/Source/MassCrowd/Private/MassCrowdServerRepresentationTrait.cpp
--- a/Engine/Plugins/AI/MassCrowd/Source/MassCrowd/Private/MassCrowdServerRepresentationTrait.cpp
+++ b/Engine/Plugins/AI/MassCrowd/Source/MassCrowd/Private/MassCrowdServerRepresentationTrait.cpp
@@ -11,16 +11,16 @@
 
 UMassCrowdServerRepresentationTrait::UMassCrowdServerRepresentationTrait()
 {
-	Config.RepresentationActorManagementClass = UMassCrowdRepresentationActorManagement::StaticClass();
-	Config.LODRepresentation[EMassLOD::High] = EMassRepresentationType::HighResSpawnedActor;
-	Config.LODRepresentation[EMassLOD::Medium] = EMassRepresentationType::None;
-	Config.LODRepresentation[EMassLOD::Low] = EMassRepresentationType::None;
-	Config.LODRepresentation[EMassLOD::Off] = EMassRepresentationType::None;
-	Config.bKeepLowResActors = false;
-	Config.bKeepActorExtraFrame = false;
-	Config.bSpreadFirstVisualizationUpdate = false;
-	Config.WorldPartitionGridNameContainingCollision = NAME_None;
-	Config.NotVisibleUpdateRate = 0.5f;
+	Params.RepresentationActorManagementClass = UMassCrowdRepresentationActorManagement::StaticClass();
+	Params.LODRepresentation[EMassLOD::High] = EMassRepresentationType::HighResSpawnedActor;
+	Params.LODRepresentation[EMassLOD::Medium] = EMassRepresentationType::None;
+	Params.LODRepresentation[EMassLOD::Low] = EMassRepresentationType::None;
+	Params.LODRepresentation[EMassLOD::Off] = EMassRepresentationType::None;
+	Params.bKeepLowResActors = false;
+	Params.bKeepActorExtraFrame = false;
+	Params.bSpreadFirstVisualizationUpdate = false;
+	Params.WorldPartitionGridNameContainingCollision = NAME_None;
+	Params.NotVisibleUpdateRate = 0.5f;
 }
 
 void UMassCrowdServerRepresentationTrait::BuildTemplate(FMassEntityTemplateBuildContext& BuildContext, UWorld& World) const
@@ -48,10 +48,10 @@ void UMassCrowdServerRepresentationTrait::BuildTemplate(FMassEntityTemplateBuild
 	FSharedStruct SubsystemFragment = EntitySubsystem->GetOrCreateSharedFragment<FMassRepresentationSubsystemSharedFragment>(SubsystemHash, Subsystem);
 	BuildContext.AddSharedFragment(SubsystemFragment);
 
-	uint32 ConfigHash = UE::StructUtils::GetStructCrc32(FConstStructView::Make(Config));
-	FConstSharedStruct ConfigFragment = EntitySubsystem->GetOrCreateConstSharedFragment<FMassRepresentationParameters>(ConfigHash, Config);
-	ConfigFragment.Get<FMassRepresentationParameters>().ComputeCachedValues();
-	BuildContext.AddConstSharedFragment(ConfigFragment);
+	uint32 ParamsHash = UE::StructUtils::GetStructCrc32(FConstStructView::Make(Params));
+	FConstSharedStruct ParamsFragment = EntitySubsystem->GetOrCreateConstSharedFragment<FMassRepresentationParameters>(ParamsHash, Params);
+	ParamsFragment.Get<FMassRepresentationParameters>().ComputeCachedValues();
+	BuildContext.AddConstSharedFragment(ParamsFragment);
 
 	FMassRepresentationFragment& RepresentationFragment = BuildContext.AddFragment_GetRef<FMassRepresentationFragment>();
 	RepresentationFragment.StaticMeshDescIndex = INDEX_NONE;
diff --git a/Engine/Plugins/AI/MassCrowd/Source/MassCrowd/Private/MassCrowdVisualizationTrait.cpp b/Engine/Plugins/AI/MassCrowd/Source/MassCrowd/Private/MassCrowdVisualizationTrait.cpp
--- a/Engine/Plugins/AI/MassCrowd/Source/MassCrowd/Private/MassCrowdVisualizationTrait.cpp
+++ b/Engine/Plugins/AI/MassCrowd/Source/MassCrowd/Private/MassCrowdVisualizationTrait.cpp
@@ -9,15 +9,15 @@ UMassCrowdVisualizationTrait::UMassCrowdVisualizationTrait()
 {
 	// Override the subsystem to support parallelization of the crowd
 	RepresentationSubsystemClass = UMassCrowdRepresentationSubsystem::StaticClass();
-	Config.RepresentationActorManagementClass = UMassCrowdRepresentationActorManagement::StaticClass();
-	Config.LODRepresentation[EMassLOD::High] = EMassRepresentationType::HighResSpawnedActor;
-	Config.LODRepresentation[EMassLOD::Medium] = EMassRepresentationType::LowResSpawnedActor;
-	Config.LODRepresentation[EMassLOD::Low] = EMassRepresentationType::StaticMeshInstance;
-	Config.LODRepresentation[EMassLOD::Off] = EMassRepresentationType::None;
+	Params.RepresentationActorManagementClass = UMassCrowdRepresentationActorManagement::StaticClass();
+	Params.LODRepresentation[EMassLOD::High] = EMassRepresentationType::HighResSpawnedActor;
+	Params.LODRepresentation[EMassLOD::Medium] = EMassRepresentationType::LowResSpawnedActor;
+	Params.LODRepresentation[EMassLOD::Low] = EMassRepresentationType::StaticMeshInstance;
+	Params.LODRepresentation[EMassLOD::Off] = EMassRepresentationType::None;
 	// Set bKeepLowResActor to true as a spawning optimization, this will keep the low-res actor if available while showing the static mesh instance
-	Config.bKeepLowResActors = true;
-	Config.bKeepActorExtraFrame = true;
-	Config.bSpreadFirstVisualizationUpdate = false;
-	Config.WorldPartitionGridNameContainingCollision = NAME_None;
-	Config.NotVisibleUpdateRate = 0.5f;
+	Params.bKeepLowResActors = true;
+	Params.bKeepActorExtraFrame = true;
+	Params.bSpreadFirstVisualizationUpdate = false;
+	Params.WorldPartitionGridNameContainingCollision = NAME_None;
+	Params.NotVisibleUpdateRate = 0.5f;
 }
diff --git a/Engine/Plugins/AI/MassCrowd/Source/MassCrowd/Public/MassCrowdServerRepresentationTrait.h b/Engine/Plugins/AI/MassCrowd/Source/MassCrowd/Public/MassCrowdServerRepresentationTrait.h
--- a/Engine/Plugins/AI/MassCrowd/Source/MassCrowd/Public/MassCrowdServerRepresentationTrait.h
+++ b/Engine/Plugins/AI/MassCrowd/Source/MassCrowd/Public/MassCrowdServerRepresentationTrait.h
@@ -27,5 +27,5 @@ protected:
 
 	/** Configuration parameters for the representation processor */
 	UPROPERTY(EditAnywhere, Category = "Mass|Visual")
-	FMassRepresentationParameters Config;
+	FMassRepresentationParameters Params;
 };
diff --git a/Engine/Plugins/Runtime/MassGameplay/Config/BaseMassGameplay.ini b/Engine/Plugins/Runtime/MassGameplay/Config/BaseMassGameplay.ini
--- a/Engine/Plugins/Runtime/MassGameplay/Config/BaseMassGameplay.ini
+++ b/Engine/Plugins/Runtime/MassGameplay/Config/BaseMassGameplay.ini
@@ -36,12 +36,15 @@
 +StructRedirects=(OldName="MassRepresentationSubsystemFragment", NewName="MassRepresentationSubsystemSharedFragment")
 +StructRedirects=(OldName="/Script/MassRepresentation.MassRepresentationConfig", NewName="/Script/MassRepresentation.MassRepresentationParameters")
 +StructRedirects=(OldName="/Script/MassRepresentation.MassRepresentationSubsystemFragment", NewName="/Script/MassRepresentation.MassRepresentationSubsystemSharedFragment")
++PropertyRedirects=(OldName="MassVisualizationTrait.Config", NewName="Params")
 
 ;MassLOD
 +StructRedirects=(OldName="MassSimulationVariableTickConfig", NewName="MassSimulationVariableTickParameters")
 +StructRedirects=(OldName="MassSimulationLODConfig", NewName="MassSimulationLODParameters")
 +StructRedirects=(OldName="/Script/MassLOD.MassSimulationVariableTickConfig", NewName="/Script/MassLOD.MassSimulationVariableTickParameters")
 +StructRedirects=(OldName="/Script/MassLOD.MassSimulationLODConfig", NewName="/Script/MassLOD.MassSimulationLODParameters")
++PropertyRedirects=(OldName="MassSimulationLODTrait.Config", NewName="Params")
++PropertyRedirects=(OldName="MassSimulationLODTrait.VariableTickConfig", NewName="VariableTickParams")
 
 ;MassReplication
 +ClassRedirects=(OldName="/Script/MassReplication.MassReplicationManager", NewName="/Script/MassReplication.MassReplicationSubsystem")
diff --git a/Engine/Plugins/Runtime/MassGameplay/Source/MassLOD/Private/MassLODTrait.cpp b/Engine/Plugins/Runtime/MassGameplay/Source/MassLOD/Private/MassLODTrait.cpp
--- a/Engine/Plugins/Runtime/MassGameplay/Source/MassLOD/Private/MassLODTrait.cpp
+++ b/Engine/Plugins/Runtime/MassGameplay/Source/MassLOD/Private/MassLODTrait.cpp
@@ -21,7 +21,7 @@ void UMassSimulationLODTrait::BuildTemplate(FMassEntityTemplateBuildContext& Bui
 	FMassSimulationLODFragment& LODFragment = BuildContext.AddFragment_GetRef<FMassSimulationLODFragment>();
 
 	// Start all simulation LOD in the Off 
-	if(Config.bSetLODTags || bEnableVariableTicking)
+	if(Params.bSetLODTags || bEnableVariableTicking)
 	{
 		LODFragment.LOD = EMassLOD::Off;
 		BuildContext.AddTag<FMassOffLODTag>();
@@ -30,10 +30,10 @@ void UMassSimulationLODTrait::BuildTemplate(FMassEntityTemplateBuildContext& Bui
 	UMassEntitySubsystem* EntitySubsystem = UWorld::GetSubsystem<UMassEntitySubsystem>(&World);
 	check(EntitySubsystem);
 
-	uint32 ConfigHash = UE::StructUtils::GetStructCrc32(FConstStructView::Make(Config));
-	FConstSharedStruct ConfigFragment = EntitySubsystem->GetOrCreateConstSharedFragment(ConfigHash, Config);
-	BuildContext.AddConstSharedFragment(ConfigFragment);
-	FSharedStruct SharedFragment = EntitySubsystem->GetOrCreateSharedFragment<FMassSimulationLODSharedFragment>(ConfigHash, Config);
+	uint32 ParamsHash = UE::StructUtils::GetStructCrc32(FConstStructView::Make(Params));
+	FConstSharedStruct ParamsFragment = EntitySubsystem->GetOrCreateConstSharedFragment(ParamsHash, Params);
+	BuildContext.AddConstSharedFragment(ParamsFragment);
+	FSharedStruct SharedFragment = EntitySubsystem->GetOrCreateSharedFragment<FMassSimulationLODSharedFragment>(ParamsHash, Params);
 	BuildContext.AddSharedFragment(SharedFragment);
 
 	// Variable ticking from simulation LOD
@@ -42,10 +42,10 @@ void UMassSimulationLODTrait::BuildTemplate(FMassEntityTemplateBuildContext& Bui
 		BuildContext.AddFragment<FMassSimulationVariableTickFragment>();
 		BuildContext.AddChunkFragment<FMassSimulationVariableTickChunkFragment>();
 
-		uint32 VariableTickConfigHash = UE::StructUtils::GetStructCrc32(FConstStructView::Make(VariableTickConfig));
-		FConstSharedStruct VariableTickConfigFragment = EntitySubsystem->GetOrCreateConstSharedFragment(VariableTickConfigHash, VariableTickConfig);
-		BuildContext.AddConstSharedFragment(VariableTickConfigFragment);
-		FSharedStruct VariableTickSharedFragment = EntitySubsystem->GetOrCreateSharedFragment<FMassSimulationVariableTickSharedFragment>(VariableTickConfigHash, VariableTickConfig);
+		uint32 VariableTickParamsHash = UE::StructUtils::GetStructCrc32(FConstStructView::Make(VariableTickParams));
+		FConstSharedStruct VariableTickParamsFragment = EntitySubsystem->GetOrCreateConstSharedFragment(VariableTickParamsHash, VariableTickParams);
+		BuildContext.AddConstSharedFragment(VariableTickParamsFragment);
+		FSharedStruct VariableTickSharedFragment = EntitySubsystem->GetOrCreateSharedFragment<FMassSimulationVariableTickSharedFragment>(VariableTickParamsHash, VariableTickParams);
 		BuildContext.AddSharedFragment(VariableTickSharedFragment);
 	}
 }
diff --git a/Engine/Plugins/Runtime/MassGameplay/Source/MassLOD/Private/MassSimulationLOD.cpp b/Engine/Plugins/Runtime/MassGameplay/Source/MassLOD/Private/MassSimulationLOD.cpp
--- a/Engine/Plugins/Runtime/MassGameplay/Source/MassLOD/Private/MassSimulationLOD.cpp
+++ b/Engine/Plugins/Runtime/MassGameplay/Source/MassLOD/Private/MassSimulationLOD.cpp
@@ -38,17 +38,17 @@ FMassSimulationVariableTickParameters::FMassSimulationVariableTickParameters()
 //-----------------------------------------------------------------------------
 // FMassSimulationLODSharedFragment
 //-----------------------------------------------------------------------------
-FMassSimulationLODSharedFragment::FMassSimulationLODSharedFragment(const FMassSimulationLODParameters& LODConfig)
+FMassSimulationLODSharedFragment::FMassSimulationLODSharedFragment(const FMassSimulationLODParameters& LODParams)
 {
-	LODCalculator.Initialize(LODConfig.LODDistance, LODConfig.BufferHysteresisOnDistancePercentage / 100.0f, LODConfig.LODMaxCount);
+	LODCalculator.Initialize(LODParams.LODDistance, LODParams.BufferHysteresisOnDistancePercentage / 100.0f, LODParams.LODMaxCount);
 }
 
 //-----------------------------------------------------------------------------
 // FMassSimulationLODSharedFragment
 //-----------------------------------------------------------------------------
-FMassSimulationVariableTickSharedFragment::FMassSimulationVariableTickSharedFragment(const FMassSimulationVariableTickParameters& TickRateConfig)
+FMassSimulationVariableTickSharedFragment::FMassSimulationVariableTickSharedFragment(const FMassSimulationVariableTickParameters& TickRateParams)
 {
-	LODTickRateController.Initialize(TickRateConfig.TickRates, TickRateConfig.bSpreadFirstSimulationUpdate);
+	LODTickRateController.Initialize(TickRateParams.TickRates, TickRateParams.bSpreadFirstSimulationUpdate);
 }
 
 //-----------------------------------------------------------------------------
@@ -101,8 +101,8 @@ void UMassSimulationLODProcessor::ConfigureQueries()
 	EntityQuerySetLODTag.AddConstSharedRequirement<FMassSimulationLODParameters>();
 	EntityQuerySetLODTag.SetArchetypeFilter([](const FMassExecutionContext& Context)
 	{
-		const FMassSimulationLODParameters& LODConfig = Context.GetConstSharedFragment<FMassSimulationLODParameters>();
-		return LODConfig.bSetLODTags;
+		const FMassSimulationLODParameters& LODParams = Context.GetConstSharedFragment<FMassSimulationLODParameters>();
+		return LODParams.bSetLODTags;
 	});
 }
 
diff --git a/Engine/Plugins/Runtime/MassGameplay/Source/MassLOD/Public/MassLODTrait.h b/Engine/Plugins/Runtime/MassGameplay/Source/MassLOD/Public/MassLODTrait.h
--- a/Engine/Plugins/Runtime/MassGameplay/Source/MassLOD/Public/MassLODTrait.h
+++ b/Engine/Plugins/Runtime/MassGameplay/Source/MassLOD/Public/MassLODTrait.h
@@ -22,13 +22,13 @@ class MASSLOD_API UMassSimulationLODTrait : public UMassEntityTraitBase
 	GENERATED_BODY()
 
 	UPROPERTY(Category = "Config", EditAnywhere)
-	FMassSimulationLODParameters Config;
+	FMassSimulationLODParameters Params;
 
 	UPROPERTY(Category = "Config", EditAnywhere)
 	bool bEnableVariableTicking = false;
 
 	UPROPERTY(Category = "Config", EditAnywhere, meta = (EditCondition = "bEnableVariableTicking", EditConditionHides))
-	FMassSimulationVariableTickParameters VariableTickConfig;
+	FMassSimulationVariableTickParameters VariableTickParams;
 
 protected:
 	virtual void BuildTemplate(FMassEntityTemplateBuildContext& BuildContext, UWorld& World) const override;
diff --git a/Engine/Plugins/Runtime/MassGameplay/Source/MassLOD/Public/MassSimulationLOD.h b/Engine/Plugins/Runtime/MassGameplay/Source/MassLOD/Public/MassSimulationLOD.h
--- a/Engine/Plugins/Runtime/MassGameplay/Source/MassLOD/Public/MassSimulationLOD.h
+++ b/Engine/Plugins/Runtime/MassGameplay/Source/MassLOD/Public/MassSimulationLOD.h
@@ -106,9 +106,9 @@ struct MASSLOD_API FMassSimulationLODSharedFragment : public FMassSharedFragment
 	GENERATED_BODY()
 
 	FMassSimulationLODSharedFragment() = default;
-	FMassSimulationLODSharedFragment(const FMassSimulationLODParameters& Config);
+	FMassSimulationLODSharedFragment(const FMassSimulationLODParameters& LODParams);
 
-	/** Runtime data for matching the simulation LOD config */
+	/** Runtime data for matching the simulation LOD parameters */
 	TMassLODCalculator<FMassSimulationLODLogic> LODCalculator;
 	bool bHasAdjustedDistancesFromCount = false;
 };
@@ -119,9 +119,9 @@ struct MASSLOD_API FMassSimulationVariableTickSharedFragment : public FMassShare
 	GENERATED_BODY()
 
 	FMassSimulationVariableTickSharedFragment() = default;
-	FMassSimulationVariableTickSharedFragment(const FMassSimulationVariableTickParameters& Config);
+	FMassSimulationVariableTickSharedFragment(const FMassSimulationVariableTickParameters& TickRateParams);
 
-	/** Runtime data for matching the simulation tick rate config */
+	/** Runtime data for matching the simulation tick rate parameters */
 	TMassLODTickRateController<FMassSimulationVariableTickChunkFragment, FMassSimulationLODLogic> LODTickRateController;
 
 	static bool ShouldCalculateLODForChunk(const FMassExecutionContext& Context)
@@ -156,7 +156,7 @@ protected:
 	virtual void ConfigureQueries() override;
 	virtual void Execute(UMassEntitySubsystem& EntitySubsystem, FMassExecutionContext& Context) override;
 
-	void CalculateLODForConfig(UMassEntitySubsystem& EntitySubsystem, FMassExecutionContext& Context, FMassSimulationLODParameters& LODConfig);
+	void CalculateLODForConfig(UMassEntitySubsystem& EntitySubsystem, FMassExecutionContext& Context, FMassSimulationLODParameters& LODParams);
 
 	FMassEntityQuery EntityQuery;
 	FMassEntityQuery EntityQueryCalculateLOD;
diff --git a/Engine/Plugins/Runtime/MassGameplay/Source/MassRepresentation/Private/MassRepresentationProcessor.cpp b/Engine/Plugins/Runtime/MassGameplay/Source/MassRepresentation/Private/MassRepresentationProcessor.cpp
--- a/Engine/Plugins/Runtime/MassGameplay/Source/MassRepresentation/Private/MassRepresentationProcessor.cpp
+++ b/Engine/Plugins/Runtime/MassGameplay/Source/MassRepresentation/Private/MassRepresentationProcessor.cpp
@@ -49,8 +49,8 @@ void UMassRepresentationProcessor::UpdateRepresentation(FMassExecutionContext& C
 {
 	UMassRepresentationSubsystem* RepresentationSubsystem = Context.GetMutableSharedFragment<FMassRepresentationSubsystemSharedFragment>().RepresentationSubsystem;
 	check(RepresentationSubsystem);
-	const FMassRepresentationParameters& RepresentationConfig = Context.GetConstSharedFragment<FMassRepresentationParameters>();
-	UMassRepresentationActorManagement* RepresentationActorManagement = RepresentationConfig.CachedRepresentationActorManagement;
+	const FMassRepresentationParameters& RepresentationParams = Context.GetConstSharedFragment<FMassRepresentationParameters>();
+	UMassRepresentationActorManagement* RepresentationActorManagement = RepresentationParams.CachedRepresentationActorManagement;
 	check(RepresentationActorManagement);
 
 	const TConstArrayView<FTransformFragment> TransformList = Context.GetFragmentView<FTransformFragment>();
@@ -58,7 +58,7 @@ void UMassRepresentationProcessor::UpdateRepresentation(FMassExecutionContext& C
 	const TConstArrayView<FMassRepresentationLODFragment> RepresentationLODList = Context.GetFragmentView<FMassRepresentationLODFragment>();
 	const TArrayView<FMassActorFragment> ActorList = Context.GetMutableFragmentView<FMassActorFragment>();
 
-	const bool bDoKeepActorExtraFrame = UE::MassRepresentation::bAllowKeepActorExtraFrame ? RepresentationConfig.bKeepLowResActors : false;
+	const bool bDoKeepActorExtraFrame = UE::MassRepresentation::bAllowKeepActorExtraFrame ? RepresentationParams.bKeepLowResActors : false;
 
 	const int32 NumEntities = Context.GetNumEntities();
 	for (int32 EntityIdx = 0; EntityIdx < NumEntities; EntityIdx++)
@@ -73,13 +73,13 @@ void UMassRepresentationProcessor::UpdateRepresentation(FMassExecutionContext& C
 		const EMassRepresentationType PrevRepresentationCopy = Representation.PrevRepresentation;
 		Representation.PrevRepresentation = Representation.CurrentRepresentation;
 		
-		EMassRepresentationType WantedRepresentationType = RepresentationConfig.LODRepresentation[FMath::Min((int32)RepresentationLOD.LOD, (int32)EMassLOD::Off)];
+		EMassRepresentationType WantedRepresentationType = RepresentationParams.LODRepresentation[FMath::Min((int32)RepresentationLOD.LOD, (int32)EMassLOD::Off)];
 
 		// Make sure we do not have actor spawned in areas not fully loaded
 		if ((WantedRepresentationType == EMassRepresentationType::HighResSpawnedActor || WantedRepresentationType == EMassRepresentationType::LowResSpawnedActor) &&
-			!RepresentationSubsystem->IsCollisionLoaded(RepresentationConfig.WorldPartitionGridNameContainingCollision, TransformFragment.GetTransform()))
+			!RepresentationSubsystem->IsCollisionLoaded(RepresentationParams.WorldPartitionGridNameContainingCollision, TransformFragment.GetTransform()))
 		{
-			WantedRepresentationType = RepresentationConfig.CachedDefaultRepresentationType;
+			WantedRepresentationType = RepresentationParams.CachedDefaultRepresentationType;
 		}
 		
 		auto DisableActorForISM = [&](AActor*& Actor)
@@ -88,7 +88,7 @@ void UMassRepresentationProcessor::UpdateRepresentation(FMassExecutionContext& C
 			{
 				// Execute only if the high res is different than the low res Actor 
 				// Or if we do not wish to keep the low res actor while in ISM
-				if (Representation.HighResTemplateActorIndex != Representation.LowResTemplateActorIndex || !RepresentationConfig.bKeepLowResActors)
+				if (Representation.HighResTemplateActorIndex != Representation.LowResTemplateActorIndex || !RepresentationParams.bKeepLowResActors)
 				{
 					// Try releasing the high actor or any high res spawning request
 					if (ReleaseActorOrCancelSpawning(*RepresentationSubsystem, MassAgent, ActorInfo, Representation.HighResTemplateActorIndex, Representation.ActorSpawnRequestHandle, Context.Defer()))
@@ -96,7 +96,7 @@ void UMassRepresentationProcessor::UpdateRepresentation(FMassExecutionContext& C
 						Actor = ActorInfo.GetOwnedByMassMutable();
 					}
 					// Do not do the same with low res if indicated so
-					if (!RepresentationConfig.bKeepLowResActors && ReleaseActorOrCancelSpawning(*RepresentationSubsystem, MassAgent, ActorInfo, Representation.LowResTemplateActorIndex, Representation.ActorSpawnRequestHandle, Context.Defer()))
+					if (!RepresentationParams.bKeepLowResActors && ReleaseActorOrCancelSpawning(*RepresentationSubsystem, MassAgent, ActorInfo, Representation.LowResTemplateActorIndex, Representation.ActorSpawnRequestHandle, Context.Defer()))
 					{
 						Actor = ActorInfo.GetOwnedByMassMutable();
 					}
@@ -177,7 +177,7 @@ void UMassRepresentationProcessor::UpdateRepresentation(FMassExecutionContext& C
 					}
 					else if (!Actor)
 					{
-						Representation.CurrentRepresentation = RepresentationConfig.CachedDefaultRepresentationType;
+						Representation.CurrentRepresentation = RepresentationParams.CachedDefaultRepresentationType;
 					}
 					break;
 				}
@@ -297,7 +297,7 @@ void UMassVisualizationProcessor::ConfigureQueries()
 
 FMassVisualizationChunkFragment& UMassVisualizationProcessor::UpdateChunkVisibility(FMassExecutionContext& Context) const
 {
-	const FMassRepresentationParameters& RepresentationConfig = Context.GetConstSharedFragment<FMassRepresentationParameters>();
+	const FMassRepresentationParameters& RepresentationParams = Context.GetConstSharedFragment<FMassRepresentationParameters>();
 	bool bFirstUpdate = false;
 
 	// Setup chunk fragment data about visibility
@@ -308,7 +308,7 @@ FMassVisualizationChunkFragment& UMassVisualizationProcessor::UpdateChunkVisibil
 		// The visibility on the chunk fragment data isn't set yet, let see if the Archetype has an visibility tag and set it on the ChunkData
 		ChunkVisibility = UE::MassRepresentation::GetVisibilityFromArchetype(Context);
 		ChunkData.SetVisibility(ChunkVisibility);
-		bFirstUpdate = RepresentationConfig.bSpreadFirstVisualizationUpdate;
+		bFirstUpdate = RepresentationParams.bSpreadFirstVisualizationUpdate;
 	}
 	else
 	{
@@ -321,13 +321,13 @@ FMassVisualizationChunkFragment& UMassVisualizationProcessor::UpdateChunkVisibil
 		if (bFirstUpdate)
 		{
 			// A DeltaTime of 0.0f means it will tick this frame.
-			DeltaTime = FMath::RandRange(0.0f, RepresentationConfig.NotVisibleUpdateRate);
+			DeltaTime = FMath::RandRange(0.0f, RepresentationParams.NotVisibleUpdateRate);
 		}
 		else 
 		{
 			if (DeltaTime < 0.0f)
 			{
-				DeltaTime += RepresentationConfig.NotVisibleUpdateRate * (1.0f + FMath::RandRange(-0.1f, 0.1f));
+				DeltaTime += RepresentationParams.NotVisibleUpdateRate * (1.0f + FMath::RandRange(-0.1f, 0.1f));
 			}
 			DeltaTime -= Context.GetDeltaTimeSeconds();
 		}
diff --git a/Engine/Plugins/Runtime/MassGameplay/Source/MassRepresentation/Private/MassVisualizationTrait.cpp b/Engine/Plugins/Runtime/MassGameplay/Source/MassRepresentation/Private/MassVisualizationTrait.cpp
--- a/Engine/Plugins/Runtime/MassGameplay/Source/MassRepresentation/Private/MassVisualizationTrait.cpp
+++ b/Engine/Plugins/Runtime/MassGameplay/Source/MassRepresentation/Private/MassVisualizationTrait.cpp
@@ -13,11 +13,11 @@ UMassVisualizationTrait::UMassVisualizationTrait()
 {
 	RepresentationSubsystemClass = UMassRepresentationSubsystem::StaticClass();
 
-	Config.RepresentationActorManagementClass = UMassRepresentationActorManagement::StaticClass();
-	Config.LODRepresentation[EMassLOD::High] = EMassRepresentationType::HighResSpawnedActor;
-	Config.LODRepresentation[EMassLOD::Medium] = EMassRepresentationType::LowResSpawnedActor;
-	Config.LODRepresentation[EMassLOD::Low] = EMassRepresentationType::StaticMeshInstance;
-	Config.LODRepresentation[EMassLOD::Off] = EMassRepresentationType::None;
+	Params.RepresentationActorManagementClass = UMassRepresentationActorManagement::StaticClass();
+	Params.LODRepresentation[EMassLOD::High] = EMassRepresentationType::HighResSpawnedActor;
+	Params.LODRepresentation[EMassLOD::Medium] = EMassRepresentationType::LowResSpawnedActor;
+	Params.LODRepresentation[EMassLOD::Low] = EMassRepresentationType::StaticMeshInstance;
+	Params.LODRepresentation[EMassLOD::Off] = EMassRepresentationType::None;
 }
 
 void UMassVisualizationTrait::BuildTemplate(FMassEntityTemplateBuildContext& BuildContext, UWorld& World) const
@@ -51,14 +51,14 @@ void UMassVisualizationTrait::BuildTemplate(FMassEntityTemplateBuildContext& Bui
 	FSharedStruct SubsystemFragment = EntitySubsystem->GetOrCreateSharedFragment<FMassRepresentationSubsystemSharedFragment>(SubsystemHash, Subsystem);
 	BuildContext.AddSharedFragment(SubsystemFragment);
 
-	if (!Config.RepresentationActorManagementClass)
+	if (!Params.RepresentationActorManagementClass)
 	{
 		UE_LOG(LogMassRepresentation, Error, TEXT("Expecting a valid class for the representation actor management"));
 	}
-	uint32 ConfigHash = UE::StructUtils::GetStructCrc32(FConstStructView::Make(Config));
-	FConstSharedStruct ConfigFragment = EntitySubsystem->GetOrCreateConstSharedFragment<FMassRepresentationParameters>(ConfigHash, Config);
-	ConfigFragment.Get<FMassRepresentationParameters>().ComputeCachedValues();
-	BuildContext.AddConstSharedFragment(ConfigFragment);
+	uint32 ParamsHash = UE::StructUtils::GetStructCrc32(FConstStructView::Make(Params));
+	FConstSharedStruct ParamsFragment = EntitySubsystem->GetOrCreateConstSharedFragment<FMassRepresentationParameters>(ParamsHash, Params);
+	ParamsFragment.Get<FMassRepresentationParameters>().ComputeCachedValues();
+	BuildContext.AddConstSharedFragment(ParamsFragment);
 
 	FMassRepresentationFragment& RepresentationFragment = BuildContext.AddFragment_GetRef<FMassRepresentationFragment>();
 	RepresentationFragment.StaticMeshDescIndex = RepresentationSubsystem->FindOrAddStaticMeshDesc(StaticMeshInstanceDesc);
diff --git a/Engine/Plugins/Runtime/MassGameplay/Source/MassRepresentation/Public/MassVisualizationTrait.h b/Engine/Plugins/Runtime/MassGameplay/Source/MassRepresentation/Public/MassVisualizationTrait.h
--- a/Engine/Plugins/Runtime/MassGameplay/Source/MassRepresentation/Public/MassVisualizationTrait.h
+++ b/Engine/Plugins/Runtime/MassGameplay/Source/MassRepresentation/Public/MassVisualizationTrait.h
@@ -40,5 +40,5 @@ public:
 
 	/** Configuration parameters for the representaion processor */
 	UPROPERTY(EditAnywhere, Category = "Mass|Visual")
-	FMassRepresentationParameters Config;
+	FMassRepresentationParameters Params;
 };
