Rename MassLODManager to MassLODSubsystem
Rename MassProcessor_LODBase to MassLODProcessorBase

#rnx

diff --git a/Engine/Plugins/AI/MassCrowd/Source/MassCrowd/Private/MassCrowdReplicationProcessor.cpp b/Engine/Plugins/AI/MassCrowd/Source/MassCrowd/Private/MassCrowdReplicationProcessor.cpp
--- a/Engine/Plugins/AI/MassCrowd/Source/MassCrowd/Private/MassCrowdReplicationProcessor.cpp
+++ b/Engine/Plugins/AI/MassCrowd/Source/MassCrowd/Private/MassCrowdReplicationProcessor.cpp
@@ -5,7 +5,7 @@
 #include "MassClientBubbleHandler.h"
 #include "MassCrowdBubble.h"
 #include "MassCommonFragments.h"
-#include "MassLODManager.h"
+#include "MassLODSubsystem.h"
 #include "MassCrowdFragments.h"
 #include "MassReplicationPathHandlers.h"
 #include "MassReplicationTransformHandlers.h"
diff --git a/Engine/Plugins/AI/MassCrowd/Source/MassCrowd/Private/MassCrowdServerRepresentationLODProcessor.cpp b/Engine/Plugins/AI/MassCrowd/Source/MassCrowd/Private/MassCrowdServerRepresentationLODProcessor.cpp
--- a/Engine/Plugins/AI/MassCrowd/Source/MassCrowd/Private/MassCrowdServerRepresentationLODProcessor.cpp
+++ b/Engine/Plugins/AI/MassCrowd/Source/MassCrowd/Private/MassCrowdServerRepresentationLODProcessor.cpp
@@ -47,8 +47,8 @@ void UMassCrowdServerRepresentationLODProcessor::Execute(UMassEntitySubsystem& E
 {
 	TRACE_CPUPROFILER_EVENT_SCOPE(TEXT("CrowdServerRepresentationLOD"))
 
-	check(LODManager);
-	const TArray<FViewerInfo>& Viewers = LODManager->GetViewers();
+	check(LODSubsystem);
+	const TArray<FViewerInfo>& Viewers = LODSubsystem->GetViewers();
 	LODCalculator.PrepareExecution(Viewers);
 	
 	{
diff --git a/Engine/Plugins/AI/MassCrowd/Source/MassCrowd/Public/MassCrowdServerRepresentationLODProcessor.h b/Engine/Plugins/AI/MassCrowd/Source/MassCrowd/Public/MassCrowdServerRepresentationLODProcessor.h
--- a/Engine/Plugins/AI/MassCrowd/Source/MassCrowd/Public/MassCrowdServerRepresentationLODProcessor.h
+++ b/Engine/Plugins/AI/MassCrowd/Source/MassCrowd/Public/MassCrowdServerRepresentationLODProcessor.h
@@ -4,14 +4,14 @@
 
 #include "MassCrowdTypes.h"
 #include "MassCrowdFragments.h"
-#include "MassLODManager.h"
+#include "MassLODSubsystem.h"
 #include "MassLODCalculator.h"
 #include "MassLODTickRateController.h"
 
 #include "MassCrowdServerRepresentationLODProcessor.generated.h"
 
 UCLASS(meta=(DisplayName="Crowd Simulation LOD"))
-class MASSCROWD_API UMassCrowdServerRepresentationLODProcessor : public UMassProcessor_LODBase
+class MASSCROWD_API UMassCrowdServerRepresentationLODProcessor : public UMassLODProcessorBase
 {
 	GENERATED_BODY()
 
diff --git a/Engine/Plugins/Runtime/MassGameplay/Source/MassLOD/Private/MassLODCollectorProcessor.cpp b/Engine/Plugins/Runtime/MassGameplay/Source/MassLOD/Private/MassLODCollectorProcessor.cpp
--- a/Engine/Plugins/Runtime/MassGameplay/Source/MassLOD/Private/MassLODCollectorProcessor.cpp
+++ b/Engine/Plugins/Runtime/MassGameplay/Source/MassLOD/Private/MassLODCollectorProcessor.cpp
@@ -72,8 +72,8 @@ void UMassLODCollectorProcessor::ExecuteInternal(UMassEntitySubsystem& EntitySub
 
 void UMassLODCollectorProcessor::Execute(UMassEntitySubsystem& EntitySubsystem, FMassExecutionContext& Context)
 {
-	check(LODManager);
-	const TArray<FViewerInfo>& Viewers = LODManager->GetViewers();
+	check(LODSubsystem);
+	const TArray<FViewerInfo>& Viewers = LODSubsystem->GetViewers();
 	Collector.PrepareExecution(Viewers);
 
 	check(World);
diff --git a/Engine/Plugins/Runtime/MassGameplay/Source/MassLOD/Private/MassLODManager.cpp b/Engine/Plugins/Runtime/MassGameplay/Source/MassLOD/Private/MassLODManager.cpp
deleted file mode 100644
diff --git a/Engine/Plugins/Runtime/MassGameplay/Source/MassLOD/Private/MassLODModule.cpp b/Engine/Plugins/Runtime/MassGameplay/Source/MassLOD/Private/MassLODModule.cpp
--- a/Engine/Plugins/Runtime/MassGameplay/Source/MassLOD/Private/MassLODModule.cpp
+++ b/Engine/Plugins/Runtime/MassGameplay/Source/MassLOD/Private/MassLODModule.cpp
@@ -25,7 +25,8 @@ void FMassLODModule::StartupModule()
 	Redirects.Emplace(ECoreRedirectFlags::Type_Struct, TEXT("/Script/MassCommon.DataFragment_MassSimulationLOD"), TEXT("/Script/MassLOD.DataFragment_MassSimulationLOD"));
 	Redirects.Emplace(ECoreRedirectFlags::Type_Struct, TEXT("/Script/MassCommon.DataFragment_MassViewerLOD"), TEXT("/Script/MassLOD.DataFragment_MassViewerLOD"));
 
-	Redirects.Emplace(ECoreRedirectFlags::Type_Class, TEXT("/Script/MassCommon.MassProcessor_LODBase"), TEXT("/Script/MassLOD.MassProcessor_LODBase"));
+	Redirects.Emplace(ECoreRedirectFlags::Type_Class, TEXT("/Script/MassCommon.MassProcessor_LODBase"), TEXT("/Script/MassLOD.MassLODProcessorBase"));
+	Redirects.Emplace(ECoreRedirectFlags::Type_Class, TEXT("/Script/MassLOD.MassProcessor_LODBase"), TEXT("/Script/MassLOD.MassLODProcessorBase"));
 
 	Redirects.Emplace(ECoreRedirectFlags::Type_Class, TEXT("/Script/MassCommon.MassProcessor_MassSimulationLODViewersInfo"), TEXT("/Script/MassLOD.MassProcessor_MassSimulationLODViewersInfo"));
 	Redirects.Emplace(ECoreRedirectFlags::Type_Class, TEXT("/Script/MassCommon.MassProcessor_MassSimulationLOD"), TEXT("/Script/MassLOD.MassProcessor_MassSimulationLOD"));
diff --git a/Engine/Plugins/Runtime/MassGameplay/Source/MassLOD/Private/MassLODSubsystem.cpp b/Engine/Plugins/Runtime/MassGameplay/Source/MassLOD/Private/MassLODSubsystem.cpp
new file mode 100644
--- /dev/null
+++ b/Engine/Plugins/Runtime/MassGameplay/Source/MassLOD/Private/MassLODSubsystem.cpp
@@ -0,0 +1,364 @@
+// Copyright Epic Games, Inc. All Rights Reserved.
+
+#include "MassLODSubsystem.h"
+#include "GameFramework/PlayerController.h"
+#include "WorldPartition/WorldPartition.h"
+#include "WorldPartition/WorldPartitionStreamingSource.h"
+#include "Engine/World.h"
+#include "MassSimulationSubsystem.h"
+
+namespace UE
+{
+namespace MassLOD
+{
+	FColor LODColors[] =
+	{
+		FColor::Red,
+		FColor::Yellow,
+		FColor::Emerald,
+		FColor::White,
+	};
+}
+}
+
+void UMassLODProcessorBase::Initialize(UObject& Owner)
+{
+	Super::Initialize(Owner);
+
+	World = Owner.GetWorld();
+	LODSubsystem = UWorld::GetSubsystem<UMassLODSubsystem>(World);
+}
+
+void UMassLODSubsystem::Initialize(FSubsystemCollectionBase& Collection)
+{
+	Collection.InitializeDependency(UMassSimulationSubsystem::StaticClass());
+
+	Super::Initialize(Collection);
+	SynchronizeViewers();
+	
+	if (UWorld* World = GetWorld())
+	{
+		UMassSimulationSubsystem* SimSystem = World->GetSubsystem<UMassSimulationSubsystem>();
+		check(SimSystem);
+		SimSystem->GetOnProcessingPhaseStarted(EMassProcessingPhase::PrePhysics).AddUObject(this, &UMassLODSubsystem::OnPrePhysicsPhaseStarted);
+	}
+}
+
+void UMassLODSubsystem::OnPrePhysicsPhaseStarted(float DeltaTime)
+{
+	SynchronizeViewers();
+}
+
+TStatId UMassLODSubsystem::GetStatId() const
+{
+	RETURN_QUICK_DECLARE_CYCLE_STAT(UMassLODSubsystem, STATGROUP_Tickables);
+}
+
+void UMassLODSubsystem::Deinitialize()
+{
+	// Remove all viewers from the viewer info list
+	for (FViewerInfo& ViewerInfo : Viewers)
+	{
+		if (ViewerInfo.Handle.IsValid())
+		{
+			// Safe to remove while iterating as it is a sparse array with a free list
+			RemoveViewer(ViewerInfo.Handle);
+		}
+	}
+	
+	if (UWorld* World = GetWorld())
+	{
+		if (UMassSimulationSubsystem* SimSystem = World->GetSubsystem<UMassSimulationSubsystem>())
+		{
+			SimSystem->GetOnProcessingPhaseStarted(EMassProcessingPhase::PrePhysics).RemoveAll(this);
+		}
+	}
+}
+
+int32 UMassLODSubsystem::GetValidViewerIdx(const FMassViewerHandle& ViewerHandle) const
+{
+	// Check against invalid handle
+	if (!ViewerHandle.IsValid())
+	{
+		return INDEX_NONE;
+	}
+
+	// Check against removed viewers
+	const int32 ViewerIdx = ViewerHandle.GetIndex();
+	if (ViewerIdx >= Viewers.Num() || ViewerHandle != Viewers[ViewerIdx].Handle)
+	{
+		return INDEX_NONE;
+	}
+
+	return ViewerIdx;
+}
+
+const TArray<FViewerInfo>& UMassLODSubsystem::GetSynchronizedViewers()
+{
+	SynchronizeViewers();
+
+	return Viewers;
+}
+
+FMassViewerHandle UMassLODSubsystem::GetViewerHandleFromPlayerController(const APlayerController* PlayerController) const
+{
+	const FMassViewerHandle* Handle = ViewerMap.Find(GetTypeHash(PlayerController->GetFName()));
+	return Handle ? *Handle : FMassViewerHandle();
+}
+
+
+FMassViewerHandle UMassLODSubsystem::GetViewerHandleFromStreamingSource(const FName StreamingSourceName) const
+{
+	const FMassViewerHandle* Handle = ViewerMap.Find(GetTypeHash(StreamingSourceName));
+	return Handle ? *Handle : FMassViewerHandle();
+}
+
+APlayerController* UMassLODSubsystem::GetPlayerControllerFromViewerHandle(const FMassViewerHandle& ViewerHandle) const
+{
+	const int32 ViewerIdx = GetValidViewerIdx(ViewerHandle);
+	return ViewerIdx != INDEX_NONE ? Viewers[ViewerIdx].PlayerController : nullptr;
+}
+
+void UMassLODSubsystem::SynchronizeViewers()
+{
+	if (LastSynchronizedFrame == GFrameCounter)
+	{
+		return;
+	}
+	LastSynchronizedFrame = GFrameCounter;
+
+	bool bNeedShrinking = false;
+
+	const UWorld* World = GetWorld();
+	UWorldPartition* WorldPartition = World ? World->GetWorldPartition() : nullptr;
+	const TArray<FWorldPartitionStreamingSource>& StreamingSources = WorldPartition ? WorldPartition->GetStreamingSources() : TArray<FWorldPartitionStreamingSource>();
+
+	// Go through the list and check validity and store the valid one into a map
+	TMap<uint32, FMassViewerHandle> LocalViewerMap;
+	for (int32 ViewerIdx = 0; ViewerIdx < Viewers.Num(); ++ViewerIdx)
+	{
+		FViewerInfo& ViewerInfo = Viewers[ViewerIdx];
+		if (!ViewerInfo.Handle.IsValid())
+		{
+			check(ViewerFreeIndices.Find(ViewerIdx) != INDEX_NONE);
+			continue;
+		}
+
+		if (ViewerInfo.PlayerController != nullptr)
+		{
+			LocalViewerMap.Add(GetTypeHash(ViewerInfo.PlayerController->GetFName()), ViewerInfo.Handle);
+		}
+		else if (!ViewerInfo.StreamingSourceName.IsNone() && StreamingSources.FindByPredicate([&ViewerInfo](const FWorldPartitionStreamingSource& Source){ return Source.Name == ViewerInfo.StreamingSourceName; }) != nullptr)
+		{
+			LocalViewerMap.Add(GetTypeHash(ViewerInfo.StreamingSourceName), ViewerInfo.Handle);
+		}
+		else
+		{
+			// Safe to remove while iterating as it is a sparse array with a free list
+			RemoveViewer(ViewerInfo.Handle);
+			bNeedShrinking |= ViewerIdx == Viewers.Num() - 1;
+		}
+	}
+
+	if (World)
+	{
+		// Now go through all current player controllers and add if they do not exist
+		for (FConstPlayerControllerIterator PlayerIterator = World->GetPlayerControllerIterator(); PlayerIterator; ++PlayerIterator)
+		{
+			APlayerController* PlayerController = (*PlayerIterator).Get();
+			check(PlayerController);
+
+			// Check if the controller already exists by trying to remove it from the map which was filled up with controllers we were tracking
+			if (LocalViewerMap.Remove(GetTypeHash(PlayerController->GetFName())) == 0)
+			{
+				// If not add it to the list
+				AddViewer(PlayerController);
+			}
+		}
+
+		// Now go through all current streaming source and add if they do not exist
+		for (const FWorldPartitionStreamingSource& StreamingSource : StreamingSources)
+		{
+			if (LocalViewerMap.Remove(GetTypeHash(StreamingSource.Name)) == 0)
+			{
+				AddViewer(nullptr, StreamingSource.Name);
+			}
+		}
+	}
+
+	// Anything left in the map need to be removed from the list
+	for (TMap<uint32, FMassViewerHandle>::TIterator Itr = LocalViewerMap.CreateIterator(); Itr; ++Itr)
+	{
+		const int32 ViewerIdx = Itr->Value.GetIndex();
+		RemoveViewer(Viewers[ViewerIdx].Handle);
+		bNeedShrinking |= ViewerIdx == Viewers.Num() - 1;
+	}
+
+	if (bNeedShrinking)
+	{
+		// Check to shrink the array of viewers if possible
+		while (Viewers.Num() > 0 && ViewerFreeIndices.Num() > 0 && ViewerFreeIndices.Remove(Viewers.Num() - 1))
+		{
+			Viewers.Pop(false);
+		}
+	}
+
+	// Update location and direction for every type of viewer
+	for (FViewerInfo& ViewerInfo : Viewers)
+	{
+		if (!ViewerInfo.Handle.IsValid())
+		{
+			continue;
+		}
+
+		if (ViewerInfo.PlayerController)
+		{
+			ViewerInfo.bEnabled = !WorldPartition || ViewerInfo.PlayerController->bEnableStreamingSource;
+
+			FVector PlayerCameraLocation(ForceInitToZero);
+			FRotator PlayerCameraRotation(FRotator::ZeroRotator);
+			ViewerInfo.PlayerController->GetPlayerViewPoint(PlayerCameraLocation, PlayerCameraRotation);
+			ViewerInfo.Location = PlayerCameraLocation;
+			ViewerInfo.Rotation = PlayerCameraRotation;
+
+			// Try to fetch a more precise FOV
+			if(ViewerInfo.PlayerController->PlayerCameraManager)
+			{
+				ViewerInfo.FOV = ViewerInfo.PlayerController->PlayerCameraManager->GetFOVAngle();
+
+				// @todo need to find a way to retrieve aspect ratio, this does not seems to work
+				//ViewerInfo.AspectRatio = MinViewInfo.AspectRatio;
+			}
+		}
+		else
+		{
+			checkf(!ViewerInfo.StreamingSourceName.IsNone(), TEXT("Expecting to have a streamingsourcename if the playercontroller is null"));
+			const FWorldPartitionStreamingSource* StreamingSource = StreamingSources.FindByPredicate([&ViewerInfo](const FWorldPartitionStreamingSource& Source) { return Source.Name == ViewerInfo.StreamingSourceName; });
+			checkf(StreamingSource, TEXT("Expecting to be pointing to a valid streaming source"));
+			ViewerInfo.bEnabled = StreamingSource != nullptr;
+			if (StreamingSource)
+			{
+				ViewerInfo.Location = StreamingSource->Location;
+				ViewerInfo.Rotation = StreamingSource->Rotation;
+			}
+		}
+	}
+}
+
+void UMassLODSubsystem::AddViewer(APlayerController* PlayerController, FName StreamingSourceName /* = NAME_None*/)
+{
+	const int32 HashValue = PlayerController ? GetTypeHash(PlayerController->GetFName()) : GetTypeHash(StreamingSourceName);
+
+	FMassViewerHandle& ViewerHandle = ViewerMap.FindOrAdd(HashValue, FMassViewerHandle());
+	if (ViewerHandle.IsValid())
+	{
+		// We are only interested to set the player controller if it was not already set.
+		if (PlayerController)
+		{
+			const int32 ViewerHandleIdx = GetValidViewerIdx(ViewerHandle);
+			check(ViewerHandleIdx != INDEX_NONE);
+
+			FViewerInfo& ViewerInfo = Viewers[ViewerHandleIdx];
+			check(ViewerInfo.PlayerController == nullptr);
+			ViewerInfo.PlayerController = PlayerController;
+			PlayerController->OnEndPlay.AddUniqueDynamic(this, &UMassLODSubsystem::OnPlayerControllerEndPlay);
+		}
+
+		return;
+	}
+
+	// Add new viewer
+#if UE_ALLOW_DEBUG_REPLICATION_DUPLICATE_VIEWERS_PER_CONTROLLER
+	//for debugging / profiling purposes create DebugNumberViwersPerController
+	//in this case ViewerMap will only contain a hash to the most recent viewer handle created.
+	for (int Idx = 0; Idx < UE::MassLOD::DebugNumberViewersPerController; ++Idx)
+#endif //UE_ALLOW_DEBUG_REPLICATION_DUPLICATE_VIEWERS_PER_CONTROLLER
+	{
+		const bool bAddNew = ViewerFreeIndices.Num() == 0;
+		const int NewIdx = bAddNew ? Viewers.Num() : ViewerFreeIndices.Pop();
+		FViewerInfo& NewViewerInfo = bAddNew ? Viewers.AddDefaulted_GetRef() : Viewers[NewIdx];
+		NewViewerInfo.PlayerController = PlayerController;
+		NewViewerInfo.StreamingSourceName = StreamingSourceName;
+		NewViewerInfo.Handle.Index = NewIdx;
+		NewViewerInfo.Handle.SerialNumber = GetNextViewerSerialNumber();
+		NewViewerInfo.HashValue = HashValue;
+
+		ViewerHandle = NewViewerInfo.Handle;
+
+		OnViewerAddedDelegate.Broadcast(NewViewerInfo.Handle, NewViewerInfo.PlayerController, NewViewerInfo.StreamingSourceName);
+	}
+
+	if(PlayerController)
+	{
+		PlayerController->OnEndPlay.AddUniqueDynamic(this, &UMassLODSubsystem::OnPlayerControllerEndPlay);
+	}
+}
+
+void UMassLODSubsystem::RemoveViewer(const FMassViewerHandle& ViewerHandle)
+{
+#if UE_ALLOW_DEBUG_REPLICATION_DUPLICATE_VIEWERS_PER_CONTROLLER
+
+	const int32 ViewerHandleIdx = GetValidViewerIdx(ViewerHandle);
+	check(ViewerHandleIdx != INDEX_NONE);
+
+	//find all the viewer handles the slow way and remove them
+	for (int32 ViewerIdx = 0; ViewerIdx < Viewers.Num(); ++ViewerIdx)
+	{
+		FViewerInfo& ViewerInfo = Viewers[ViewerIdx];
+
+		if (ViewerInfo.PlayerController == Viewers[ViewerHandleIdx].PlayerController)
+		{
+			RemoveViewerInternal(ViewerInfo.Handle);
+		}
+	}
+
+#else
+
+	RemoveViewerInternal(ViewerHandle);
+
+#endif //UE_ALLOW_DEBUG_REPLICATION_DUPLICATE_VIEWERS_PER_CONTROLLER
+}
+
+void UMassLODSubsystem::RemoveViewerInternal(const FMassViewerHandle& ViewerHandle)
+{
+	const int32 ViewerIdx = GetValidViewerIdx(ViewerHandle);
+	check(ViewerIdx != INDEX_NONE);
+	FViewerInfo& ViewerInfo = Viewers[ViewerIdx];
+
+	OnViewerRemovedDelegate.Broadcast(ViewerInfo.Handle, ViewerInfo.PlayerController, ViewerInfo.StreamingSourceName);
+
+	if (ViewerInfo.PlayerController)
+	{
+		ViewerInfo.PlayerController->OnEndPlay.RemoveDynamic(this, &UMassLODSubsystem::OnPlayerControllerEndPlay);
+	}
+
+	ViewerMap.Remove(ViewerInfo.HashValue);
+
+	ViewerInfo.Reset();
+	ViewerFreeIndices.Push(ViewerIdx);
+}
+
+void UMassLODSubsystem::OnPlayerControllerEndPlay(AActor* Actor, EEndPlayReason::Type EndPlayReason)
+{
+	APlayerController* PlayerController = Cast<APlayerController>(Actor);
+	if (ensure(PlayerController))
+	{
+		const FMassViewerHandle ViewerHandle = GetViewerHandleFromPlayerController(PlayerController);
+		if (ensure(ViewerHandle.IsValid()))
+		{
+			RemoveViewer(ViewerHandle);
+		}
+	}
+}
+
+void FViewerInfo::Reset()
+{
+	Handle.Invalidate();
+	PlayerController = nullptr;
+	HashValue = 0;
+}
+
+bool FViewerInfo::IsLocal() const
+{
+	return (PlayerController && PlayerController->IsLocalController()) || !StreamingSourceName.IsNone();
+}
diff --git a/Engine/Plugins/Runtime/MassGameplay/Source/MassLOD/Private/MassRepresentationLODCollectorProcessor.cpp b/Engine/Plugins/Runtime/MassGameplay/Source/MassLOD/Private/MassRepresentationLODCollectorProcessor.cpp
--- a/Engine/Plugins/Runtime/MassGameplay/Source/MassLOD/Private/MassRepresentationLODCollectorProcessor.cpp
+++ b/Engine/Plugins/Runtime/MassGameplay/Source/MassLOD/Private/MassRepresentationLODCollectorProcessor.cpp
@@ -56,8 +56,8 @@ void CollectLODInfo(UMassEntitySubsystem& EntitySubsystem, FMassExecutionContext
 
 void UMassRepresentationLODCollectorProcessor::Execute(UMassEntitySubsystem& EntitySubsystem, FMassExecutionContext& Context)
 {
-	check(LODManager);
-	const TArray<FViewerInfo>& Viewers = LODManager->GetViewers();
+	check(LODSubsystem);
+	const TArray<FViewerInfo>& Viewers = LODSubsystem->GetViewers();
 
 	check(World);
 
diff --git a/Engine/Plugins/Runtime/MassGameplay/Source/MassLOD/Private/MassSimulationLOD.cpp b/Engine/Plugins/Runtime/MassGameplay/Source/MassLOD/Private/MassSimulationLOD.cpp
--- a/Engine/Plugins/Runtime/MassGameplay/Source/MassLOD/Private/MassSimulationLOD.cpp
+++ b/Engine/Plugins/Runtime/MassGameplay/Source/MassLOD/Private/MassSimulationLOD.cpp
@@ -112,8 +112,8 @@ void UMassSimulationLODProcessor::Execute(UMassEntitySubsystem& EntitySubsystem,
 
 	{
 		TRACE_CPUPROFILER_EVENT_SCOPE(TEXT("PrepareExecution"));
-		check(LODManager);
-		const TArray<FViewerInfo>& Viewers = LODManager->GetViewers();
+		check(LODSubsystem);
+		const TArray<FViewerInfo>& Viewers = LODSubsystem->GetViewers();
 
 		EntitySubsystem.ForEachSharedFragment<FMassSimulationLODSharedFragment>([&Viewers](FMassSimulationLODSharedFragment& LODSharedFragment)
 		{
diff --git a/Engine/Plugins/Runtime/MassGameplay/Source/MassLOD/Public/MassLODCollectorProcessor.h b/Engine/Plugins/Runtime/MassGameplay/Source/MassLOD/Public/MassLODCollectorProcessor.h
--- a/Engine/Plugins/Runtime/MassGameplay/Source/MassLOD/Public/MassLODCollectorProcessor.h
+++ b/Engine/Plugins/Runtime/MassGameplay/Source/MassLOD/Public/MassLODCollectorProcessor.h
@@ -3,7 +3,7 @@
 #pragma once
 
 #include "MassEntityTypes.h"
-#include "MassLODManager.h"
+#include "MassLODSubsystem.h"
 #include "MassLODCollector.h"
 
 #include "MassLODCollectorProcessor.generated.h"
@@ -20,7 +20,7 @@ struct FMassGenericCollectorLogic : public FLODDefaultLogic
  * LOD collector which combines collection of LOD information for both Viewer and Simulation LODing when possible.
  */
 UCLASS(meta = (DisplayName = "LOD Collector"))
-class MASSLOD_API UMassLODCollectorProcessor : public UMassProcessor_LODBase
+class MASSLOD_API UMassLODCollectorProcessor : public UMassLODProcessorBase
 {
 	GENERATED_BODY()
 
diff --git a/Engine/Plugins/Runtime/MassGameplay/Source/MassLOD/Public/MassLODLogic.h b/Engine/Plugins/Runtime/MassGameplay/Source/MassLOD/Public/MassLODLogic.h
--- a/Engine/Plugins/Runtime/MassGameplay/Source/MassLOD/Public/MassLODLogic.h
+++ b/Engine/Plugins/Runtime/MassGameplay/Source/MassLOD/Public/MassLODLogic.h
@@ -2,7 +2,7 @@
 
 #pragma once
 
-#include "MassLODManager.h"
+#include "MassLODSubsystem.h"
 #include "Containers/StaticArray.h"
 #include "ConvexVolume.h"
 
diff --git a/Engine/Plugins/Runtime/MassGameplay/Source/MassLOD/Public/MassLODManager.h b/Engine/Plugins/Runtime/MassGameplay/Source/MassLOD/Public/MassLODManager.h
deleted file mode 100644
diff --git a/Engine/Plugins/Runtime/MassGameplay/Source/MassLOD/Public/MassLODSubsystem.h b/Engine/Plugins/Runtime/MassGameplay/Source/MassLOD/Public/MassLODSubsystem.h
new file mode 100644
--- /dev/null
+++ b/Engine/Plugins/Runtime/MassGameplay/Source/MassLOD/Public/MassLODSubsystem.h
@@ -0,0 +1,158 @@
+// Copyright Epic Games, Inc. All Rights Reserved.
+
+#pragma once
+
+#include "MassCommonTypes.h"
+#include "MassProcessor.h"
+#include "IndexedHandle.h"
+#include "MassLODTypes.h"
+#include "Subsystems/WorldSubsystem.h"
+
+#include "MassLODSubsystem.generated.h"
+
+class UMassLODSubsystem;
+
+/*
+ * Base mass LOD processor to store common information for all LOD processors
+ */
+UCLASS(abstract)
+class MASSLOD_API UMassLODProcessorBase : public UMassProcessor
+{
+	GENERATED_BODY()
+
+public:
+	virtual void Initialize(UObject& Owner) override;
+
+protected:
+	UPROPERTY()
+	UWorld* World = nullptr;
+
+	UPROPERTY()
+	UMassLODSubsystem* LODSubsystem = nullptr;
+};
+
+/*
+ * Handle that lets you reference the concept of a viewer
+ */
+USTRUCT()
+struct MASSLOD_API FMassViewerHandle : public FIndexedHandleBase
+{
+	GENERATED_BODY()
+
+	friend class UMassLODSubsystem;
+};
+
+USTRUCT()
+struct FViewerInfo
+{
+	GENERATED_BODY()
+
+	UPROPERTY(transient)
+	APlayerController* PlayerController = nullptr;
+	
+	FName StreamingSourceName;
+
+	FMassViewerHandle Handle;
+	uint32 HashValue = 0;
+
+	FVector Location;
+	FRotator Rotation;
+	float FOV = 90.0f;
+	float AspectRatio = 16.0f / 9.0f;
+
+	bool bEnabled = true;
+
+	void Reset();
+
+	bool IsLocal() const;
+};
+
+DECLARE_MULTICAST_DELEGATE_ThreeParams(FOnViewerAdded, FMassViewerHandle ViewerHandle, APlayerController* PlayerController, FName StreamingSourceName);
+DECLARE_MULTICAST_DELEGATE_ThreeParams(FOnViewerRemoved, FMassViewerHandle ViewerHandle, APlayerController* PlayerController, FName StreamingSourceName);
+
+/*
+ * Manager responsible to manage and synchronized available viewers
+ */
+UCLASS()
+class MASSLOD_API UMassLODSubsystem : public UTickableWorldSubsystem
+{
+	GENERATED_BODY()
+
+public:
+	/** Checks the validity of a viewer handle */
+	bool IsValidViewer(const FMassViewerHandle& ViewerHandle) const { return GetValidViewerIdx(ViewerHandle) != INDEX_NONE; }
+
+	/** Returns the index of the viewer if valid, otherwise INDEX_NONE is return */
+	int32 GetValidViewerIdx(const FMassViewerHandle& ViewerHandle) const;
+
+	/** Returns the array of viewers */
+	const TArray<FViewerInfo>& GetViewers() const { return Viewers; }
+
+	/** Synchronize the viewers if not done this frame and returns the updated array */
+	const TArray<FViewerInfo>& GetSynchronizedViewers();
+
+	/** Returns viewer handle from the PlayerController pointer */
+	FMassViewerHandle GetViewerHandleFromPlayerController(const APlayerController* PlayerController) const;
+
+	/** Returns viewer handle from the streaming source name */
+	FMassViewerHandle GetViewerHandleFromStreamingSource(const FName StreamingSourceName) const;
+
+	/** Returns PlayerController pointer from the viewer handle */
+	APlayerController* GetPlayerControllerFromViewerHandle(const FMassViewerHandle& ViewerHandle) const;
+
+	/** Returns the delegate called when new viewer are added to the list */
+	FOnViewerAdded& GetOnViewerAddedDelegate() { return OnViewerAddedDelegate;  }
+
+	/** Returns the delegate called when viewer are removed from the list */
+	FOnViewerRemoved& GetOnViewerRemovedDelegate() { return OnViewerRemovedDelegate; }
+
+protected:
+	virtual void Initialize(FSubsystemCollectionBase& Collection) override;
+	virtual TStatId GetStatId() const override;
+	virtual void Deinitialize() override;
+
+	/** Called at the start of the PrePhysics mass processing phase and calls SynchronizeViewers */ 
+	void OnPrePhysicsPhaseStarted(float DeltaTime);
+
+	/** Synchronizes the viewers from the engine PlayerController list */
+	void SynchronizeViewers();
+
+	/** Adds a viewer to the list and sends notification about addition */
+	void AddViewer(APlayerController* PlayerController, FName StreamingSourceName = NAME_None);
+
+	/** Removes a viewer to the list and send notification about removal */
+	void RemoveViewer(const FMassViewerHandle& ViewerHandle);
+
+	/** Returns the next new viewer serial number */
+	uint32 GetNextViewerSerialNumber() { return ViewerSerialNumberCounter++; }
+
+	/** Player controller EndPlay callback, removing viewers from the list */
+	UFUNCTION()
+	void OnPlayerControllerEndPlay(AActor* Actor, EEndPlayReason::Type EndPlayReason);
+
+private:
+	/** Removes a viewer to the list and send notification about removal */
+	void RemoveViewerInternal(const FMassViewerHandle& ViewerHandle);
+
+	/** The actual array of viewer's information*/
+	UPROPERTY(Transient)
+	TArray<FViewerInfo> Viewers;
+
+	/** The map that do reverse look up to get ViewerHandle */
+	UPROPERTY(Transient)
+	TMap<uint32, FMassViewerHandle> ViewerMap;
+
+	uint64 LastSynchronizedFrame = 0;
+
+	/** Viewer serial number counter */
+	uint32 ViewerSerialNumberCounter = 0;
+
+	/** Free list of indices in the sparse viewer array */
+	TArray<int32> ViewerFreeIndices;
+
+	/** Delegates to notify anyone who needs to know about viewer changes */
+	FOnViewerAdded OnViewerAddedDelegate;
+	FOnViewerRemoved OnViewerRemovedDelegate;
+
+};
+
diff --git a/Engine/Plugins/Runtime/MassGameplay/Source/MassLOD/Public/MassRepresentationLODCollectorProcessor.h b/Engine/Plugins/Runtime/MassGameplay/Source/MassLOD/Public/MassRepresentationLODCollectorProcessor.h
--- a/Engine/Plugins/Runtime/MassGameplay/Source/MassLOD/Public/MassRepresentationLODCollectorProcessor.h
+++ b/Engine/Plugins/Runtime/MassGameplay/Source/MassLOD/Public/MassRepresentationLODCollectorProcessor.h
@@ -2,7 +2,7 @@
 
 #pragma once
 
-#include "MassLODManager.h"
+#include "MassLODSubsystem.h"
 #include "MassLODCollector.h"
 #include "MassLODCalculator.h"
 #include "MassEntityTypes.h"
@@ -29,7 +29,7 @@ struct FMassRepresentationLODCollectorConfig
 * LOD collector which collects LOD information for Viewer LODing when possible.
 */
 UCLASS(meta = (DisplayName = "LOD Collector"))
-class MASSLOD_API UMassRepresentationLODCollectorProcessor : public UMassProcessor_LODBase
+class MASSLOD_API UMassRepresentationLODCollectorProcessor : public UMassLODProcessorBase
 {
 	GENERATED_BODY()
 
diff --git a/Engine/Plugins/Runtime/MassGameplay/Source/MassLOD/Public/MassSimulationLOD.h b/Engine/Plugins/Runtime/MassGameplay/Source/MassLOD/Public/MassSimulationLOD.h
--- a/Engine/Plugins/Runtime/MassGameplay/Source/MassLOD/Public/MassSimulationLOD.h
+++ b/Engine/Plugins/Runtime/MassGameplay/Source/MassLOD/Public/MassSimulationLOD.h
@@ -2,7 +2,7 @@
 
 #pragma once
 
-#include "MassLODManager.h"
+#include "MassLODSubsystem.h"
 #include "MassLODCollector.h"
 #include "MassLODCalculator.h"
 #include "MassLODTickRateController.h"
@@ -144,7 +144,7 @@ struct MASSLOD_API FMassSimulationVariableTickSharedFragment : public FMassShare
 };
 
 UCLASS(meta = (DisplayName = "Simulation LOD"))
-class MASSLOD_API UMassSimulationLODProcessor : public UMassProcessor_LODBase
+class MASSLOD_API UMassSimulationLODProcessor : public UMassLODProcessorBase
 {
 	GENERATED_BODY()
 
diff --git a/Engine/Plugins/Runtime/MassGameplay/Source/MassReplication/Private/MassReplicationProcessorBase.cpp b/Engine/Plugins/Runtime/MassGameplay/Source/MassReplication/Private/MassReplicationProcessorBase.cpp
--- a/Engine/Plugins/Runtime/MassGameplay/Source/MassReplication/Private/MassReplicationProcessorBase.cpp
+++ b/Engine/Plugins/Runtime/MassGameplay/Source/MassReplication/Private/MassReplicationProcessorBase.cpp
@@ -2,7 +2,7 @@
 
 #include "MassReplicationProcessorBase.h"
 #include "MassClientBubbleHandler.h"
-#include "MassLODManager.h"
+#include "MassLODSubsystem.h"
 #include "MassCommonFragments.h"
 
 namespace UE::Mass::Replication
@@ -73,8 +73,8 @@ void UMassReplicationProcessor::PrepareExecution(UMassEntitySubsystem& EntitySub
 	//first synchronize clients and viewers
 	ReplicationSubsystem->SynchronizeClientsAndViewers();
 
-	check(LODManager);
-	const TArray<FViewerInfo>& Viewers = LODManager->GetViewers();
+	check(LODSubsystem);
+	const TArray<FViewerInfo>& Viewers = LODSubsystem->GetViewers();
 	EntitySubsystem.ForEachSharedFragment<FMassReplicationSharedFragment>([this, &Viewers](FMassReplicationSharedFragment& RepSharedFragment)
 	{
 		if (!RepSharedFragment.bEntityQueryInitialized)
@@ -155,7 +155,7 @@ void UMassReplicationProcessor::Execute(UMassEntitySubsystem& EntitySubsystem, F
 {
 #if UE_REPLICATION_COMPILE_SERVER_CODE
 	check(World);
-	check(LODManager);
+	check(LODSubsystem);
 	check(ReplicationSubsystem);
 
 	{
@@ -197,7 +197,7 @@ void UMassReplicationProcessor::Execute(UMassEntitySubsystem& EntitySubsystem, F
 
 	{
 		QUICK_SCOPE_CYCLE_COUNTER(UMassReplicationProcessor_ProcessClientReplication);
-		FMassReplicationContext ReplicationContext(*World, *LODManager, *ReplicationSubsystem);
+		FMassReplicationContext ReplicationContext(*World, *LODSubsystem, *ReplicationSubsystem);
 		EntitySubsystem.ForEachSharedFragment<FMassReplicationSharedFragment>([this, &EntitySubsystem, &Context, &ReplicationContext](FMassReplicationSharedFragment& RepSharedFragment)
 		{
 			RepSharedFragment.EntityQuery.ForEachEntityChunk(EntitySubsystem, Context, [&ReplicationContext, &RepSharedFragment](FMassExecutionContext& Context)
diff --git a/Engine/Plugins/Runtime/MassGameplay/Source/MassReplication/Private/MassReplicationSubsystem.cpp b/Engine/Plugins/Runtime/MassGameplay/Source/MassReplication/Private/MassReplicationSubsystem.cpp
--- a/Engine/Plugins/Runtime/MassGameplay/Source/MassReplication/Private/MassReplicationSubsystem.cpp
+++ b/Engine/Plugins/Runtime/MassGameplay/Source/MassReplication/Private/MassReplicationSubsystem.cpp
@@ -20,9 +20,9 @@ void UMassReplicationSubsystem::Initialize(FSubsystemCollectionBase& Collection)
 
 	check(World);
 
-	MassLODManager = Collection.InitializeDependency<UMassLODManager>();
+	MassLODSubsystem = Collection.InitializeDependency<UMassLODSubsystem>();
 
-	check(MassLODManager);
+	check(MassLODSubsystem);
 
 	EntitySystem = UWorld::GetSubsystem<UMassEntitySubsystem>(GetWorld());
 }
@@ -46,7 +46,7 @@ void UMassReplicationSubsystem::Deinitialize()
 	ViewerToClientHandleArray.Reset();
 
 	World = nullptr;
-	MassLODManager = nullptr;
+	MassLODSubsystem = nullptr;
 	EntitySystem = nullptr;
 }
 
@@ -136,7 +136,7 @@ bool UMassReplicationSubsystem::SynchronizeClients(const TArray<FViewerInfo>& Vi
 	// will have the same player controller owner.
 	const FMassClientBubbleInfoData& InfoData = BubbleInfoArray[0];
 
-	check(MassLODManager);
+	check(MassLODSubsystem);
 
 	FControllerMap ClientConnectionMap;
 
@@ -236,7 +236,7 @@ bool UMassReplicationSubsystem::SynchronizeClients(const TArray<FViewerInfo>& Vi
 // synchronize the ClientViewers
 void UMassReplicationSubsystem::SynchronizeClientViewers(const TArray<FViewerInfo>& Viewers)
 {
-	check(MassLODManager);
+	check(MassLODSubsystem);
 
 	struct FMassClientViewerHandle
 	{
@@ -269,7 +269,7 @@ void UMassReplicationSubsystem::SynchronizeClientViewers(const TArray<FViewerInf
 				{
 					const FMassViewerHandle& ClientViewer = ClientViewers.Handles[ViewerIdx];
 
-					APlayerController* Controller = MassLODManager->GetPlayerControllerFromViewerHandle(ClientViewer);
+					APlayerController* Controller = MassLODSubsystem->GetPlayerControllerFromViewerHandle(ClientViewer);
 
 					const UE::Mass::Replication::ENetConnectionType ConnectionType = UE::Mass::Replication::HasParentOrChildWithValidParentNetConnection(Controller);
 
@@ -303,16 +303,16 @@ void UMassReplicationSubsystem::SynchronizeClientViewers(const TArray<FViewerInf
 				// check if the parent controller is valid and already exists
 				if (ParentController && (ClientViewerMap.Find(Viewer.PlayerController) == nullptr))
 				{
-					FMassViewerHandle ParentViewerHandle = MassLODManager->GetViewerHandleFromPlayerController(ParentController);
+					FMassViewerHandle ParentViewerHandle = MassLODSubsystem->GetViewerHandleFromPlayerController(ParentController);
 
-					if (ensureMsgf(ParentViewerHandle.IsValid(), TEXT("MassLODManager handles are out of sync with PlayerController NetConnections!")))
+					if (ensureMsgf(ParentViewerHandle.IsValid(), TEXT("MassLODSubsystem handles are out of sync with PlayerController NetConnections!")))
 					{
 						// note all clients (parent NetConnections) should already be set up so we would expect valid handles here
 						check(ViewerToClientHandleArray.IsValidIndex(ParentViewerHandle.GetIndex()));
 
 						const FViewerClientPair& ParentViewerClientPair = ViewerToClientHandleArray[ParentViewerHandle.GetIndex()];
 
-						check(MassLODManager->IsValidViewer(ParentViewerClientPair.ViewerHandle));
+						check(MassLODSubsystem->IsValidViewer(ParentViewerClientPair.ViewerHandle));
 						check(ClientHandleManager.IsValidHandle(ParentViewerClientPair.ClientHandle));
 
 						// remove APlayerController from the ClientViewerMap and Add the viewer to the ClientsViewerHandles
@@ -349,14 +349,14 @@ void UMassReplicationSubsystem::SynchronizeClientsAndViewers()
 	}
 	LastSynchronizedFrame = GFrameCounter;
 
-	if (MassLODManager == nullptr)
+	if (MassLODSubsystem == nullptr)
 	{
 		checkNoEntry();
 		return;
 	}
 
 	// makes sure the LOD manager Viewers are synced before we process them
-	const TArray<FViewerInfo>& Viewers = MassLODManager->GetSynchronizedViewers();
+	const TArray<FViewerInfo>& Viewers = MassLODSubsystem->GetSynchronizedViewers();
 
 	const bool bNeedShrinking = SynchronizeClients(Viewers);
 
@@ -573,8 +573,8 @@ void UMassReplicationSubsystem::DebugCheckArraysAreInSync()
 void UMassReplicationSubsystem::AddClient(FMassViewerHandle ViewerHandle, APlayerController& InController)
 {
 	check(World);
-	check(MassLODManager);
-	checkf(MassLODManager->IsValidViewer(ViewerHandle), TEXT("ViewerHandle must be valid"));
+	check(MassLODSubsystem);
+	checkf(MassLODSubsystem->IsValidViewer(ViewerHandle), TEXT("ViewerHandle must be valid"));
 
 #if !UE_ALLOW_DEBUG_REPLICATION_BUBBLES_STANDALONE
 	checkf(UE::Mass::Replication::HasParentNetConnection(&InController), TEXT("InController must have a parent net connection or replication will not occur!"));
diff --git a/Engine/Plugins/Runtime/MassGameplay/Source/MassReplication/Public/MassReplicationProcessorBase.h b/Engine/Plugins/Runtime/MassGameplay/Source/MassReplication/Public/MassReplicationProcessorBase.h
--- a/Engine/Plugins/Runtime/MassGameplay/Source/MassReplication/Public/MassReplicationProcessorBase.h
+++ b/Engine/Plugins/Runtime/MassGameplay/Source/MassReplication/Public/MassReplicationProcessorBase.h
@@ -9,7 +9,7 @@
 #include "MassReplicationTypes.h"
 #include "MassReplicationFragments.h"
 #include "MassSpawnerTypes.h"
-#include "MassLODManager.h"
+#include "MassLODSubsystem.h"
 #include "MassReplicationSubsystem.h"
 #include "MassReplicationProcessorBase.generated.h"
 
@@ -22,7 +22,7 @@ class UWorld;
  *  query Mass entity fragments and set those values for replication when appropriate, using the MassClientBubbleHandler.
  */
 UCLASS()
-class MASSREPLICATION_API UMassReplicationProcessor : public UMassProcessor_LODBase
+class MASSREPLICATION_API UMassReplicationProcessor : public UMassLODProcessorBase
 {
 	GENERATED_BODY()
 
@@ -50,14 +50,14 @@ protected:
 
 struct FMassReplicationContext
 {
-	FMassReplicationContext(UWorld& InWorld, UMassLODManager& InLODManager, UMassReplicationSubsystem& InReplicationSubsystem)
+	FMassReplicationContext(UWorld& InWorld, UMassLODSubsystem& InLODSubsystem, UMassReplicationSubsystem& InReplicationSubsystem)
 		: World(InWorld)
-		, LODManager(InLODManager)
+		, LODSubsystem(InLODSubsystem)
 		, ReplicationSubsystem(InReplicationSubsystem)
 	{}
 
 	UWorld& World;
-	UMassLODManager& LODManager;
+	UMassLODSubsystem& LODSubsystem;
 	UMassReplicationSubsystem& ReplicationSubsystem;
 };
 
@@ -144,7 +144,7 @@ void UMassReplicatorBase::CalculateClientReplication(FMassExecutionContext& Cont
 				for (const FMassViewerHandle& ViewerHandle : ClientViewers.Handles)
 				{
 					//this should always we valid as we synchronized the viewers just previously
-					check(ReplicationContext.LODManager.IsValidViewer(ViewerHandle));
+					check(ReplicationContext.LODSubsystem.IsValidViewer(ViewerHandle));
 
 					const EMassLOD::Type MassLOD = ViewerLODList[EntityIdx].LODPerViewer[ViewerHandle.GetIndex()];
 					check(MassLOD <= EMassLOD::Off);
diff --git a/Engine/Plugins/Runtime/MassGameplay/Source/MassReplication/Public/MassReplicationSubsystem.h b/Engine/Plugins/Runtime/MassGameplay/Source/MassReplication/Public/MassReplicationSubsystem.h
--- a/Engine/Plugins/Runtime/MassGameplay/Source/MassReplication/Public/MassReplicationSubsystem.h
+++ b/Engine/Plugins/Runtime/MassGameplay/Source/MassReplication/Public/MassReplicationSubsystem.h
@@ -11,7 +11,7 @@
 
 class UMassEntitySubsystem;
 class AMassCrowdClientBubbleInfo;
-class UMassLODManager;
+class UMassLODSubsystem;
 class AMassClientBubbleInfoBase;
 
 struct FClientViewerHandles
@@ -225,7 +225,7 @@ protected:
 
 	/** Adds a Client and sets up all the data relevant to the bubble. Where there are multiple viewers for
 	 *  one NetConnection only the parent net connection should be added here
-	 *  @param ViewerHandle handle to viewer to add as a client. This must be UMassLODManager->IsValidViewer()
+	 *  @param ViewerHandle handle to viewer to add as a client. This must be UMassLODSubsystem->IsValidViewer()
 	 *  @param InController associated with the viewer
 	 */
 	void AddClient(FMassViewerHandle ViewerHandle, APlayerController& InController);
@@ -273,7 +273,7 @@ protected:
 	UWorld* World;
 
 	UPROPERTY()
-	UMassLODManager* MassLODManager;
+	UMassLODSubsystem* MassLODSubsystem;
 
 #if UE_REPLICATION_COMPILE_CLIENT_CODE
 	TMap<FMassNetworkID, FMassReplicationEntityInfo> EntityInfoMap;
diff --git a/Engine/Plugins/Runtime/MassGameplay/Source/MassReplication/Public/MassReplicationTypes.h b/Engine/Plugins/Runtime/MassGameplay/Source/MassReplication/Public/MassReplicationTypes.h
--- a/Engine/Plugins/Runtime/MassGameplay/Source/MassReplication/Public/MassReplicationTypes.h
+++ b/Engine/Plugins/Runtime/MassGameplay/Source/MassReplication/Public/MassReplicationTypes.h
@@ -3,7 +3,7 @@
 #pragma once
 
 #include "MassCommonTypes.h"
-#include "MassLODManager.h"
+#include "MassLODSubsystem.h"
 #include "IndexedHandle.h"
 #include "MassEntityTemplate.h"
 #include "MassObserverProcessor.h"
diff --git a/Engine/Plugins/Runtime/MassGameplay/Source/MassRepresentation/Private/MassVisualizationLODProcessor.cpp b/Engine/Plugins/Runtime/MassGameplay/Source/MassRepresentation/Private/MassVisualizationLODProcessor.cpp
--- a/Engine/Plugins/Runtime/MassGameplay/Source/MassRepresentation/Private/MassVisualizationLODProcessor.cpp
+++ b/Engine/Plugins/Runtime/MassGameplay/Source/MassRepresentation/Private/MassVisualizationLODProcessor.cpp
@@ -56,8 +56,8 @@ void UMassVisualizationLODProcessor::PrepareExecution()
 {
 	check(World);
 
-	check(LODManager);
-	const TArray<FViewerInfo>& Viewers = LODManager->GetViewers();
+	check(LODSubsystem);
+	const TArray<FViewerInfo>& Viewers = LODSubsystem->GetViewers();
 	LODCalculator.PrepareExecution(Viewers);
 }
 
diff --git a/Engine/Plugins/Runtime/MassGameplay/Source/MassRepresentation/Public/MassVisualizationLODProcessor.h b/Engine/Plugins/Runtime/MassGameplay/Source/MassRepresentation/Public/MassVisualizationLODProcessor.h
--- a/Engine/Plugins/Runtime/MassGameplay/Source/MassRepresentation/Public/MassVisualizationLODProcessor.h
+++ b/Engine/Plugins/Runtime/MassGameplay/Source/MassRepresentation/Public/MassVisualizationLODProcessor.h
@@ -9,7 +9,7 @@
 #include "MassVisualizationLODProcessor.generated.h"
 
 UCLASS()
-class MASSREPRESENTATION_API UMassVisualizationLODProcessor : public UMassProcessor_LODBase
+class MASSREPRESENTATION_API UMassVisualizationLODProcessor : public UMassLODProcessorBase
 {
 	GENERATED_BODY()
 
