Removed old reference to ActorManager and AgentManager

#rnx

diff --git a/Engine/Plugins/AI/MassAI/Source/MassAIDebug/Private/GameplayDebuggerCategory_Mass.cpp b/Engine/Plugins/AI/MassAI/Source/MassAIDebug/Private/GameplayDebuggerCategory_Mass.cpp
--- a/Engine/Plugins/AI/MassAI/Source/MassAIDebug/Private/GameplayDebuggerCategory_Mass.cpp
+++ b/Engine/Plugins/AI/MassAI/Source/MassAIDebug/Private/GameplayDebuggerCategory_Mass.cpp
@@ -34,9 +34,9 @@ namespace UE::Mass::Debug
 			EntityHandle = AgentComp->GetEntityHandle();
 			OutMassAgentComponent = AgentComp;
 		}
-		else if (UMassActorSubsystem* ActorManager = UWorld::GetSubsystem<UMassActorSubsystem>(Actor.GetWorld()))
+		else if (UMassActorSubsystem* ActorSubsystem = UWorld::GetSubsystem<UMassActorSubsystem>(Actor.GetWorld()))
 		{
-			EntityHandle = ActorManager->GetEntityHandleFromActor(&Actor);
+			EntityHandle = ActorSubsystem->GetEntityHandleFromActor(&Actor);
 		}
 		return EntityHandle;
 	};
@@ -140,9 +140,9 @@ void FGameplayDebuggerCategory_Mass::PickEntity(const APlayerController& OwnerPC
 	AActor* BestActor = nullptr;
 	if (BestEntity.IsSet())
 	{
-		if (const UMassActorSubsystem* ActorManager = World.GetSubsystem<UMassActorSubsystem>())
+		if (const UMassActorSubsystem* ActorSubsystem = World.GetSubsystem<UMassActorSubsystem>())
 		{
-			BestActor = ActorManager->GetActorFromHandle(FMassEntityHandle(BestEntity));
+			BestActor = ActorSubsystem->GetActorFromHandle(FMassEntityHandle(BestEntity));
 		}
 	}
 
diff --git a/Engine/Plugins/Runtime/MassGameplay/Source/MassActors/Private/MassActorSpawnerSubsystem.cpp b/Engine/Plugins/Runtime/MassGameplay/Source/MassActors/Private/MassActorSpawnerSubsystem.cpp
--- a/Engine/Plugins/Runtime/MassGameplay/Source/MassActors/Private/MassActorSpawnerSubsystem.cpp
+++ b/Engine/Plugins/Runtime/MassGameplay/Source/MassActors/Private/MassActorSpawnerSubsystem.cpp
@@ -95,7 +95,7 @@ bool UMassActorSpawnerSubsystem::ReleaseActorToPool(AActor* Actor)
 			Actor->SetActorHiddenInGame(true);
 			if (AgentComp)
 			{
-				AgentComp->UnregisterWithAgentManager();
+				AgentComp->UnregisterWithAgentSubsystem();
 			}
 
 			TArray<AActor*>& Pool = PooledActors.FindOrAdd(Actor->GetClass());
@@ -193,7 +193,7 @@ AActor* UMassActorSpawnerSubsystem::SpawnOrRetrieveFromPool(FConstStructView Spa
 
 			if (UMassAgentComponent* AgentComp = PooledActor->FindComponentByClass<UMassAgentComponent>())
 			{
-				AgentComp->RegisterWithAgentManager();
+				AgentComp->RegisterWithAgentSubsystem();
 				AgentComp->SetPuppetHandle(SpawnRequest.MassAgent);
 			}
 
diff --git a/Engine/Plugins/Runtime/MassGameplay/Source/MassActors/Private/MassActorSubsystem.cpp b/Engine/Plugins/Runtime/MassGameplay/Source/MassActors/Private/MassActorSubsystem.cpp
--- a/Engine/Plugins/Runtime/MassGameplay/Source/MassActors/Private/MassActorSubsystem.cpp
+++ b/Engine/Plugins/Runtime/MassGameplay/Source/MassActors/Private/MassActorSubsystem.cpp
@@ -59,7 +59,7 @@ void FMassActorFragment::ResetNoHandleMapUpdate()
 void UMassActorSubsystem::Initialize(FSubsystemCollectionBase& Collection)
 {
 	Super::Initialize(Collection);
-	// making sure UMassSimulationSubsystem gets created before the MassActorManager
+	// making sure UMassSimulationSubsystem gets created before the MassActorSubsystem
 	Collection.InitializeDependency<UMassSimulationSubsystem>();
 	
 	EntitySystem = UWorld::GetSubsystem<UMassEntitySubsystem>(GetWorld());
diff --git a/Engine/Plugins/Runtime/MassGameplay/Source/MassActors/Private/MassAgentComponent.cpp b/Engine/Plugins/Runtime/MassGameplay/Source/MassActors/Private/MassAgentComponent.cpp
--- a/Engine/Plugins/Runtime/MassGameplay/Source/MassActors/Private/MassAgentComponent.cpp
+++ b/Engine/Plugins/Runtime/MassGameplay/Source/MassActors/Private/MassAgentComponent.cpp
@@ -91,7 +91,7 @@ void UMassAgentComponent::OnRegister()
 	}
 
 	// @todo hook up to pawn possessing stuff, maybe?
-	RegisterWithAgentManager();
+	RegisterWithAgentSubsystem();
 }
 
 bool UMassAgentComponent::IsReadyForPooling() const
@@ -105,17 +105,17 @@ bool UMassAgentComponent::IsReadyForPooling() const
 	return true;
 }
 
-void UMassAgentComponent::RegisterWithAgentManager()
+void UMassAgentComponent::RegisterWithAgentSubsystem()
 {
-	UMassAgentSubsystem* AgentManager = UWorld::GetSubsystem<UMassAgentSubsystem>(GetWorld());
-	UE_CVLOG_UELOG(AgentManager == nullptr, GetOwner(), LogMass, Error, TEXT("Unable to find UMassAgentSubsystem instance. Make sure the world is initialized"));
-	if (ensureMsgf(AgentManager, TEXT("Unable to find UMassAgentSubsystem instance. Make sure the world is initialized")))
+	UMassAgentSubsystem* AgentSubsystem = UWorld::GetSubsystem<UMassAgentSubsystem>(GetWorld());
+	UE_CVLOG_UELOG(AgentSubsystem == nullptr, GetOwner(), LogMass, Error, TEXT("Unable to find UMassAgentSubsystem instance. Make sure the world is initialized"));
+	if (ensureMsgf(AgentSubsystem, TEXT("Unable to find UMassAgentSubsystem instance. Make sure the world is initialized")))
 	{
-		TemplateID = AgentManager->RegisterAgentComponent(*this);
+		TemplateID = AgentSubsystem->RegisterAgentComponent(*this);
 	}
 }
 
-void UMassAgentComponent::UnregisterWithAgentManager()
+void UMassAgentComponent::UnregisterWithAgentSubsystem()
 {
 	if (State != EAgentComponentState::None)
 	{
@@ -142,7 +142,7 @@ void UMassAgentComponent::UnregisterWithAgentManager()
 void UMassAgentComponent::OnUnregister()
 {
 	UE_VLOG(GetOwner(), LogMass, Verbose, TEXT("%s"), ANSI_TO_TCHAR(__FUNCTION__));
-	UnregisterWithAgentManager();
+	UnregisterWithAgentSubsystem();
 
 	Super::OnUnregister();
 }
@@ -462,7 +462,7 @@ void UMassAgentComponent::PausePuppet(const bool bPause)
 	MASSAGENT_CHECK(IsPuppet(), TEXT("%s can only be called when the the mass agent component acts as a puppet. Current state is %s"), ANSI_TO_TCHAR(__FUNCTION__), *UEnum::GetValueAsString(State));
 	if (IsPuppetPaused() != bPause)
 	{
-		if (UMassAgentSubsystem* AgentManager = UWorld::GetSubsystem<UMassAgentSubsystem>(GetWorld()))
+		if (UMassAgentSubsystem* AgentSubsystem = UWorld::GetSubsystem<UMassAgentSubsystem>(GetWorld()))
 		{
 			if (bPause)
 			{
@@ -470,14 +470,14 @@ void UMassAgentComponent::PausePuppet(const bool bPause)
 					State == EAgentComponentState::PuppetInitialized,
 					TEXT("%s(true) is expecting to be in state[PuppetPendingInitialization|PuppetInitialized] but is in %s"), ANSI_TO_TCHAR(__FUNCTION__), *UEnum::GetValueAsString(State));
 
-				AgentManager->UnregisterAgentComponent(*this);
+				AgentSubsystem->UnregisterAgentComponent(*this);
 			}
 			else
 			{
 				MASSAGENT_CHECK(State == EAgentComponentState::PuppetPaused,
 					TEXT("%s(false) is expecting to be in state[PuppetPaused] but is in %s"), ANSI_TO_TCHAR(__FUNCTION__), *UEnum::GetValueAsString(State));
 
-				AgentManager->RegisterAgentComponent(*this);
+				AgentSubsystem->RegisterAgentComponent(*this);
 			}
 		}
 	}
@@ -509,10 +509,10 @@ void UMassAgentComponent::ClearReplicatedPuppetHandle()
 	MASSAGENT_CHECK(State == EAgentComponentState::PuppetPaused,
 		TEXT("%s is expecting to be in state [PuppetPaused] but is in %s"), ANSI_TO_TCHAR(__FUNCTION__), *UEnum::GetValueAsString(State));
 
-	if (UMassAgentSubsystem* AgentManager = UWorld::GetSubsystem<UMassAgentSubsystem>(GetWorld()))
+	if (UMassAgentSubsystem* AgentSubsystem = UWorld::GetSubsystem<UMassAgentSubsystem>(GetWorld()))
 	{
 		UE_VLOG(GetOwner(), LogMass, Verbose, TEXT("%s"), ANSI_TO_TCHAR(__FUNCTION__));
-		AgentManager->UnregisterAgentComponent(*this);
+		AgentSubsystem->UnregisterAgentComponent(*this);
 	}
 	ClearEntityHandleInternal();
 	SwitchToState(EAgentComponentState::PuppetReplicatedOrphan);
diff --git a/Engine/Plugins/Runtime/MassGameplay/Source/MassActors/Private/MassAgentSubsystem.cpp b/Engine/Plugins/Runtime/MassGameplay/Source/MassActors/Private/MassAgentSubsystem.cpp
--- a/Engine/Plugins/Runtime/MassGameplay/Source/MassActors/Private/MassAgentSubsystem.cpp
+++ b/Engine/Plugins/Runtime/MassGameplay/Source/MassActors/Private/MassAgentSubsystem.cpp
@@ -38,7 +38,7 @@ void UMassAgentSubsystem::Initialize(FSubsystemCollectionBase& Collection)
 {
 	Super::Initialize(Collection);
 
-	// making sure UMassSimulationSubsystem gets created before the MassAgentManager
+	// making sure UMassSimulationSubsystem gets created before the MassAgentSubsystem
 	Collection.InitializeDependency<UMassSimulationSubsystem>();
 	Collection.InitializeDependency<UMassSpawnerSubsystem>();
 #if UE_REPLICATION_COMPILE_CLIENT_CODE
diff --git a/Engine/Plugins/Runtime/MassGameplay/Source/MassActors/Public/MassAgentComponent.h b/Engine/Plugins/Runtime/MassGameplay/Source/MassActors/Public/MassAgentComponent.h
--- a/Engine/Plugins/Runtime/MassGameplay/Source/MassActors/Public/MassAgentComponent.h
+++ b/Engine/Plugins/Runtime/MassGameplay/Source/MassActors/Public/MassAgentComponent.h
@@ -57,8 +57,8 @@ protected:
 
 public:
 	// Support poolable mass agents going into stasis
-	void RegisterWithAgentManager();
-	void UnregisterWithAgentManager();
+	void RegisterWithAgentSubsystem();
+	void UnregisterWithAgentSubsystem();
 	bool IsReadyForPooling() const;
 
 	/** @todo to enforce encapsulation we could make this protected and have a UMassAgentSubsystem friend. I'm not sure I like it though. */
diff --git a/Engine/Plugins/Runtime/MassGameplay/Source/MassSpawner/Private/MassSpawnerSubsystem.cpp b/Engine/Plugins/Runtime/MassGameplay/Source/MassSpawner/Private/MassSpawnerSubsystem.cpp
--- a/Engine/Plugins/Runtime/MassGameplay/Source/MassSpawner/Private/MassSpawnerSubsystem.cpp
+++ b/Engine/Plugins/Runtime/MassGameplay/Source/MassSpawner/Private/MassSpawnerSubsystem.cpp
@@ -23,7 +23,7 @@ void UMassSpawnerSubsystem::Initialize(FSubsystemCollectionBase& Collection)
 {	
 	Super::Initialize(Collection);
 
-	// making sure UMassSimulationSubsystem gets created before the MassActorManager
+	// making sure UMassSimulationSubsystem gets created before the MassSpawnerSubsystem
 	Collection.InitializeDependency<UMassSimulationSubsystem>();
 
 	UWorld* World = GetWorld();
