Migrate the traffic to new LOD collector trait

#rnx

diff --git a/Engine/Plugins/Runtime/MassGameplay/Source/MassLOD/Private/MassRepresentationLODCollectorProcessor.cpp b/Engine/Plugins/Runtime/MassGameplay/Source/MassLOD/Private/MassRepresentationLODCollectorProcessor.cpp
deleted file mode 100644
diff --git a/Engine/Plugins/Runtime/MassGameplay/Source/MassLOD/Public/MassLODTrait.h b/Engine/Plugins/Runtime/MassGameplay/Source/MassLOD/Public/MassLODTrait.h
--- a/Engine/Plugins/Runtime/MassGameplay/Source/MassLOD/Public/MassLODTrait.h
+++ b/Engine/Plugins/Runtime/MassGameplay/Source/MassLOD/Public/MassLODTrait.h
@@ -12,7 +12,7 @@ class MASSLOD_API UMassLODCollectorTrait : public UMassEntityTraitBase
 {
 	GENERATED_BODY()
 
-protected:
+public:
 	virtual void BuildTemplate(FMassEntityTemplateBuildContext& BuildContext, UWorld& World) const override;
 };
 
diff --git a/Engine/Plugins/Runtime/MassGameplay/Source/MassLOD/Public/MassRepresentationLODCollectorProcessor.h b/Engine/Plugins/Runtime/MassGameplay/Source/MassLOD/Public/MassRepresentationLODCollectorProcessor.h
deleted file mode 100644
