Fixed a crash in AMassSpawner::PostEditChangeProperty when editing the class default EntityTypes of an AMassSpawner BP subclass.

In this case PostEditChangeProperty is called on an AMassSpawner in an EditorPreview world that has no UMassSpawnerSubsystem due to UWorldSubsystem::DoesSupportWorldType rejecting EditorPreview worlds.

#rnx

diff --git a/Engine/Plugins/Runtime/MassGameplay/Source/MassSpawner/Private/MassSpawner.cpp b/Engine/Plugins/Runtime/MassGameplay/Source/MassSpawner/Private/MassSpawner.cpp
--- a/Engine/Plugins/Runtime/MassGameplay/Source/MassSpawner/Private/MassSpawner.cpp
+++ b/Engine/Plugins/Runtime/MassGameplay/Source/MassSpawner/Private/MassSpawner.cpp
@@ -233,7 +233,11 @@ void AMassSpawner::PostEditChangeProperty(FPropertyChangedEvent& PropertyChanged
 		if (PropName == EntityTypesName)
 		{
 			// TODO: Should optimize this, i.e. set a dirty flag and update only when needed.
-			RegisterEntityTemplates();
+			UMassSpawnerSubsystem* SpawnerSystem = UWorld::GetSubsystem<UMassSpawnerSubsystem>(GetWorld());
+			if (SpawnerSystem)
+			{
+				RegisterEntityTemplates();
+			}
 		}
 	}
 }
