Exposed Count & SpawningCountScale access to blueprints and made SpawningCountScale EditAnywhere so an initial scale can be set.

diff --git a/Engine/Plugins/Runtime/MassGameplay/Source/MassSpawner/Private/MassSpawner.cpp b/Engine/Plugins/Runtime/MassGameplay/Source/MassSpawner/Private/MassSpawner.cpp
--- a/Engine/Plugins/Runtime/MassGameplay/Source/MassSpawner/Private/MassSpawner.cpp
+++ b/Engine/Plugins/Runtime/MassGameplay/Source/MassSpawner/Private/MassSpawner.cpp
@@ -525,6 +525,16 @@ bool AMassSpawner::DespawnEntity(const FMassEntityHandle Entity)
 	return false;
 }
 
+int32 AMassSpawner::GetCount() const
+{
+	return Count;
+}
+
+float AMassSpawner::GetSpawningCountScale() const
+{
+	return SpawningCountScale;
+}
+
 void AMassSpawner::ClearTemplates()
 {
 	for (FMassSpawnedEntityType& EntityType : EntityTypes)
diff --git a/Engine/Plugins/Runtime/MassGameplay/Source/MassSpawner/Public/MassSpawner.h b/Engine/Plugins/Runtime/MassGameplay/Source/MassSpawner/Public/MassSpawner.h
--- a/Engine/Plugins/Runtime/MassGameplay/Source/MassSpawner/Public/MassSpawner.h
+++ b/Engine/Plugins/Runtime/MassGameplay/Source/MassSpawner/Public/MassSpawner.h
@@ -79,6 +79,12 @@ public:
 	UFUNCTION(BlueprintCallable, Category = "Spawning")
 	void ScaleSpawningCount(float Scale) { SpawningCountScale = Scale; }
 
+	UFUNCTION(BlueprintCallable, Category = "Spawning")
+	int32 GetCount() const;
+
+	UFUNCTION(BlueprintCallable, Category = "Spawning")
+	float GetSpawningCountScale() const;
+
 	/** Called once DoSpawning completes and all entities have been spawned. */
 	UPROPERTY(BlueprintAssignable)
 	FMassSpawnerOnSpawningFinishedEvent OnSpawningFinishedEvent;
@@ -126,6 +132,7 @@ protected:
 	TArray<TObjectPtr<UMassProcessor>> PostSpawnProcessors;
 
 	/** Scale of the spawning count */
+	UPROPERTY(EditAnywhere, Category = "Mass|Spawn")
 	float SpawningCountScale = 1.0f;
 
 	FDelegateHandle SimulationStartedHandle;
