Fixed MassSpawner's erroneous error-level log when its UWorld is being shut down

diff --git a/Engine/Plugins/Runtime/MassGameplay/Source/MassSpawner/Private/MassSpawner.cpp b/Engine/Plugins/Runtime/MassGameplay/Source/MassSpawner/Private/MassSpawner.cpp
--- a/Engine/Plugins/Runtime/MassGameplay/Source/MassSpawner/Private/MassSpawner.cpp
+++ b/Engine/Plugins/Runtime/MassGameplay/Source/MassSpawner/Private/MassSpawner.cpp
@@ -217,6 +217,8 @@ void AMassSpawner::EndPlay(const EEndPlayReason::Type EndPlayReason)
 {
 	UMassSimulationSubsystem::GetOnSimulationStarted().Remove(SimulationStartedHandle);
 
+	DoDespawning();
+
 	Super::EndPlay(EndPlayReason);
 }
 
