[MassSpawner] fixed crash in ZoneGraphSpawnPointsGenerator when there is no registered graphs.
#rnx

diff --git a/Engine/Plugins/Runtime/MassGameplay/Source/MassSpawner/Private/MassEntityZoneGraphSpawnPointsGenerator.cpp b/Engine/Plugins/Runtime/MassGameplay/Source/MassSpawner/Private/MassEntityZoneGraphSpawnPointsGenerator.cpp
--- a/Engine/Plugins/Runtime/MassGameplay/Source/MassSpawner/Private/MassEntityZoneGraphSpawnPointsGenerator.cpp
+++ b/Engine/Plugins/Runtime/MassGameplay/Source/MassSpawner/Private/MassEntityZoneGraphSpawnPointsGenerator.cpp
@@ -15,7 +15,7 @@ void UMassEntityZoneGraphSpawnPointsGenerator::Generate(UObject& QueryOwner, TCo
 	const UZoneGraphSubsystem* ZoneGraph = UWorld::GetSubsystem<UZoneGraphSubsystem>(QueryOwner.GetWorld());
 	if (ZoneGraph == nullptr)
 	{
-		UE_VLOG_UELOG(&QueryOwner, LogMassSpawner, Error, TEXT("No zone graph found in world"));
+		UE_VLOG_UELOG(&QueryOwner, LogMassSpawner, Error, TEXT("No zone graph subsystem found in world"));
 		return;
 	}
 
@@ -23,6 +23,12 @@ void UMassEntityZoneGraphSpawnPointsGenerator::Generate(UObject& QueryOwner, TCo
 	
 	const FRandomStream RandomStream(GFrameNumber);
 	const TConstArrayView<FRegisteredZoneGraphData> RegisteredZoneGraphs = ZoneGraph->GetRegisteredZoneGraphData();
+	if (RegisteredZoneGraphs.IsEmpty())
+	{
+		UE_VLOG_UELOG(&QueryOwner, LogMassSpawner, Error, TEXT("No zone graphs found"));
+		return;
+	}
+
 	for (const FRegisteredZoneGraphData& Registered : RegisteredZoneGraphs)
 	{
 		if (Registered.bInUse && Registered.ZoneGraphData)
@@ -31,6 +37,12 @@ void UMassEntityZoneGraphSpawnPointsGenerator::Generate(UObject& QueryOwner, TCo
 		}
 	}
 
+	if (Locations.IsEmpty())
+	{
+		UE_VLOG_UELOG(&QueryOwner, LogMassSpawner, Error, TEXT("No locations found on zone graphs"));
+		return;
+	}
+
 	// Randomize them
 	for (int32 I = 0; I < Locations.Num(); ++I)
 	{
