[MassSmartObject] removed direct accesses to the SmartObjectSubsystem's octree
#rnx

diff --git a/Engine/Plugins/Runtime/MassGameplay/Source/MassSmartObjects/Private/MassSmartObjectHandler.cpp b/Engine/Plugins/Runtime/MassGameplay/Source/MassSmartObjects/Private/MassSmartObjectHandler.cpp
--- a/Engine/Plugins/Runtime/MassGameplay/Source/MassSmartObjects/Private/MassSmartObjectHandler.cpp
+++ b/Engine/Plugins/Runtime/MassGameplay/Source/MassSmartObjects/Private/MassSmartObjectHandler.cpp
@@ -156,7 +156,7 @@ EMassSmartObjectClaimResult FMassSmartObjectHandler::ClaimCandidate(const FMassE
 
 	for (const FSmartObjectCandidate& Candidate : View)
 	{
-		if (ClaimSmartObject(Entity, User, Candidate.Handle))
+		if (ClaimSmartObject(Entity, User, Candidate.Result))
 		{
 #if WITH_MASSGAMEPLAY_DEBUG
 			UE_CVLOG(UE::Mass::Debug::IsDebuggingEntity(Entity),
@@ -165,7 +165,7 @@ EMassSmartObjectClaimResult FMassSmartObjectHandler::ClaimCandidate(const FMassE
 				Log,
 				TEXT("[%s] claimed [%s]"),
 				*Entity.DebugGetDescription(),
-				*LexToString(Candidate.Handle));
+				*LexToString(Candidate.Result));
 #endif // WITH_MASSGAMEPLAY_DEBUG
 			break;
 		}
@@ -174,12 +174,9 @@ EMassSmartObjectClaimResult FMassSmartObjectHandler::ClaimCandidate(const FMassE
 	return User.ClaimHandle.IsValid() ? EMassSmartObjectClaimResult::Succeeded : EMassSmartObjectClaimResult::Failed_NoAvailableCandidate;
 }
 
-bool FMassSmartObjectHandler::ClaimSmartObject(const FMassEntityHandle Entity, FMassSmartObjectUserFragment& User, const FSmartObjectHandle& ObjectHandle) const
+bool FMassSmartObjectHandler::ClaimSmartObject(const FMassEntityHandle Entity, FMassSmartObjectUserFragment& User, const FSmartObjectRequestResult& RequestResult) const
 {
-	FSmartObjectRequestFilter Filter;
-	Filter.BehaviorDefinitionClass = USmartObjectMassBehaviorDefinition::StaticClass();
-
-	const FSmartObjectClaimHandle ClaimHandle = SmartObjectSubsystem.Claim(ObjectHandle, Filter);
+	const FSmartObjectClaimHandle ClaimHandle = SmartObjectSubsystem.Claim(RequestResult);
 	const bool bSuccess = ClaimHandle.IsValid();
 	if (bSuccess)
 	{
@@ -206,7 +203,7 @@ bool FMassSmartObjectHandler::ClaimSmartObject(const FMassEntityHandle Entity, F
 		Log,
 		TEXT("[%s] claim for [%s] %s"),
 		*Entity.DebugGetDescription(),
-		*LexToString(ObjectHandle),
+		*LexToString(RequestResult),
 		bSuccess ? TEXT("Succeeded") : TEXT("Failed"));
 #endif // WITH_MASSGAMEPLAY_DEBUG
 
diff --git a/Engine/Plugins/Runtime/MassGameplay/Source/MassSmartObjects/Private/MassSmartObjectProcessor.cpp b/Engine/Plugins/Runtime/MassGameplay/Source/MassSmartObjects/Private/MassSmartObjectProcessor.cpp
--- a/Engine/Plugins/Runtime/MassGameplay/Source/MassSmartObjects/Private/MassSmartObjectProcessor.cpp
+++ b/Engine/Plugins/Runtime/MassGameplay/Source/MassSmartObjects/Private/MassSmartObjectProcessor.cpp
@@ -85,14 +85,15 @@ void UMassSmartObjectCandidatesFinderProcessor::Execute(UMassEntitySubsystem& En
 	// Process world location based requests
 	WorldRequestQuery.ForEachEntityChunk(EntitySubsystem, Context, [this, &Filter, &EntitiesToSignal, &BeginRequestProcessing, &EndRequestProcessing](FMassExecutionContext& Context)
 	{
-		const FSmartObjectOctree& Octree = SmartObjectSubsystem->GetOctree();
-
 		const int32 NumEntities = Context.GetNumEntities();
 		EntitiesToSignal.Reserve(EntitiesToSignal.Num() + NumEntities);
 
 		const TConstArrayView<FMassSmartObjectWorldLocationRequestFragment> RequestList = Context.GetFragmentView<FMassSmartObjectWorldLocationRequestFragment>();
 		const TArrayView<FMassSmartObjectRequestResultFragment> ResultList = Context.GetMutableFragmentView<FMassSmartObjectRequestResultFragment>();
 
+		TArray<FSmartObjectRequestResult> QueryResults;
+		TArray<FSmartObjectCandidate> SortedResults;
+
 		for (int32 i = 0; i < NumEntities; ++i)
 		{
 			const FMassSmartObjectWorldLocationRequestFragment& RequestFragment = RequestList[i];
@@ -113,30 +114,31 @@ void UMassSmartObjectCandidatesFinderProcessor::Execute(UMassEntitySubsystem& En
 			BeginRequestProcessing(Entity, Context, Result);
 			ON_SCOPE_EXIT{ EndRequestProcessing(SmartObjectSubsystem, Entity, Result);	};
 
-			Octree.FindElementsWithBoundsTest(SearchBounds, [this, &Filter, &Entity, &Result, &SearchOrigin, &bDisplayDebug, &DebugColor](const FSmartObjectOctreeElement& Element)
+			QueryResults.Reset();
+			SmartObjectSubsystem->FindSmartObjects(FSmartObjectRequest(SearchBounds, Filter), QueryResults);
+
+			SortedResults.Reset(QueryResults.Num());
+			for (const FSmartObjectRequestResult& QueryResult : QueryResults)
 			{
-				if (Result.NumCandidates < FMassSmartObjectRequestResult::MaxNumCandidates)
-				{
-					// Make sure that we can use a slot in that object (availability with supported definitions, etc.)
-					const FSmartObjectRequestResult SlotResult = SmartObjectSubsystem->FindSlot(Element.SmartObjectHandle, Filter);
-					if (SlotResult.IsValid())
-					{
-						const FVector ObjectLocation = Element.Bounds.Center;
-						Result.Candidates[Result.NumCandidates++] = FSmartObjectCandidate(Element.SmartObjectHandle, FVector::DistSquared(SearchOrigin, ObjectLocation));
+				const FVector SlotLocation = SmartObjectSubsystem->GetSlotLocation(QueryResult.SlotHandle).GetValue();
+				SortedResults.Emplace(QueryResult, FVector::DistSquared(SearchOrigin, SlotLocation));
 
 #if WITH_MASSGAMEPLAY_DEBUG
-						if (bDisplayDebug)
-						{
-							constexpr float DebugRadius = 10.f;
-							// use more precise slot location for debug
-							UE_VLOG_LOCATION(SmartObjectSubsystem, LogSmartObject, Display,
-								SmartObjectSubsystem->GetSlotLocation(SlotResult.SlotHandle).Get(ObjectLocation), DebugRadius, DebugColor, TEXT("%s"), *LexToString(Element.SmartObjectHandle));
-							UE_VLOG_SEGMENT(SmartObjectSubsystem, LogSmartObject, Display, SearchOrigin, ObjectLocation, DebugColor, TEXT(""));
-						}
-#endif // WITH_MASSGAMEPLAY_DEBUG
-					}
+				if (bDisplayDebug)
+				{
+					constexpr float DebugRadius = 10.f;
+					UE_VLOG_LOCATION(SmartObjectSubsystem, LogSmartObject, Display, SlotLocation, DebugRadius, DebugColor, TEXT("%s"), *LexToString(QueryResult.SmartObjectHandle));
+					UE_VLOG_SEGMENT(SmartObjectSubsystem, LogSmartObject, Display, SearchOrigin, SlotLocation, DebugColor, TEXT(""));
 				}
-			});
+#endif // WITH_MASSGAMEPLAY_DEBUG
+			}
+			SortedResults.Sort([](const FSmartObjectCandidate& First, const FSmartObjectCandidate& Second){ return First.Cost < Second.Cost; });
+
+			Result.NumCandidates = FMath::Min<uint8>(FMassSmartObjectRequestResult::MaxNumCandidates, SortedResults.Num());
+			for (int ResultIndex = 0; ResultIndex < Result.NumCandidates; ResultIndex++)
+			{
+				Result.Candidates[ResultIndex] = SortedResults[ResultIndex];
+			}
 		}
 	});
 
@@ -239,7 +241,7 @@ void UMassSmartObjectCandidatesFinderProcessor::Execute(UMassEntitySubsystem& En
 						continue;
 					}
 
-					Result.Candidates[Result.NumCandidates++] = FSmartObjectCandidate(Handle, Cost);
+					Result.Candidates[Result.NumCandidates++] = FSmartObjectCandidate(FoundSlot, Cost);
 
 #if WITH_MASSGAMEPLAY_DEBUG
 					if (bDisplayDebug)
diff --git a/Engine/Plugins/Runtime/MassGameplay/Source/MassSmartObjects/Public/MassSmartObjectHandler.h b/Engine/Plugins/Runtime/MassGameplay/Source/MassSmartObjects/Public/MassSmartObjectHandler.h
--- a/Engine/Plugins/Runtime/MassGameplay/Source/MassSmartObjects/Public/MassSmartObjectHandler.h
+++ b/Engine/Plugins/Runtime/MassGameplay/Source/MassSmartObjects/Public/MassSmartObjectHandler.h
@@ -98,10 +98,10 @@ struct MASSSMARTOBJECTS_API FMassSmartObjectHandler
 	 * associated to the provided identifier.
 	 * @param Entity MassEntity associated to the user fragment
 	 * @param User Fragment of the user claiming
-	 * @param ObjectHandle A valid smart object identifier (method will ensure otherwise)
+	 * @param RequestResult A valid smart object request result (method will ensure otherwise)
 	 * @return Whether the slot has been successfully claimed or not
 	 */
-	bool ClaimSmartObject(const FMassEntityHandle Entity, FMassSmartObjectUserFragment& User, const FSmartObjectHandle& ObjectHandle) const;
+	bool ClaimSmartObject(const FMassEntityHandle Entity, FMassSmartObjectUserFragment& User, const FSmartObjectRequestResult& RequestResult) const;
 
 	/**
 	 * Activates the mass gameplay behavior associated to the previously claimed smart object.
diff --git a/Engine/Plugins/Runtime/MassGameplay/Source/MassSmartObjects/Public/MassSmartObjectRequest.h b/Engine/Plugins/Runtime/MassGameplay/Source/MassSmartObjects/Public/MassSmartObjectRequest.h
--- a/Engine/Plugins/Runtime/MassGameplay/Source/MassSmartObjects/Public/MassSmartObjectRequest.h
+++ b/Engine/Plugins/Runtime/MassGameplay/Source/MassSmartObjects/Public/MassSmartObjectRequest.h
@@ -2,7 +2,7 @@
 
 #pragma once
 
-#include "SmartObjectTypes.h"
+#include "SmartObjectSubsystem.h"
 #include "MassEntityTypes.h"
 #include "ZoneGraphTypes.h"
 #include "Containers/StaticArray.h"
@@ -17,10 +17,10 @@ struct MASSSMARTOBJECTS_API FSmartObjectCandidate
 	GENERATED_BODY()
 
 	FSmartObjectCandidate() = default;
-	FSmartObjectCandidate(const FSmartObjectHandle InHandle, const float InCost) : Handle(InHandle), Cost(InCost) {}
+	FSmartObjectCandidate(const FSmartObjectRequestResult InResult, const float InCost)	: Result(InResult), Cost(InCost) {}
 
 	UPROPERTY(Transient)
-	FSmartObjectHandle Handle;
+	FSmartObjectRequestResult Result;
 
 	UPROPERTY(Transient)
 	float Cost = 0.f;
@@ -36,7 +36,7 @@ struct MASSSMARTOBJECTS_API FMassSmartObjectRequestID
 	GENERATED_BODY()
 
 	FMassSmartObjectRequestID() = default;
-	FMassSmartObjectRequestID(const FMassEntityHandle InEntity) : Entity(InEntity) {}
+	explicit FMassSmartObjectRequestID(const FMassEntityHandle InEntity) : Entity(InEntity) {}
 
 	bool IsSet() const { return Entity.IsSet(); }
 	void Reset() { Entity.Reset(); }
