Fixed performance regression caused by visibility refactor

#rnx

diff --git a/Engine/Plugins/Runtime/MassGameplay/Source/MassLOD/Public/MassLODCalculator.h b/Engine/Plugins/Runtime/MassGameplay/Source/MassLOD/Public/MassLODCalculator.h
--- a/Engine/Plugins/Runtime/MassGameplay/Source/MassLOD/Public/MassLODCalculator.h
+++ b/Engine/Plugins/Runtime/MassGameplay/Source/MassLOD/Public/MassLODCalculator.h
@@ -285,6 +285,7 @@ void TMassLODCalculator<FLODLogic>::Initialize(const float InBaseLODDistance[EMa
 	VisibleBucketSize[EMassLOD::Off] = FLT_MAX;
 	LODMaxCount[EMassLOD::Off] = INT_MAX;
 	LODMaxCountPerViewer[EMassLOD::Off] = INT_MAX;
+	BufferHysteresisOnDistanceRatio = InBufferHysteresisOnDistanceRatio;
 
 	// Calculate the size for each LOD buckets
 	float BasePrevLODDistance = BaseLODDistance[0];
@@ -411,7 +412,7 @@ void TMassLODCalculator<FLODLogic>::CalculateLOD(FMassExecutionContext& Context,
 
 		// Set visibility
 		SetPrevVisibility<bCalculateVisibility>(EntityLOD, PrevVisibility);
-		SetVisibility<bCalculateVisibility>(EntityLOD, bIsVisibleByAViewer ? EMassVisibility::CanBeSeen : (bIsInAVisibleRange ? EMassVisibility::CulledByFrustum : EMassVisibility::CulledByDistance));
+		SetVisibility<bCalculateVisibility>(EntityLOD, bIsInAVisibleRange ? (bIsVisibleByAViewer ? EMassVisibility::CanBeSeen : EMassVisibility::CulledByFrustum) : EMassVisibility::CulledByDistance);
 
 		// Accumulate in buckets
 		const float LODSignificance = AccumulateCountInRuntimeData<bCalculateVisibility, FLODLogic::bCalculateLODSignificance>(EntityLOD.LOD, EntityViewersInfo.ClosestViewerDistanceSq, bIsVisibleByAViewer, RuntimeData);
@@ -471,7 +472,7 @@ void TMassLODCalculator<FLODLogic>::CalculateLOD(FMassExecutionContext& Context,
 					}
 
 					// Set visibility
-					SetPrevVisibilityPerViewer<bCalculateVisibilityPerViewer>(EntityLOD, ViewerIdx, bIsVisibleByViewer ? EMassVisibility::CanBeSeen : (bIsInVisibleRange ? EMassVisibility::CulledByFrustum : EMassVisibility::CulledByDistance));
+					SetPrevVisibilityPerViewer<bCalculateVisibilityPerViewer>(EntityLOD, ViewerIdx, bIsInAVisibleRange ? (bIsVisibleByAViewer ? EMassVisibility::CanBeSeen : EMassVisibility::CulledByFrustum) : EMassVisibility::CulledByDistance);
 					SetVisibilityPerViewer<bCalculateVisibilityPerViewer>(EntityLOD, ViewerIdx, PrevVisibilityPerViewer);
 				}
 			}
