Fixed uninitialized struct variables

#rnx

diff --git a/Engine/Plugins/Runtime/MassGameplay/Source/MassRepresentation/Public/MassRepresentationFragments.h b/Engine/Plugins/Runtime/MassGameplay/Source/MassRepresentation/Public/MassRepresentationFragments.h
--- a/Engine/Plugins/Runtime/MassGameplay/Source/MassRepresentation/Public/MassRepresentationFragments.h
+++ b/Engine/Plugins/Runtime/MassGameplay/Source/MassRepresentation/Public/MassRepresentationFragments.h
@@ -69,6 +69,7 @@ struct FMassRepresentationParameters : public FMassSharedFragment
 {
 	GENERATED_BODY()
 
+	FMassRepresentationParameters() = default;
 
 	/** Allow subclasses to override the representation actor management behavior */
 	UPROPERTY(EditAnywhere, Category = "Mass|Visual")
@@ -76,7 +77,7 @@ struct FMassRepresentationParameters : public FMassSharedFragment
 
 	/** What should be the representation of this entity for each specific LOD */
 	UPROPERTY(EditAnywhere, Category = "Mass|Representation", config)
-	EMassRepresentationType LODRepresentation[EMassLOD::Max] = {};
+	EMassRepresentationType LODRepresentation[EMassLOD::Max] = { EMassRepresentationType::HighResSpawnedActor, EMassRepresentationType::LowResSpawnedActor, EMassRepresentationType::StaticMeshInstance, EMassRepresentationType::None };
 
 	/** If true, LowRes actors will be kept around, disabled, whilst StaticMeshInstance representation is active */
 	UPROPERTY(EditAnywhere, Category = "Mass|Representation", config)
@@ -139,15 +140,15 @@ struct FMassVisualizationLODParameters : public FMassSharedFragment
 
 	/** Distances where each LOD becomes relevant */
 	UPROPERTY(EditAnywhere, Category = "Mass|LOD", config)
-	float BaseLODDistance[EMassLOD::Max];
+	float BaseLODDistance[EMassLOD::Max] = { 0.f, 1000.f, 2500.f, 10000.f };
 	UPROPERTY(EditAnywhere, Category = "Mass|LOD", config)
-	float VisibleLODDistance[EMassLOD::Max];
+	float VisibleLODDistance[EMassLOD::Max] = { 0.f, 2000.f, 4000.f, 15000.f };
 	UPROPERTY(EditAnywhere, Category = "Mass|LOD", meta = (ClampMin = "0.0", UIMin = "0.0"), config)
 	float BufferHysteresisOnDistancePercentage = 10.0f;
 
 	/** Maximum limit for each entity per LOD */
 	UPROPERTY(EditAnywhere, Category = "Mass|LOD", config)
-	int32 LODMaxCount[EMassLOD::Max];
+	int32 LODMaxCount[EMassLOD::Max] = {50, 100, 500, MAX_int32};
 
 	/** How far away from frustum does this entities are considered visible */
 	UPROPERTY(EditAnywhere, Category = "Mass|LOD", meta = (ClampMin = "0.0", UIMin = "0.0"), config)
