Rename files following replication refactor.

#rnx

diff --git a/Engine/Plugins/AI/MassCrowd/Source/MassCrowd/Private/MassCrowdReplicationProcessor.cpp b/Engine/Plugins/AI/MassCrowd/Source/MassCrowd/Private/MassCrowdReplicationProcessor.cpp
deleted file mode 100644
diff --git a/Engine/Plugins/AI/MassCrowd/Source/MassCrowd/Private/MassCrowdReplicator.cpp b/Engine/Plugins/AI/MassCrowd/Source/MassCrowd/Private/MassCrowdReplicator.cpp
new file mode 100644
--- /dev/null
+++ b/Engine/Plugins/AI/MassCrowd/Source/MassCrowd/Private/MassCrowdReplicator.cpp
@@ -0,0 +1,68 @@
+// Copyright Epic Games, Inc. All Rights Reserved.
+
+#include "MassCrowdReplicator.h"
+#include "MassReplicationTypes.h"
+#include "MassClientBubbleHandler.h"
+#include "MassCrowdBubble.h"
+#include "MassCommonFragments.h"
+#include "MassLODSubsystem.h"
+#include "MassCrowdFragments.h"
+#include "MassReplicationPathHandlers.h"
+#include "MassReplicationTransformHandlers.h"
+
+//----------------------------------------------------------------------//
+//  UMassCrowdReplicator
+//----------------------------------------------------------------------//
+void UMassCrowdReplicator::AddRequirements(FMassEntityQuery& EntityQuery)
+{
+	FMassReplicationProcessorPositionYawHandler::AddRequirements(EntityQuery);
+	FMassReplicationProcessorPathHandler::AddRequirements(EntityQuery);
+}
+
+void UMassCrowdReplicator::ProcessClientReplication(FMassExecutionContext& Context, FMassReplicationContext& ReplicationContext)
+{
+#if UE_REPLICATION_COMPILE_SERVER_CODE
+
+	FMassReplicationProcessorPathHandler PathHandler;
+	FMassReplicationProcessorPositionYawHandler PositionYawHandler;
+	FMassReplicationSharedFragment* RepSharedFrag = nullptr;
+
+	auto CacheViewsCallback = [&RepSharedFrag, &PathHandler, &PositionYawHandler](FMassExecutionContext& Context)
+	{
+		PathHandler.CacheFragmentViews(Context);
+		PositionYawHandler.CacheFragmentViews(Context);
+		RepSharedFrag = &Context.GetMutableSharedFragment<FMassReplicationSharedFragment>();
+		check(RepSharedFrag);
+	};
+
+	auto AddEntityCallback = [&RepSharedFrag, &PathHandler, &PositionYawHandler](FMassExecutionContext& Context,const int32 EntityIdx, FReplicatedCrowdAgent& InReplicatedAgent, const FMassClientHandle ClientHandle)->FMassReplicatedAgentHandle
+	{
+		AMassCrowdClientBubbleInfo& CrowdBubbleInfo = RepSharedFrag->GetTypedClientBubbleInfoChecked<AMassCrowdClientBubbleInfo>(ClientHandle);
+
+		PathHandler.AddEntity(EntityIdx, InReplicatedAgent.GetReplicatedPathDataMutable());
+		PositionYawHandler.AddEntity(EntityIdx, InReplicatedAgent.GetReplicatedPositionYawDataMutable());
+
+		return CrowdBubbleInfo.GetCrowdSerializer().Bubble.AddAgent(Context.GetEntity(EntityIdx), InReplicatedAgent);
+	};
+
+	auto ModifyEntityCallback = [&RepSharedFrag, &PathHandler](FMassExecutionContext& Context, const int32 EntityIdx, const EMassLOD::Type LOD, const float Time, const FMassReplicatedAgentHandle Handle, const FMassClientHandle ClientHandle)
+	{
+		AMassCrowdClientBubbleInfo& CrowdBubbleInfo = RepSharedFrag->GetTypedClientBubbleInfoChecked<AMassCrowdClientBubbleInfo>(ClientHandle);
+
+		FMassCrowdClientBubbleHandler& Bubble = CrowdBubbleInfo.GetCrowdSerializer().Bubble;
+
+		PathHandler.ModifyEntity<FCrowdFastArrayItem>(Handle, EntityIdx, Bubble.GetPathHandlerMutable());
+
+		// Don't call the PositionYawHandler here as we currently only replicate the position and yaw when we add an entity to Mass
+	};
+
+	auto RemoveEntityCallback = [&RepSharedFrag](FMassExecutionContext& Context, const FMassReplicatedAgentHandle Handle, const FMassClientHandle ClientHandle)
+	{
+		AMassCrowdClientBubbleInfo& CrowdBubbleInfo = RepSharedFrag->GetTypedClientBubbleInfoChecked<AMassCrowdClientBubbleInfo>(ClientHandle);
+
+		CrowdBubbleInfo.GetCrowdSerializer().Bubble.RemoveAgentChecked(Handle);
+	};
+
+	CalculateClientReplication<FCrowdFastArrayItem>(Context, ReplicationContext, CacheViewsCallback, AddEntityCallback, ModifyEntityCallback, RemoveEntityCallback);
+#endif // UE_REPLICATION_COMPILE_SERVER_CODE
+}
diff --git a/Engine/Plugins/AI/MassCrowd/Source/MassCrowd/Public/MassCrowdReplicationProcessor.h b/Engine/Plugins/AI/MassCrowd/Source/MassCrowd/Public/MassCrowdReplicationProcessor.h
deleted file mode 100644
diff --git a/Engine/Plugins/AI/MassCrowd/Source/MassCrowd/Public/MassCrowdReplicator.h b/Engine/Plugins/AI/MassCrowd/Source/MassCrowd/Public/MassCrowdReplicator.h
new file mode 100644
--- /dev/null
+++ b/Engine/Plugins/AI/MassCrowd/Source/MassCrowd/Public/MassCrowdReplicator.h
@@ -0,0 +1,27 @@
+// Copyright Epic Games, Inc. All Rights Reserved.
+
+#pragma once
+
+#include "MassReplicationProcessor.h"
+
+#include "MassCrowdReplicator.generated.h"
+
+/** Class that handles replication and only runs on the server. It queries Mass entity fragments and sets those values when appropriate using the MassClientBubbleHandler. */
+UCLASS()
+class MASSCROWD_API UMassCrowdReplicator : public UMassReplicatorBase
+{
+	GENERATED_BODY()
+
+public:
+	/**
+	 * Overridden to add specific entity query requirements for replication.
+	 * Usually we add replication processor handler requirements.
+	 */
+	void AddRequirements(FMassEntityQuery& EntityQuery) override;
+	
+	/**
+	 * Overridden to process the client replication.
+	 * This methods should call CalculateClientReplication with the appropriate callback implementation.
+	 */
+	void ProcessClientReplication(FMassExecutionContext& Context, FMassReplicationContext& ReplicationContext) override;
+};
\ No newline at end of file
diff --git a/Engine/Plugins/Runtime/MassGameplay/Source/MassReplication/Private/MassReplicationFragments.cpp b/Engine/Plugins/Runtime/MassGameplay/Source/MassReplication/Private/MassReplicationFragments.cpp
--- a/Engine/Plugins/Runtime/MassGameplay/Source/MassReplication/Private/MassReplicationFragments.cpp
+++ b/Engine/Plugins/Runtime/MassGameplay/Source/MassReplication/Private/MassReplicationFragments.cpp
@@ -2,7 +2,7 @@
 
 #include "MassReplicationFragments.h"
 #include "MassReplicationSubsystem.h"
-#include "MassReplicationProcessorBase.h"
+#include "MassReplicationProcessor.h"
 #include "Engine/World.h"
 
 //----------------------------------------------------------------------//
diff --git a/Engine/Plugins/Runtime/MassGameplay/Source/MassReplication/Private/MassReplicationProcessor.cpp b/Engine/Plugins/Runtime/MassGameplay/Source/MassReplication/Private/MassReplicationProcessor.cpp
new file mode 100644
--- /dev/null
+++ b/Engine/Plugins/Runtime/MassGameplay/Source/MassReplication/Private/MassReplicationProcessor.cpp
@@ -0,0 +1,222 @@
+// Copyright Epic Games, Inc. All Rights Reserved.UMassSimulationSettings
+
+#include "MassReplicationProcessor.h"
+#include "MassClientBubbleHandler.h"
+#include "MassLODSubsystem.h"
+#include "MassCommonFragments.h"
+
+namespace UE::Mass::Replication
+{
+	int32 bDebugReplicationLOD = 0;
+	FAutoConsoleVariableRef CVarDebugReplicationViewerLOD(TEXT("ai.debug.ReplicationLOD"), bDebugReplicationLOD, TEXT("Debug Replication LOD"), ECVF_Cheat);
+} // UE::Mass::Crowd
+
+//----------------------------------------------------------------------//
+//  UMassReplicationProcessor
+//----------------------------------------------------------------------//
+UMassReplicationProcessor::UMassReplicationProcessor()
+{
+#if !UE_ALLOW_DEBUG_REPLICATION_BUBBLES_STANDALONE
+	ExecutionFlags = int32(EProcessorExecutionFlags::Server);
+#else
+	ExecutionFlags = int32(EProcessorExecutionFlags::All);
+#endif // UE_ALLOW_DEBUG_REPLICATION_BUBBLES_STANDALONE
+
+	ProcessingPhase = EMassProcessingPhase::PostPhysics;
+}
+
+void UMassReplicationProcessor::ConfigureQueries()
+{
+	CollectViewerInfoQuery.AddRequirement<FTransformFragment>(EMassFragmentAccess::ReadOnly);
+	CollectViewerInfoQuery.AddRequirement<FMassReplicationViewerInfoFragment>(EMassFragmentAccess::ReadOnly);
+	CollectViewerInfoQuery.AddSharedRequirement<FMassReplicationSharedFragment>(EMassFragmentAccess::ReadWrite);
+
+	CalculateLODQuery.AddRequirement<FMassReplicationViewerInfoFragment>(EMassFragmentAccess::ReadOnly);
+	CalculateLODQuery.AddRequirement<FMassReplicationLODFragment>(EMassFragmentAccess::ReadWrite);
+	CalculateLODQuery.AddConstSharedRequirement<FMassReplicationParameters>();
+	CalculateLODQuery.AddSharedRequirement<FMassReplicationSharedFragment>(EMassFragmentAccess::ReadWrite);
+
+	AdjustLODDistancesQuery.AddRequirement<FMassReplicationViewerInfoFragment>(EMassFragmentAccess::ReadOnly);
+	AdjustLODDistancesQuery.AddRequirement<FMassReplicationLODFragment>(EMassFragmentAccess::ReadWrite);
+	AdjustLODDistancesQuery.AddSharedRequirement<FMassReplicationSharedFragment>(EMassFragmentAccess::ReadWrite);
+	AdjustLODDistancesQuery.SetArchetypeFilter([](const FMassExecutionContext& Context)
+	{
+		const FMassReplicationSharedFragment& LODSharedFragment = Context.GetSharedFragment<FMassReplicationSharedFragment>();
+		return LODSharedFragment.bHasAdjustedDistancesFromCount;
+	});
+
+	EntityQuery.AddRequirement<FMassNetworkIDFragment>(EMassFragmentAccess::ReadOnly);
+	EntityQuery.AddRequirement<FReplicationTemplateIDFragment>(EMassFragmentAccess::ReadOnly);
+	EntityQuery.AddRequirement<FMassReplicationLODFragment>(EMassFragmentAccess::ReadWrite);
+	EntityQuery.AddRequirement<FMassReplicatedAgentFragment>(EMassFragmentAccess::ReadWrite);
+	EntityQuery.AddConstSharedRequirement<FMassReplicationParameters>();
+	EntityQuery.AddSharedRequirement<FMassReplicationSharedFragment>(EMassFragmentAccess::ReadWrite);
+}
+
+void UMassReplicationProcessor::Initialize(UObject& Owner)
+{
+	Super::Initialize(Owner);
+
+#if UE_REPLICATION_COMPILE_SERVER_CODE
+	ReplicationSubsystem = UWorld::GetSubsystem<UMassReplicationSubsystem>(World);
+
+	check(ReplicationSubsystem);
+#endif //UE_REPLICATION_COMPILE_SERVER_CODE
+}
+
+void UMassReplicationProcessor::PrepareExecution(UMassEntitySubsystem& EntitySubsystem)
+{
+#if UE_REPLICATION_COMPILE_SERVER_CODE
+
+	check(ReplicationSubsystem);
+
+	//first synchronize clients and viewers
+	ReplicationSubsystem->SynchronizeClientsAndViewers();
+
+	check(LODSubsystem);
+	const TArray<FViewerInfo>& Viewers = LODSubsystem->GetViewers();
+	EntitySubsystem.ForEachSharedFragment<FMassReplicationSharedFragment>([this, &Viewers](FMassReplicationSharedFragment& RepSharedFragment)
+	{
+		if (!RepSharedFragment.bEntityQueryInitialized)
+		{
+			RepSharedFragment.EntityQuery = EntityQuery;
+			RepSharedFragment.EntityQuery.SetArchetypeFilter([&RepSharedFragment](const FMassExecutionContext& Context)
+			{
+				const FMassReplicationSharedFragment& CurRepSharedFragment = Context.GetSharedFragment<FMassReplicationSharedFragment>();
+				return &CurRepSharedFragment == &RepSharedFragment;
+			});
+			RepSharedFragment.CachedReplicator->AddRequirements(RepSharedFragment.EntityQuery);
+			RepSharedFragment.bEntityQueryInitialized = true;
+		}
+
+		RepSharedFragment.LODCollector.PrepareExecution(Viewers);
+		RepSharedFragment.LODCalculator.PrepareExecution(Viewers);
+
+		const TArray<FMassClientHandle>& CurrentClientHandles = ReplicationSubsystem->GetClientReplicationHandles();
+		const int32 MinNumHandles = FMath::Min(RepSharedFragment.CachedClientHandles.Num(), CurrentClientHandles.Num());
+
+		//check to see if we don't have enough cached client handles
+		if (RepSharedFragment.CachedClientHandles.Num() < CurrentClientHandles.Num())
+		{
+			RepSharedFragment.CachedClientHandles.Reserve(CurrentClientHandles.Num());
+			RepSharedFragment.bBubbleChanged.Reserve(CurrentClientHandles.Num());
+			RepSharedFragment.BubbleInfos.Reserve(CurrentClientHandles.Num());
+
+			for (int32 Idx = RepSharedFragment.CachedClientHandles.Num(); Idx < CurrentClientHandles.Num(); ++Idx)
+			{
+				const FMassClientHandle& CurrentClientHandle = CurrentClientHandles[Idx];
+
+				RepSharedFragment.CachedClientHandles.Add(CurrentClientHandle);
+				RepSharedFragment.bBubbleChanged.Add(true);
+				AMassClientBubbleInfoBase* Info = CurrentClientHandle.IsValid() ?
+					ReplicationSubsystem->GetClientBubbleChecked(RepSharedFragment.BubbleInfoClassHandle, CurrentClientHandle) :
+					nullptr;
+
+				check(Info);
+
+				RepSharedFragment.BubbleInfos.Add(Info);
+			}
+		}
+		//check to see if we have too many cached client handles
+		else if (RepSharedFragment.CachedClientHandles.Num() > CurrentClientHandles.Num())
+		{
+			const int32 NumRemove = RepSharedFragment.CachedClientHandles.Num() - CurrentClientHandles.Num();
+
+			RepSharedFragment.CachedClientHandles.RemoveAt(CurrentClientHandles.Num(), NumRemove, /* bAllowShrinking */ false);
+			RepSharedFragment.bBubbleChanged.RemoveAt(CurrentClientHandles.Num(), NumRemove, /* bAllowShrinking */ false);
+			RepSharedFragment.BubbleInfos.RemoveAt(CurrentClientHandles.Num(), NumRemove, /* bAllowShrinking */ false);
+		}
+
+		//check to see if any cached client handles have changed, if they have set the BubbleInfo[] appropriately
+		for (int32 Idx = 0; Idx < MinNumHandles; ++Idx)
+		{
+			const FMassClientHandle& CurrentClientHandle = CurrentClientHandles[Idx];
+			FMassClientHandle& CachedClientHandle = RepSharedFragment.CachedClientHandles[Idx];
+
+			const bool bChanged = (RepSharedFragment.bBubbleChanged[Idx] = (CurrentClientHandle != CachedClientHandle));
+			if (bChanged)
+			{
+				AMassClientBubbleInfoBase* Info = CurrentClientHandle.IsValid() ?
+					ReplicationSubsystem->GetClientBubbleChecked(RepSharedFragment.BubbleInfoClassHandle, CurrentClientHandle) :
+					nullptr;
+
+				check(Info);
+
+				RepSharedFragment.BubbleInfos[Idx] = Info;
+				CachedClientHandle = CurrentClientHandle;
+			}
+		}
+	});
+
+#endif //UE_REPLICATION_COMPILE_SERVER_CODE
+}
+
+void UMassReplicationProcessor::Execute(UMassEntitySubsystem& EntitySubsystem, FMassExecutionContext& Context)
+{
+#if UE_REPLICATION_COMPILE_SERVER_CODE
+	check(World);
+	check(LODSubsystem);
+	check(ReplicationSubsystem);
+
+	{
+		QUICK_SCOPE_CYCLE_COUNTER(UMassReplicationProcessor_Preperation);
+		PrepareExecution(EntitySubsystem);
+	}
+
+	{
+		QUICK_SCOPE_CYCLE_COUNTER(UMassReplicationProcessor_LODCaculation);
+		CollectViewerInfoQuery.ForEachEntityChunk(EntitySubsystem, Context, [this](FMassExecutionContext& Context)
+		{
+			const TConstArrayView<FTransformFragment> LocationList = Context.GetFragmentView<FTransformFragment>();
+			const TArrayView<FMassReplicationViewerInfoFragment> ViewersInfoList = Context.GetMutableFragmentView<FMassReplicationViewerInfoFragment>();
+			FMassReplicationSharedFragment& RepSharedFragment = Context.GetMutableSharedFragment<FMassReplicationSharedFragment>();
+			RepSharedFragment.LODCollector.CollectLODInfo(Context, LocationList, ViewersInfoList, ViewersInfoList);
+		});
+
+		CalculateLODQuery.ForEachEntityChunk(EntitySubsystem, Context, [this](FMassExecutionContext& Context)
+		{
+			const TConstArrayView<FMassReplicationViewerInfoFragment> ViewersInfoList = Context.GetFragmentView<FMassReplicationViewerInfoFragment>();
+			const TArrayView<FMassReplicationLODFragment> ViewerLODList = Context.GetMutableFragmentView<FMassReplicationLODFragment>();
+			FMassReplicationSharedFragment& RepSharedFragment = Context.GetMutableSharedFragment<FMassReplicationSharedFragment>();
+			RepSharedFragment.LODCalculator.CalculateLOD(Context, ViewersInfoList, ViewerLODList, ViewersInfoList);
+		});
+
+		EntitySubsystem.ForEachSharedFragment<FMassReplicationSharedFragment>([](FMassReplicationSharedFragment& RepSharedFragment)
+		{
+			RepSharedFragment.bHasAdjustedDistancesFromCount = RepSharedFragment.LODCalculator.AdjustDistancesFromCount();
+		});
+
+		AdjustLODDistancesQuery.ForEachEntityChunk(EntitySubsystem, Context, [this](FMassExecutionContext& Context)
+		{
+			const TConstArrayView<FMassReplicationViewerInfoFragment> ViewersInfoList = Context.GetFragmentView<FMassReplicationViewerInfoFragment>();
+			const TArrayView<FMassReplicationLODFragment> ViewerLODList = Context.GetMutableFragmentView<FMassReplicationLODFragment>();
+			FMassReplicationSharedFragment& RepSharedFragment = Context.GetMutableSharedFragment<FMassReplicationSharedFragment>();
+			RepSharedFragment.LODCalculator.AdjustLODFromCount(Context, ViewersInfoList, ViewerLODList, ViewersInfoList);
+		});
+	}
+
+	{
+		QUICK_SCOPE_CYCLE_COUNTER(UMassReplicationProcessor_ProcessClientReplication);
+		FMassReplicationContext ReplicationContext(*World, *LODSubsystem, *ReplicationSubsystem);
+		EntitySubsystem.ForEachSharedFragment<FMassReplicationSharedFragment>([this, &EntitySubsystem, &Context, &ReplicationContext](FMassReplicationSharedFragment& RepSharedFragment)
+		{
+			RepSharedFragment.EntityQuery.ForEachEntityChunk(EntitySubsystem, Context, [&ReplicationContext, &RepSharedFragment](FMassExecutionContext& Context)
+			{
+				RepSharedFragment.CachedReplicator->ProcessClientReplication(Context, ReplicationContext);
+			});
+		});
+	}
+
+	// Optional debug display
+	if (UE::Mass::Replication::bDebugReplicationLOD)
+	{
+		EntityQuery.ForEachEntityChunk(EntitySubsystem, Context, [this](FMassExecutionContext& Context)
+		{
+			const TConstArrayView<FTransformFragment> TransformList = Context.GetFragmentView<FTransformFragment>();
+			const TConstArrayView<FMassReplicationLODFragment> ViewerLODList = Context.GetFragmentView<FMassReplicationLODFragment>();
+			FMassReplicationSharedFragment& RepSharedFragment = Context.GetMutableSharedFragment<FMassReplicationSharedFragment>();
+			RepSharedFragment.LODCalculator.DebugDisplayLOD(Context, ViewerLODList, TransformList, World);
+		});
+	}
+#endif //UE_REPLICATION_COMPILE_SERVER_CODE
+}
\ No newline at end of file
diff --git a/Engine/Plugins/Runtime/MassGameplay/Source/MassReplication/Private/MassReplicationProcessorBase.cpp b/Engine/Plugins/Runtime/MassGameplay/Source/MassReplication/Private/MassReplicationProcessorBase.cpp
deleted file mode 100644
diff --git a/Engine/Plugins/Runtime/MassGameplay/Source/MassReplication/Public/MassReplicationProcessor.h b/Engine/Plugins/Runtime/MassGameplay/Source/MassReplication/Public/MassReplicationProcessor.h
new file mode 100644
--- /dev/null
+++ b/Engine/Plugins/Runtime/MassGameplay/Source/MassReplication/Public/MassReplicationProcessor.h
@@ -0,0 +1,197 @@
+// Copyright Epic Games, Inc. All Rights Reserved.
+
+#pragma once
+
+#include "Engine/World.h"
+#include "MassProcessor.h"
+#include "MassCommonTypes.h"
+#include "MassSimulationLOD.h"
+#include "MassReplicationTypes.h"
+#include "MassReplicationFragments.h"
+#include "MassSpawnerTypes.h"
+#include "MassLODSubsystem.h"
+#include "MassReplicationSubsystem.h"
+#include "MassReplicationProcessor.generated.h"
+
+class UMassReplicationSubsystem;
+class AMassClientBubbleInfoBase;
+class UWorld;
+
+/** 
+ *  Base processor that handles replication and only runs on the server. You should derive from this per entity type (that require different replication processing). It and its derived classes 
+ *  query Mass entity fragments and set those values for replication when appropriate, using the MassClientBubbleHandler.
+ */
+UCLASS()
+class MASSREPLICATION_API UMassReplicationProcessor : public UMassLODProcessorBase
+{
+	GENERATED_BODY()
+
+public:
+	UMassReplicationProcessor();
+
+protected:
+	virtual void ConfigureQueries() override;
+	virtual void Initialize(UObject& Owner) override;
+	virtual void Execute(UMassEntitySubsystem& EntitySubsystem, FMassExecutionContext& Context) override;
+
+	void PrepareExecution(UMassEntitySubsystem& EntitySubsystem);
+
+protected:
+
+	UPROPERTY(Transient)
+	UMassReplicationSubsystem* ReplicationSubsystem = nullptr;
+
+	FMassEntityQuery CollectViewerInfoQuery;
+	FMassEntityQuery CalculateLODQuery;
+	FMassEntityQuery AdjustLODDistancesQuery;
+	FMassEntityQuery EntityQuery;
+};
+
+
+struct FMassReplicationContext
+{
+	FMassReplicationContext(UWorld& InWorld, UMassLODSubsystem& InLODSubsystem, UMassReplicationSubsystem& InReplicationSubsystem)
+		: World(InWorld)
+		, LODSubsystem(InLODSubsystem)
+		, ReplicationSubsystem(InReplicationSubsystem)
+	{}
+
+	UWorld& World;
+	UMassLODSubsystem& LODSubsystem;
+	UMassReplicationSubsystem& ReplicationSubsystem;
+};
+
+UCLASS(Abstract)
+class MASSREPLICATION_API UMassReplicatorBase : public UObject
+{
+	GENERATED_BODY()
+
+public:
+	/**
+	 * Must override to add specific entity query requirements for replication
+	 * Usually we add replication processor handler requirements
+	 */
+	virtual void AddRequirements(FMassEntityQuery& EntityQuery) PURE_VIRTUAL(UMassReplicatorBase::ConfigureQueries, );
+
+	/**
+	 * Must override to process the client replication
+	 * This methods should call CalculateClientReplication with the appropriate callback implementation
+	 */
+	virtual void ProcessClientReplication(FMassExecutionContext& Context, FMassReplicationContext& ReplicationContext) PURE_VIRTUAL(UMassReplicatorBase::ProcessClientReplication, );
+
+protected:
+	/**
+	 *  Implemented as straight template callbacks as when profiled this was faster than TFunctionRef. Its probably easier to pass Lamdas in to these
+	 *  but Functors can also be used as well as TFunctionRefs etc. Its also fairly straight forward to call member functions via some Lamda glue code
+	 */
+	template<typename AgentArrayItem, typename CacheViewsCallback, typename AddEntityCallback, typename ModifyEntityCallback, typename RemoveEntityCallback>
+	static void CalculateClientReplication(FMassExecutionContext& Context, FMassReplicationContext& ReplicationContext, CacheViewsCallback&& CacheViews, AddEntityCallback&& AddEntity, ModifyEntityCallback&& ModifyEntity, RemoveEntityCallback&& RemoveEntity);
+};
+
+template<typename AgentArrayItem, typename CacheViewsCallback, typename AddEntityCallback, typename ModifyEntityCallback, typename RemoveEntityCallback>
+void UMassReplicatorBase::CalculateClientReplication(FMassExecutionContext& Context, FMassReplicationContext& ReplicationContext, CacheViewsCallback&& CacheViews, AddEntityCallback&& AddEntity, ModifyEntityCallback&& ModifyEntity, RemoveEntityCallback&& RemoveEntity)
+{
+#if UE_REPLICATION_COMPILE_SERVER_CODE
+
+	const int32 NumEntities = Context.GetNumEntities();
+
+	TConstArrayView<FMassNetworkIDFragment> NetworkIDList = Context.GetFragmentView<FMassNetworkIDFragment>();
+	TArrayView<FMassReplicationLODFragment> ViewerLODList = Context.GetMutableFragmentView<FMassReplicationLODFragment>();
+	TArrayView<FMassReplicatedAgentFragment> ReplicatedAgentList = Context.GetMutableFragmentView<FMassReplicatedAgentFragment>();
+	TConstArrayView<FReplicationTemplateIDFragment> TemplateIDList = Context.GetFragmentView<FReplicationTemplateIDFragment>();
+	FMassReplicationSharedFragment& RepSharedFragment = Context.GetMutableSharedFragment<FMassReplicationSharedFragment>();
+
+	CacheViews(Context);
+
+	const float Time = ReplicationContext.World.GetRealTimeSeconds();
+
+	for (int EntityIdx = 0; EntityIdx < NumEntities; EntityIdx++)
+	{
+		FMassReplicatedAgentFragment& AgentFragment = ReplicatedAgentList[EntityIdx];
+
+		if (AgentFragment.AgentsData.Num() < RepSharedFragment.CachedClientHandles.Num())
+		{
+			AgentFragment.AgentsData.AddDefaulted(RepSharedFragment.CachedClientHandles.Num() - AgentFragment.AgentsData.Num());
+		}
+		else if (AgentFragment.AgentsData.Num() > RepSharedFragment.CachedClientHandles.Num())
+		{
+			AgentFragment.AgentsData.RemoveAt(RepSharedFragment.CachedClientHandles.Num(), AgentFragment.AgentsData.Num() - RepSharedFragment.CachedClientHandles.Num(), /* bAllowShrinking */ false);
+		}
+
+		for (int32 ClientIdx = 0; ClientIdx < RepSharedFragment.CachedClientHandles.Num(); ++ClientIdx)
+		{
+			const FMassClientHandle& ClientHandle = RepSharedFragment.CachedClientHandles[ClientIdx];
+
+			if (ClientHandle.IsValid())
+			{
+				checkSlow(RepSharedFragment.BubbleInfos[ClientHandle.GetIndex()] != nullptr);
+
+				FMassReplicatedAgentArrayData& AgentData = AgentFragment.AgentsData[ClientIdx];
+
+				//if the bubble has changed, set the handle Invalid. We will set this to something valid if the agent is going to replicate to the bubble
+				//When a bubble has changed the existing AMassCrowdClientBubbleInfo will reset all the data associated with it.
+				if (RepSharedFragment.bBubbleChanged[ClientIdx])
+				{
+					AgentData.Invalidate();
+				}
+
+				//now we need to see what our highest viewer LOD is on this client (split screen etc), we can use the unsafe version as we have already checked all the
+				//CachedClientHandles for validity
+				const FClientViewerHandles& ClientViewers = ReplicationContext.ReplicationSubsystem.GetClientViewersChecked(ClientHandle);
+
+				EMassLOD::Type HighestLOD = EMassLOD::Off;
+
+				for (const FMassViewerHandle& ViewerHandle : ClientViewers.Handles)
+				{
+					//this should always we valid as we synchronized the viewers just previously
+					check(ReplicationContext.LODSubsystem.IsValidViewer(ViewerHandle));
+
+					const EMassLOD::Type MassLOD = ViewerLODList[EntityIdx].LODPerViewer[ViewerHandle.GetIndex()];
+					check(MassLOD <= EMassLOD::Off);
+
+					if (HighestLOD > MassLOD)
+					{
+						HighestLOD = MassLOD;
+					}
+				}
+
+				if (HighestLOD < EMassLOD::Off)
+				{
+#if UE_ALLOW_DEBUG_REPLICATION
+					AgentData.LOD = HighestLOD;
+#endif // UE_ALLOW_DEBUG_REPLICATION
+
+					//if the handle isn't valid we need to add the agent
+					if (!AgentData.Handle.IsValid())
+					{
+						typename AgentArrayItem::FReplicatedAgentType ReplicatedAgent;
+
+						const FMassNetworkIDFragment& NetIDFragment = NetworkIDList[EntityIdx];
+						const FReplicationTemplateIDFragment& TemplateIDFragment = TemplateIDList[EntityIdx];
+
+						ReplicatedAgent.SetNetID(NetIDFragment.NetID);
+						ReplicatedAgent.SetTemplateID(TemplateIDFragment.ID);
+
+						AgentData.Handle = AddEntity(Context, EntityIdx, ReplicatedAgent, ClientHandle);
+
+						AgentData.LastUpdateTime = Time;
+					}
+					else
+					{
+						ModifyEntity(Context, EntityIdx, HighestLOD, Time, AgentData.Handle, ClientHandle);
+					}
+				}
+				else
+				{
+					// as this is a fresh handle, if its valid then we can use the unsafe remove function
+					if (AgentData.Handle.IsValid())
+					{
+						RemoveEntity(Context, AgentData.Handle, ClientHandle);
+						AgentData.Invalidate();
+					}
+				}
+			}
+		}
+	}
+#endif //UE_REPLICATION_COMPILE_SERVER_CODE
+}
diff --git a/Engine/Plugins/Runtime/MassGameplay/Source/MassReplication/Public/MassReplicationProcessorBase.h b/Engine/Plugins/Runtime/MassGameplay/Source/MassReplication/Public/MassReplicationProcessorBase.h
deleted file mode 100644
