Moved SpawnCount <= 0 early out checks from AMassSpawner to MassEntityZoneGraphSpawnPointsGenerator & MassEntityEQSSpawnPointGenerator.

This allows spawners to still function with other spawn data generators which ignore the Count.

diff --git a/Engine/Plugins/Runtime/MassGameplay/Source/MassSpawner/Private/MassEntityEQSSpawnPointsGenerator.cpp b/Engine/Plugins/Runtime/MassGameplay/Source/MassSpawner/Private/MassEntityEQSSpawnPointsGenerator.cpp
--- a/Engine/Plugins/Runtime/MassGameplay/Source/MassSpawner/Private/MassEntityEQSSpawnPointsGenerator.cpp
+++ b/Engine/Plugins/Runtime/MassGameplay/Source/MassSpawner/Private/MassEntityEQSSpawnPointsGenerator.cpp
@@ -14,6 +14,12 @@ UMassEntityEQSSpawnPointsGenerator::UMassEntityEQSSpawnPointsGenerator()
 
 void UMassEntityEQSSpawnPointsGenerator::Generate(UObject& QueryOwner, TConstArrayView<FMassSpawnedEntityType> EntityTypes, const int32 Count, FFinishedGeneratingSpawnDataSignature& FinishedGeneratingSpawnPointsDelegate) const
 {
+	if (Count <= 0)
+	{
+		FinishedGeneratingSpawnPointsDelegate.Execute(TArray<FMassEntitySpawnDataGeneratorResult>());
+		return;
+	}
+	
 	// Need to copy the request as it is called inside a CDO and CDO states cannot be changed.
 	FEQSParametrizedQueryExecutionRequest EQSRequestInstanced = EQSRequest;
 	if (EQSRequestInstanced.IsValid() == false)
diff --git a/Engine/Plugins/Runtime/MassGameplay/Source/MassSpawner/Private/MassEntityZoneGraphSpawnPointsGenerator.cpp b/Engine/Plugins/Runtime/MassGameplay/Source/MassSpawner/Private/MassEntityZoneGraphSpawnPointsGenerator.cpp
--- a/Engine/Plugins/Runtime/MassGameplay/Source/MassSpawner/Private/MassEntityZoneGraphSpawnPointsGenerator.cpp
+++ b/Engine/Plugins/Runtime/MassGameplay/Source/MassSpawner/Private/MassEntityZoneGraphSpawnPointsGenerator.cpp
@@ -12,6 +12,12 @@
 
 void UMassEntityZoneGraphSpawnPointsGenerator::Generate(UObject& QueryOwner, TConstArrayView<FMassSpawnedEntityType> EntityTypes, int32 Count, FFinishedGeneratingSpawnDataSignature& FinishedGeneratingSpawnPointsDelegate) const
 {
+	if (Count <= 0)
+	{
+		FinishedGeneratingSpawnPointsDelegate.Execute(TArray<FMassEntitySpawnDataGeneratorResult>());
+		return;
+	}
+	
 	const UZoneGraphSubsystem* ZoneGraph = UWorld::GetSubsystem<UZoneGraphSubsystem>(QueryOwner.GetWorld());
 	if (ZoneGraph == nullptr)
 	{
diff --git a/Engine/Plugins/Runtime/MassGameplay/Source/MassSpawner/Private/MassSpawner.cpp b/Engine/Plugins/Runtime/MassGameplay/Source/MassSpawner/Private/MassSpawner.cpp
--- a/Engine/Plugins/Runtime/MassGameplay/Source/MassSpawner/Private/MassSpawner.cpp
+++ b/Engine/Plugins/Runtime/MassGameplay/Source/MassSpawner/Private/MassSpawner.cpp
@@ -269,12 +269,6 @@ void AMassSpawner::RegisterEntityTemplates()
 
 void AMassSpawner::DoSpawning()
 {
-	const int32 TotalSpawnCount = GetSpawnCount();
-	if (TotalSpawnCount <= 0)
-	{
-		return;
-	}
-	
 	// no spawn point generators configured. Let user know and fall back to the spawner's location
 	if (SpawnDataGenerators.Num() == 0)
 	{
@@ -309,23 +303,21 @@ void AMassSpawner::DoSpawning()
 	TArray<FSoftObjectPath> AssetsToLoad;
 	for (const FMassSpawnedEntityType& EntityType : EntityTypes)
 	{
-		const int32 EntityCount = int32(TotalSpawnCount * EntityType.Proportion / TotalProportion);
-		if (EntityCount > 0)
+		if (!EntityType.IsLoaded())
 		{
-			if (!EntityType.IsLoaded())
-			{
-				AssetsToLoad.Add(EntityType.EntityConfig.ToSoftObjectPath());
-			}
+			AssetsToLoad.Add(EntityType.EntityConfig.ToSoftObjectPath());
 		}
 	}
-
+	
+	const int32 TotalSpawnCount = GetSpawnCount();
+	
 	auto GenerateSpawningPoints = [this, TotalSpawnCount, TotalProportion]()
 	{
 		int32 SpawnCountRemaining = TotalSpawnCount;
 		float ProportionRemaining = TotalProportion;
 		for (FMassSpawnDataGenerator& Generator : SpawnDataGenerators)
 		{
-			if (Generator.Proportion == 0.0f || ProportionRemaining <= 0.0f || SpawnCountRemaining <= 0)
+			if (Generator.Proportion == 0.0f || ProportionRemaining <= 0.0f)
 			{
 				// If there's nothing to spawn, mark the generator done as OnSpawnDataGenerationFinished() will wait for all generators to complete before the actual spawning.
 				Generator.bDataGenerated = true;
@@ -336,13 +328,11 @@ void AMassSpawner::DoSpawning()
 			{
 				const float ProportionRatio = FMath::Min(Generator.Proportion / ProportionRemaining, 1.0f);
 				const int32 SpawnCount = FMath::CeilToInt(SpawnCountRemaining * ProportionRatio);
-				if (SpawnCount > 0)
-				{
-					FFinishedGeneratingSpawnDataSignature Delegate = FFinishedGeneratingSpawnDataSignature::CreateUObject(this, &AMassSpawner::OnSpawnDataGenerationFinished, &Generator);
-					Generator.GeneratorInstance->Generate(*this, EntityTypes, SpawnCount, Delegate);
-					SpawnCountRemaining -= SpawnCount;
-					ProportionRemaining -= Generator.Proportion;
-				}
+				
+				FFinishedGeneratingSpawnDataSignature Delegate = FFinishedGeneratingSpawnDataSignature::CreateUObject(this, &AMassSpawner::OnSpawnDataGenerationFinished, &Generator);
+				Generator.GeneratorInstance->Generate(*this, EntityTypes, SpawnCount, Delegate);
+				SpawnCountRemaining -= SpawnCount;
+				ProportionRemaining -= Generator.Proportion;
 			}
 		}
 	};
