Integration of cl 19121884. needed to continue working in main.

Fix visualization LOD Multithreading issue

Since we do not have other ways to create data dependency, we need to create different processors to have independent dependencies (crowd, Traffic, Intersection are ran in parallel). Doing that does not work well with the current implementation where we iterate on all shared fragments of a certain type, because we will iterate on all 3 types of shared fragment in all 3 derived processor class. If the 3 processors are ran in parallel, this is where the problem happens. By specifying a Filter tag in the shared fragment to match the one in the processor so that it only handle the matching tag, it fixes the problem but not it is not an ideal fix.

#rnx

diff --git a/Engine/Plugins/Runtime/MassGameplay/Source/MassRepresentation/Private/MassRepresentationFragments.cpp b/Engine/Plugins/Runtime/MassGameplay/Source/MassRepresentation/Private/MassRepresentationFragments.cpp
--- a/Engine/Plugins/Runtime/MassGameplay/Source/MassRepresentation/Private/MassRepresentationFragments.cpp
+++ b/Engine/Plugins/Runtime/MassGameplay/Source/MassRepresentation/Private/MassRepresentationFragments.cpp
@@ -5,4 +5,5 @@
 FMassVisualizationLODSharedFragment::FMassVisualizationLODSharedFragment(const FMassVisualizationLODParameters& LODParams)
 {
 	LODCalculator.Initialize(LODParams.BaseLODDistance, LODParams.BufferHysteresisOnDistancePercentage / 100.f, LODParams.LODMaxCount, nullptr, LODParams.DistanceToFrustum, LODParams.DistanceToFrustumHysteresis, LODParams.VisibleLODDistance);
+	FilterTag = LODParams.FilterTag;
 }
\ No newline at end of file
diff --git a/Engine/Plugins/Runtime/MassGameplay/Source/MassRepresentation/Private/MassVisualizationLODProcessor.cpp b/Engine/Plugins/Runtime/MassGameplay/Source/MassRepresentation/Private/MassVisualizationLODProcessor.cpp
--- a/Engine/Plugins/Runtime/MassGameplay/Source/MassRepresentation/Private/MassVisualizationLODProcessor.cpp
+++ b/Engine/Plugins/Runtime/MassGameplay/Source/MassRepresentation/Private/MassVisualizationLODProcessor.cpp
@@ -61,9 +61,12 @@ void UMassVisualizationLODProcessor::Execute(UMassEntitySubsystem& EntitySubsyst
 		TRACE_CPUPROFILER_EVENT_SCOPE(PrepareExecution)
 		check(LODSubsystem);
 		const TArray<FViewerInfo>& Viewers = LODSubsystem->GetViewers();
-		EntitySubsystem.ForEachSharedFragment<FMassVisualizationLODSharedFragment>([&Viewers](FMassVisualizationLODSharedFragment& LODSharedFragment)
+		EntitySubsystem.ForEachSharedFragment<FMassVisualizationLODSharedFragment>([this, &Viewers](FMassVisualizationLODSharedFragment& LODSharedFragment)
 		{
-			LODSharedFragment.LODCalculator.PrepareExecution(Viewers);
+			if (FilterTag == LODSharedFragment.FilterTag)
+			{
+				LODSharedFragment.LODCalculator.PrepareExecution(Viewers);
+			}
 		});
 	}
 
@@ -83,9 +86,12 @@ void UMassVisualizationLODProcessor::Execute(UMassEntitySubsystem& EntitySubsyst
 
 	{
 		TRACE_CPUPROFILER_EVENT_SCOPE(AdjustDistanceAndLODFromCount)
-		EntitySubsystem.ForEachSharedFragment<FMassVisualizationLODSharedFragment>([](FMassVisualizationLODSharedFragment& LODSharedFragment)
+		EntitySubsystem.ForEachSharedFragment<FMassVisualizationLODSharedFragment>([this](FMassVisualizationLODSharedFragment& LODSharedFragment)
 		{
-			LODSharedFragment.bHasAdjustedDistancesFromCount = LODSharedFragment.LODCalculator.AdjustDistancesFromCount();
+			if (FilterTag == LODSharedFragment.FilterTag)
+			{
+				LODSharedFragment.bHasAdjustedDistancesFromCount = LODSharedFragment.LODCalculator.AdjustDistancesFromCount();
+			}
 		});
 
 		CloseEntityAdjustDistanceQuery.ForEachEntityChunk(EntitySubsystem, Context, [this](FMassExecutionContext& Context)
@@ -110,4 +116,4 @@ void UMassVisualizationLODProcessor::Execute(UMassEntitySubsystem& EntitySubsyst
 			LODSharedFragment.LODCalculator.DebugDisplayLOD(Context, RepresentationLODList, TransformList, World);
 		});
 	}
-}
\ No newline at end of file
+}
diff --git a/Engine/Plugins/Runtime/MassGameplay/Source/MassRepresentation/Public/MassRepresentationFragments.h b/Engine/Plugins/Runtime/MassGameplay/Source/MassRepresentation/Public/MassRepresentationFragments.h
--- a/Engine/Plugins/Runtime/MassGameplay/Source/MassRepresentation/Public/MassRepresentationFragments.h
+++ b/Engine/Plugins/Runtime/MassGameplay/Source/MassRepresentation/Public/MassRepresentationFragments.h
@@ -156,6 +156,10 @@ struct FMassVisualizationLODParameters : public FMassSharedFragment
 	/** Once visible how much further than DistanceToFrustum does the entities need to be before being cull again */
 	UPROPERTY(EditAnywhere, Category = "Mass|LOD", meta = (ClampMin = "0.0", UIMin = "0.0"), config)
 	float DistanceToFrustumHysteresis = 0.0f;
+
+	/** Filter these settings with specified tag */
+	UPROPERTY(EditAnywhere, Category = "Mass|LOD", meta = (BaseStruct = "MassTag"))
+	UScriptStruct* FilterTag;
 };
 
 USTRUCT()
@@ -168,4 +172,7 @@ struct FMassVisualizationLODSharedFragment : public FMassSharedFragment
 
 	TMassLODCalculator<FMassRepresentationLODLogic> LODCalculator;
 	bool bHasAdjustedDistancesFromCount = false;
+
+	UPROPERTY(Transient)
+	const UScriptStruct* FilterTag = nullptr;
 };
diff --git a/Engine/Plugins/Runtime/MassGameplay/Source/MassRepresentation/Public/MassVisualizationLODProcessor.h b/Engine/Plugins/Runtime/MassGameplay/Source/MassRepresentation/Public/MassVisualizationLODProcessor.h
--- a/Engine/Plugins/Runtime/MassGameplay/Source/MassRepresentation/Public/MassVisualizationLODProcessor.h
+++ b/Engine/Plugins/Runtime/MassGameplay/Source/MassRepresentation/Public/MassVisualizationLODProcessor.h
@@ -41,4 +41,7 @@ protected:
 	FMassEntityQuery DebugEntityQuery;
 
 	bool bForceOFFLOD = false;
+
+	UPROPERTY(Transient)
+	const UScriptStruct* FilterTag = nullptr;
 };
\ No newline at end of file
