Reenable UdpateInstance code that disabled because of early integration of Mass Code into Release 5.0. The feature was not integrated yet in release5.0.

#rnx

diff --git a/Engine/Plugins/Runtime/MassGameplay/Source/MassRepresentation/Private/MassVisualizationComponent.cpp b/Engine/Plugins/Runtime/MassGameplay/Source/MassRepresentation/Private/MassVisualizationComponent.cpp
--- a/Engine/Plugins/Runtime/MassGameplay/Source/MassRepresentation/Private/MassVisualizationComponent.cpp
+++ b/Engine/Plugins/Runtime/MassGameplay/Source/MassRepresentation/Private/MassVisualizationComponent.cpp
@@ -219,12 +219,11 @@ void UMassVisualizationComponent::EndVisualChanges()
 			const int32 NewNumInstances = SharedData.StaticMeshInstanceTransforms.Num();
 
 			// Update PerInstanceSMData transforms
-			// @Todo reenable when new UpdateInstances code is in main branch
-// 			if ((bool)UE::MassRepresentation::GCallUpdateInstances)
-// 			{
-// 				InstancedStaticMeshComponent->UpdateInstances(SharedData.UpdateInstanceIds, SharedData.StaticMeshInstanceTransforms, SharedData.StaticMeshInstancePrevTransforms, NumCustomDataFloats, SharedData.StaticMeshInstanceCustomFloats);
-// 			}
-// 			else
+			if ((bool)UE::MassRepresentation::GCallUpdateInstances)
+			{
+				InstancedStaticMeshComponent->UpdateInstances(SharedData.UpdateInstanceIds, SharedData.StaticMeshInstanceTransforms, SharedData.StaticMeshInstancePrevTransforms, NumCustomDataFloats, SharedData.StaticMeshInstanceCustomFloats);
+			}
+			else
 			{
 				// Update NumCustomDataFloats
 				InstancedStaticMeshComponent->NumCustomDataFloats = NumCustomDataFloats;
