[Mass] added missing initialization for property FMassVisualizationLODParameters::FilterTag
#rnx

diff --git a/Engine/Plugins/Runtime/MassGameplay/Source/MassRepresentation/Public/MassRepresentationFragments.h b/Engine/Plugins/Runtime/MassGameplay/Source/MassRepresentation/Public/MassRepresentationFragments.h
--- a/Engine/Plugins/Runtime/MassGameplay/Source/MassRepresentation/Public/MassRepresentationFragments.h
+++ b/Engine/Plugins/Runtime/MassGameplay/Source/MassRepresentation/Public/MassRepresentationFragments.h
@@ -159,7 +159,7 @@ struct FMassVisualizationLODParameters : public FMassSharedFragment
 
 	/** Filter these settings with specified tag */
 	UPROPERTY(EditAnywhere, Category = "Mass|LOD", meta = (BaseStruct = "MassTag"))
-	UScriptStruct* FilterTag;
+	UScriptStruct* FilterTag = nullptr;
 };
 
 USTRUCT()
