AMassSpawner - Added DoDespawning(TConstArrayView<FMassEntityHandle> EntitiesToIgnore) to allow for skipping entities you want to persist.

#rnx

diff --git a/Engine/Plugins/Runtime/MassGameplay/Source/MassSpawner/Private/MassSpawner.cpp b/Engine/Plugins/Runtime/MassGameplay/Source/MassSpawner/Private/MassSpawner.cpp
--- a/Engine/Plugins/Runtime/MassGameplay/Source/MassSpawner/Private/MassSpawner.cpp
+++ b/Engine/Plugins/Runtime/MassGameplay/Source/MassSpawner/Private/MassSpawner.cpp
@@ -494,6 +494,36 @@ void AMassSpawner::DoDespawning()
 	OnDespawningFinishedEvent.Broadcast();
 }
 
+void AMassSpawner::DoDespawning(TConstArrayView<FMassEntityHandle> EntitiesToIgnore)
+{
+	// Remove EntitiesToIgnore from SpawnedEntities so they get skipped by DoDespawning() and add to AllEntitiesToKeep
+	// to restore after.
+	TArray<FSpawnedEntities> AllEntitiesToKeep;
+	for (FSpawnedEntities& SpawnedEntities : AllSpawnedEntities)
+	{
+		FSpawnedEntities* EntitiesToKeep = nullptr;
+		for (const FMassEntityHandle& EntityToIgnore : EntitiesToIgnore)
+		{
+			if (SpawnedEntities.Entities.RemoveSingleSwap(EntityToIgnore, /*bAllowShrinking*/false))
+			{
+				if (!EntitiesToKeep)
+				{
+					EntitiesToKeep = &AllEntitiesToKeep.AddDefaulted_GetRef();
+					EntitiesToKeep->TemplateID = SpawnedEntities.TemplateID;
+				}
+				
+				EntitiesToKeep->Entities.Add(EntityToIgnore);
+			}
+		}
+	}
+
+	// Despawn the remaining entities in AllSpawnedEntities
+	DoDespawning();
+
+	// Restore AllEntitiesToKeep to AllSpawnedEntities so they remain tracked
+	AllSpawnedEntities = AllEntitiesToKeep;
+}
+
 bool AMassSpawner::DespawnEntity(const FMassEntityHandle Entity)
 {
 	UMassSpawnerSubsystem* SpawnerSystem = UWorld::GetSubsystem<UMassSpawnerSubsystem>(GetWorld());
@@ -509,8 +539,7 @@ bool AMassSpawner::DespawnEntity(const FMassEntityHandle Entity)
 		if (Index != INDEX_NONE)
 		{
 			SpawnerSystem->DestroyEntities(SpawnedEntities.TemplateID, MakeArrayView(&Entity, 1));
-			// not using "RemoveAtSwap" on purpose to retain the order which will help creation of sparse chunks if need be
-			SpawnedEntities.Entities.RemoveAt(Index, 1, /*bAllowShrinking=*/false);
+			SpawnedEntities.Entities.RemoveAtSwap(Index, 1, /*bAllowShrinking=*/false);
 			return true;
 		}
 	}
diff --git a/Engine/Plugins/Runtime/MassGameplay/Source/MassSpawner/Public/MassSpawner.h b/Engine/Plugins/Runtime/MassGameplay/Source/MassSpawner/Public/MassSpawner.h
--- a/Engine/Plugins/Runtime/MassGameplay/Source/MassSpawner/Public/MassSpawner.h
+++ b/Engine/Plugins/Runtime/MassGameplay/Source/MassSpawner/Public/MassSpawner.h
@@ -59,6 +59,13 @@ public:
 	UFUNCTION(BlueprintCallable, Category = "Spawning")
 	void DoDespawning();
 
+	/**
+	 * Despawn all mass agent that was spawned by this spawner, except EntitiesToIgnore.
+	 *
+	 * Any EntitiesToIgnore previously spawned by this spawner will remain spawned and tracked by this spawner.   
+	 */
+	void DoDespawning(TConstArrayView<FMassEntityHandle> EntitiesToIgnore);
+
 	UFUNCTION(BlueprintCallable, Category = "Spawning")
 	void ClearTemplates();
 
