Fixed UE::Mass::Debug::GetDebugEntitiesAndLocations's use of GetDebugEntitiesRange

diff --git a/Engine/Plugins/Runtime/MassGameplay/Source/MassCommon/Private/MassCommonTypes.cpp b/Engine/Plugins/Runtime/MassGameplay/Source/MassCommon/Private/MassCommonTypes.cpp
--- a/Engine/Plugins/Runtime/MassGameplay/Source/MassCommon/Private/MassCommonTypes.cpp
+++ b/Engine/Plugins/Runtime/MassGameplay/Source/MassCommon/Private/MassCommonTypes.cpp
@@ -87,7 +87,7 @@ bool GetDebugEntitiesRange(int32& OutBegin, int32& OutEnd)
 {
 	OutBegin = DebugEntityBegin;
 	OutEnd = DebugEntityEnd;
-	return DebugEntityBegin != INDEX_NONE && DebugEntityEnd != INDEX_NONE;  
+	return DebugEntityBegin != INDEX_NONE && DebugEntityEnd != INDEX_NONE && DebugEntityBegin <= DebugEntityEnd;
 }
 	
 bool IsDebuggingEntity(FMassEntityHandle Entity, FColor* OutEntityColor)
