Replication optimization using HGrid to partition replicated entities to each client can fetch entities from it to optimize the replication LOD calculation.

#rnx

diff --git a/Engine/Plugins/Runtime/MassGameplay/Source/MassLOD/Public/MassLODCalculator.h b/Engine/Plugins/Runtime/MassGameplay/Source/MassLOD/Public/MassLODCalculator.h
--- a/Engine/Plugins/Runtime/MassGameplay/Source/MassLOD/Public/MassLODCalculator.h
+++ b/Engine/Plugins/Runtime/MassGameplay/Source/MassLOD/Public/MassLODCalculator.h
@@ -162,6 +162,11 @@ public:
 	template <typename TLODFragment, typename TTransformFragment>
 	void DebugDisplayLOD(FMassExecutionContext& Context, TConstArrayView<TLODFragment> LODList, TConstArrayView<TTransformFragment> LocationList, UWorld* World);
 
+	/**
+	 * Return the maximum distance at which the LOD will be turn off
+	 */
+	 float GetMaxLODDistance() const { return MaxLODDistance; }
+
 protected:
 
 	struct FMassLODRuntimeData
diff --git a/Engine/Plugins/Runtime/MassGameplay/Source/MassLOD/Public/MassLODTypes.h b/Engine/Plugins/Runtime/MassGameplay/Source/MassLOD/Public/MassLODTypes.h
--- a/Engine/Plugins/Runtime/MassGameplay/Source/MassLOD/Public/MassLODTypes.h
+++ b/Engine/Plugins/Runtime/MassGameplay/Source/MassLOD/Public/MassLODTypes.h
@@ -17,11 +17,6 @@ namespace UE::MassLOD
 {
 #if UE_ALLOW_DEBUG_REPLICATION_DUPLICATE_VIEWERS_PER_CONTROLLER
 	constexpr int32 DebugNumberViewersPerController = 50;
-	constexpr int32 MaxNumOfViewers = 50;
-#else
-	// @TODO Setup it up to 3 for as an optimization, but need a better way later
-	//constexpr int32 MaxNumOfViewers = 16;
-	constexpr int32 MaxNumOfViewers = 3;
 #endif // UE_ALLOW_DEBUG_REPLICATION_DUPLICATE_VIEWERS_PER_CONTROLLER
 
 	constexpr int32 MaxBucketsPerLOD = 250;
diff --git a/Engine/Plugins/Runtime/MassGameplay/Source/MassReplication/Private/MassReplicationFragments.cpp b/Engine/Plugins/Runtime/MassGameplay/Source/MassReplication/Private/MassReplicationFragments.cpp
--- a/Engine/Plugins/Runtime/MassGameplay/Source/MassReplication/Private/MassReplicationFragments.cpp
+++ b/Engine/Plugins/Runtime/MassGameplay/Source/MassReplication/Private/MassReplicationFragments.cpp
@@ -86,7 +86,7 @@ FMassReplicationParameters::FMassReplicationParameters()
 //----------------------------------------------------------------------//
 FMassReplicationSharedFragment::FMassReplicationSharedFragment(UMassReplicationSubsystem& ReplicationSubsystem, const FMassReplicationParameters& Params)
 {
-	LODCalculator.Initialize(Params.LODDistance, Params.BufferHysteresisOnDistancePercentage / 100.0f, Params.LODMaxCount, Params.LODMaxCountPerViewer);
+	LODCalculator.Initialize(Params.LODDistance, Params.BufferHysteresisOnDistancePercentage / 100.0f, Params.LODMaxCountPerViewer); // Setting the LODMaxCoundPerViewer as the MAXCount as the processor will be called for each client and will recalculate it per client
 	BubbleInfoClassHandle = ReplicationSubsystem.GetBubbleInfoClassHandle(Params.BubbleInfoClass);
 
 	CachedReplicator = Params.ReplicatorClass.GetDefaultObject();
diff --git a/Engine/Plugins/Runtime/MassGameplay/Source/MassReplication/Private/MassReplicationGridProcessor.cpp b/Engine/Plugins/Runtime/MassGameplay/Source/MassReplication/Private/MassReplicationGridProcessor.cpp
new file mode 100644
--- /dev/null
+++ b/Engine/Plugins/Runtime/MassGameplay/Source/MassReplication/Private/MassReplicationGridProcessor.cpp
@@ -0,0 +1,163 @@
+// Copyright Epic Games, Inc. All Rights Reserved.
+
+#include "MassReplicationGridProcessor.h"
+#include "MassReplicationSubsystem.h"
+#include "MassReplicationTypes.h"
+#include "MassReplicationFragments.h"
+#include "MassCommonFragments.h"
+#include "Engine/World.h"
+
+//----------------------------------------------------------------------//
+//  UMassReplicationGridProcessor
+//----------------------------------------------------------------------//
+UMassReplicationGridProcessor::UMassReplicationGridProcessor()
+{
+#if !UE_ALLOW_DEBUG_REPLICATION_BUBBLES_STANDALONE
+	ExecutionFlags = int32(EProcessorExecutionFlags::Server);
+#else
+	ExecutionFlags = int32(EProcessorExecutionFlags::All);
+#endif // UE_ALLOW_DEBUG_REPLICATION_BUBBLES_STANDALONE
+
+	ProcessingPhase = EMassProcessingPhase::PostPhysics;
+}
+
+void UMassReplicationGridProcessor::ConfigureQueries()
+{
+	AddToGridEntityQuery.AddRequirement<FTransformFragment>(EMassFragmentAccess::ReadOnly);
+	AddToGridEntityQuery.AddRequirement<FAgentRadiusFragment>(EMassFragmentAccess::ReadOnly);
+	AddToGridEntityQuery.AddRequirement<FMassReplicationGridCellLocationFragment>(EMassFragmentAccess::ReadWrite);
+	UpdateGridEntityQuery = AddToGridEntityQuery;
+	RemoveFromGridEntityQuery = AddToGridEntityQuery;
+
+	AddToGridEntityQuery.AddTagRequirement<FMassOffLODTag>(EMassFragmentPresence::None);
+	AddToGridEntityQuery.AddTagRequirement<FMassInReplicationGridTag>(EMassFragmentPresence::None);
+
+	UpdateGridEntityQuery.AddTagRequirement<FMassOffLODTag>(EMassFragmentPresence::None);
+	UpdateGridEntityQuery.AddTagRequirement<FMassInReplicationGridTag>(EMassFragmentPresence::All);
+
+	RemoveFromGridEntityQuery.AddTagRequirement<FMassOffLODTag>(EMassFragmentPresence::All);
+	RemoveFromGridEntityQuery.AddTagRequirement<FMassInReplicationGridTag>(EMassFragmentPresence::All);
+}
+
+void UMassReplicationGridProcessor::Initialize(UObject& Owner)
+{
+	Super::Initialize(Owner);
+
+	ReplicationSubsystem = UWorld::GetSubsystem<UMassReplicationSubsystem>(Owner.GetWorld());
+}
+
+void UMassReplicationGridProcessor::Execute(UMassEntitySubsystem& EntitySubsystem, FMassExecutionContext& Context)
+{
+	check(ReplicationSubsystem);
+
+	// can't be ParallelFor due to ReplicationSubsystem->GetGridMutable().Move not being thread-safe
+	FReplicationHashGrid2D& ReplicationGrid = ReplicationSubsystem->GetGridMutable();
+	AddToGridEntityQuery.ForEachEntityChunk(EntitySubsystem, Context, [&ReplicationGrid, &EntitySubsystem](FMassExecutionContext& Context)
+	{
+		const int32 NumEntities = Context.GetNumEntities();
+
+		TConstArrayView<FTransformFragment> LocationList = Context.GetFragmentView<FTransformFragment>();
+		TConstArrayView<FAgentRadiusFragment> RadiusList = Context.GetFragmentView<FAgentRadiusFragment>();
+		TArrayView<FMassReplicationGridCellLocationFragment> ReplicationCellLocationList = Context.GetMutableFragmentView<FMassReplicationGridCellLocationFragment>();
+
+		for (int32 EntityIndex = 0; EntityIndex < NumEntities; ++EntityIndex)
+		{
+			// Add to the grid
+			const FVector NewPos = LocationList[EntityIndex].GetTransform().GetLocation();
+			const float Radius = RadiusList[EntityIndex].Radius;
+
+			const FMassEntityHandle EntityHandle = Context.GetEntity(EntityIndex);
+			const FBox NewBounds(NewPos - FVector(Radius, Radius, 0.f), NewPos + FVector(Radius, Radius, 0.f));
+			ReplicationCellLocationList[EntityIndex].CellLoc = ReplicationGrid.Add(EntityHandle, NewBounds);
+
+			Context.Defer().AddTag<FMassInReplicationGridTag>(EntityHandle);
+		}
+	});
+
+	UpdateGridEntityQuery.ForEachEntityChunk(EntitySubsystem, Context, [&ReplicationGrid, &EntitySubsystem](FMassExecutionContext& Context)
+	{
+		const int32 NumEntities = Context.GetNumEntities();
+
+		TConstArrayView<FTransformFragment> LocationList = Context.GetFragmentView<FTransformFragment>();
+		TConstArrayView<FAgentRadiusFragment> RadiusList = Context.GetFragmentView<FAgentRadiusFragment>();
+		TArrayView<FMassReplicationGridCellLocationFragment> ReplicationCellLocationList = Context.GetMutableFragmentView<FMassReplicationGridCellLocationFragment>();
+
+		for (int32 EntityIndex = 0; EntityIndex < NumEntities; ++EntityIndex)
+		{
+			// Update position in grid
+			const FVector NewPos = LocationList[EntityIndex].GetTransform().GetLocation();
+			const float Radius = RadiusList[EntityIndex].Radius;
+			const FMassEntityHandle EntityHandle = Context.GetEntity(EntityIndex);
+			const FBox NewBounds(NewPos - FVector(Radius, Radius, 0.f), NewPos + FVector(Radius, Radius, 0.f));
+			ReplicationCellLocationList[EntityIndex].CellLoc = ReplicationGrid.Move(EntityHandle, ReplicationCellLocationList[EntityIndex].CellLoc, NewBounds);
+
+#if WITH_MASSGAMEPLAY_DEBUG && 0
+			const FDebugContext BaseDebugContext(this, LogMassReplication, nullptr, EntityHandle);
+			if (DebugIsSelected(EntityHandle))
+			{
+				FBox Box = ReplicationGrid.CalcCellBounds(ReplicationCellLocationList[EntityIndex].CellLoc);
+				Box.Max.Z += 200.f;
+				DebugDrawBox(BaseDebugContext, Box, FColor::Yellow);
+			}
+#endif // WITH_MASSGAMEPLAY_DEBUG
+		}
+	});
+
+	RemoveFromGridEntityQuery.ForEachEntityChunk(EntitySubsystem, Context, [&ReplicationGrid](FMassExecutionContext& Context)
+	{
+		const int32 NumEntities = Context.GetNumEntities();
+
+		TArrayView<FMassReplicationGridCellLocationFragment> ReplicationCellLocationList = Context.GetMutableFragmentView<FMassReplicationGridCellLocationFragment>();
+		for (int32 EntityIndex = 0; EntityIndex < NumEntities; ++EntityIndex)
+		{
+			const FMassEntityHandle EntityHandle = Context.GetEntity(EntityIndex);
+			ReplicationGrid.Remove(EntityHandle, ReplicationCellLocationList[EntityIndex].CellLoc);
+			ReplicationCellLocationList[EntityIndex].CellLoc = FReplicationHashGrid2D::FCellLocation();
+
+			Context.Defer().RemoveTag<FMassInReplicationGridTag>(EntityHandle);
+		}
+	});
+}
+
+//----------------------------------------------------------------------//
+//  UMassReplicationGridRemoverProcessor
+//----------------------------------------------------------------------//
+UMassReplicationGridRemoverProcessor::UMassReplicationGridRemoverProcessor()
+{
+	ObservedType = FMassReplicationGridCellLocationFragment::StaticStruct();
+	Operation = EMassObservedOperation::Remove;
+	ExecutionFlags = (int32)(EProcessorExecutionFlags::All);
+}
+
+void UMassReplicationGridRemoverProcessor::ConfigureQueries()
+{
+	EntityQuery.AddRequirement<FMassReplicationGridCellLocationFragment>(EMassFragmentAccess::ReadWrite);
+}
+
+void UMassReplicationGridRemoverProcessor::Initialize(UObject& Owner)
+{
+	Super::Initialize(Owner);
+	ReplicationSubsystem = UWorld::GetSubsystem<UMassReplicationSubsystem>(Owner.GetWorld());
+}
+
+void UMassReplicationGridRemoverProcessor::Execute(UMassEntitySubsystem& EntitySubsystem, FMassExecutionContext& Context)
+{
+	if (!ReplicationSubsystem)
+	{
+		return;
+	}
+
+	FReplicationHashGrid2D& ReplicationGrid = ReplicationSubsystem->GetGridMutable();
+	EntityQuery.ForEachEntityChunk(EntitySubsystem, Context, [&ReplicationGrid](FMassExecutionContext& Context)
+	{
+		const int32 NumEntities = Context.GetNumEntities();
+		const TArrayView<FMassReplicationGridCellLocationFragment> ReplicationCellLocationList = Context.GetMutableFragmentView<FMassReplicationGridCellLocationFragment>();
+
+		for (int32 EntityIndex = 0; EntityIndex < NumEntities; ++EntityIndex)
+		{
+			const FMassEntityHandle EntityHandle = Context.GetEntity(EntityIndex);
+			ReplicationGrid.Remove(EntityHandle, ReplicationCellLocationList[EntityIndex].CellLoc);
+			ReplicationCellLocationList[EntityIndex].CellLoc = FReplicationHashGrid2D::FCellLocation();
+		}
+	});
+}
\ No newline at end of file
diff --git a/Engine/Plugins/Runtime/MassGameplay/Source/MassReplication/Private/MassReplicationProcessor.cpp b/Engine/Plugins/Runtime/MassGameplay/Source/MassReplication/Private/MassReplicationProcessor.cpp
--- a/Engine/Plugins/Runtime/MassGameplay/Source/MassReplication/Private/MassReplicationProcessor.cpp
+++ b/Engine/Plugins/Runtime/MassGameplay/Source/MassReplication/Private/MassReplicationProcessor.cpp
@@ -7,8 +7,8 @@
 
 namespace UE::Mass::Replication
 {
-	int32 bDebugReplicationLOD = 0;
-	FAutoConsoleVariableRef CVarDebugReplicationViewerLOD(TEXT("ai.debug.ReplicationLOD"), bDebugReplicationLOD, TEXT("Debug Replication LOD"), ECVF_Cheat);
+	int32 DebugClientReplicationLOD = -1;
+	FAutoConsoleVariableRef CVarDebugReplicationViewerLOD(TEXT("ai.debug.ClientReplicationLOD"), DebugClientReplicationLOD, TEXT("Debug Replication LOD of the specified client index"), ECVF_Cheat);
 } // UE::Mass::Crowd
 
 //----------------------------------------------------------------------//
@@ -27,6 +27,9 @@ UMassReplicationProcessor::UMassReplicationProcessor()
 
 void UMassReplicationProcessor::ConfigureQueries()
 {
+	SyncClientData.AddRequirement<FMassReplicationLODFragment>(EMassFragmentAccess::ReadWrite);
+	SyncClientData.AddRequirement<FMassReplicatedAgentFragment>(EMassFragmentAccess::ReadWrite);
+
 	CollectViewerInfoQuery.AddRequirement<FTransformFragment>(EMassFragmentAccess::ReadOnly);
 	CollectViewerInfoQuery.AddRequirement<FMassReplicationViewerInfoFragment>(EMassFragmentAccess::ReadOnly);
 	CollectViewerInfoQuery.AddSharedRequirement<FMassReplicationSharedFragment>(EMassFragmentAccess::ReadWrite);
@@ -74,8 +77,7 @@ void UMassReplicationProcessor::PrepareExecution(UMassEntitySubsystem& EntitySub
 	ReplicationSubsystem->SynchronizeClientsAndViewers();
 
 	check(LODSubsystem);
-	const TArray<FViewerInfo>& Viewers = LODSubsystem->GetViewers();
-	EntitySubsystem.ForEachSharedFragment<FMassReplicationSharedFragment>([this, &Viewers](FMassReplicationSharedFragment& RepSharedFragment)
+	EntitySubsystem.ForEachSharedFragment<FMassReplicationSharedFragment>([this](FMassReplicationSharedFragment& RepSharedFragment)
 	{
 		if (!RepSharedFragment.bEntityQueryInitialized)
 		{
@@ -89,17 +91,13 @@ void UMassReplicationProcessor::PrepareExecution(UMassEntitySubsystem& EntitySub
 			RepSharedFragment.bEntityQueryInitialized = true;
 		}
 
-		RepSharedFragment.LODCollector.PrepareExecution(Viewers);
-		RepSharedFragment.LODCalculator.PrepareExecution(Viewers);
-
 		const TArray<FMassClientHandle>& CurrentClientHandles = ReplicationSubsystem->GetClientReplicationHandles();
-		const int32 MinNumHandles = FMath::Min(RepSharedFragment.CachedClientHandles.Num(), CurrentClientHandles.Num());
+		const int32 MinNumHandles = FMath::Min(RepSharedFragment.CachedClientHandles.Num(), CurrentClientHandles.Num()); // Why is this the min not the max?
 
 		//check to see if we don't have enough cached client handles
 		if (RepSharedFragment.CachedClientHandles.Num() < CurrentClientHandles.Num())
 		{
 			RepSharedFragment.CachedClientHandles.Reserve(CurrentClientHandles.Num());
-			RepSharedFragment.bBubbleChanged.Reserve(CurrentClientHandles.Num());
 			RepSharedFragment.BubbleInfos.Reserve(CurrentClientHandles.Num());
 
 			for (int32 Idx = RepSharedFragment.CachedClientHandles.Num(); Idx < CurrentClientHandles.Num(); ++Idx)
@@ -107,7 +105,6 @@ void UMassReplicationProcessor::PrepareExecution(UMassEntitySubsystem& EntitySub
 				const FMassClientHandle& CurrentClientHandle = CurrentClientHandles[Idx];
 
 				RepSharedFragment.CachedClientHandles.Add(CurrentClientHandle);
-				RepSharedFragment.bBubbleChanged.Add(true);
 				AMassClientBubbleInfoBase* Info = CurrentClientHandle.IsValid() ?
 					ReplicationSubsystem->GetClientBubbleChecked(RepSharedFragment.BubbleInfoClassHandle, CurrentClientHandle) :
 					nullptr;
@@ -123,7 +120,6 @@ void UMassReplicationProcessor::PrepareExecution(UMassEntitySubsystem& EntitySub
 			const int32 NumRemove = RepSharedFragment.CachedClientHandles.Num() - CurrentClientHandles.Num();
 
 			RepSharedFragment.CachedClientHandles.RemoveAt(CurrentClientHandles.Num(), NumRemove, /* bAllowShrinking */ false);
-			RepSharedFragment.bBubbleChanged.RemoveAt(CurrentClientHandles.Num(), NumRemove, /* bAllowShrinking */ false);
 			RepSharedFragment.BubbleInfos.RemoveAt(CurrentClientHandles.Num(), NumRemove, /* bAllowShrinking */ false);
 		}
 
@@ -133,7 +129,7 @@ void UMassReplicationProcessor::PrepareExecution(UMassEntitySubsystem& EntitySub
 			const FMassClientHandle& CurrentClientHandle = CurrentClientHandles[Idx];
 			FMassClientHandle& CachedClientHandle = RepSharedFragment.CachedClientHandles[Idx];
 
-			const bool bChanged = (RepSharedFragment.bBubbleChanged[Idx] = (CurrentClientHandle != CachedClientHandle));
+			const bool bChanged = (CurrentClientHandle != CachedClientHandle);
 			if (bChanged)
 			{
 				AMassClientBubbleInfoBase* Info = CurrentClientHandle.IsValid() ?
@@ -163,60 +159,232 @@ void UMassReplicationProcessor::Execute(UMassEntitySubsystem& EntitySubsystem, F
 		PrepareExecution(EntitySubsystem);
 	}
 
+
+	const TArray<FViewerInfo>& AllViewersInfo = LODSubsystem->GetViewers();
+	const TArray<FMassClientHandle>& ClientHandles = ReplicationSubsystem->GetClientReplicationHandles();
+	for( const FMassClientHandle ClientHandle : ClientHandles )
 	{
-		QUICK_SCOPE_CYCLE_COUNTER(UMassReplicationProcessor_LODCaculation);
-		CollectViewerInfoQuery.ForEachEntityChunk(EntitySubsystem, Context, [this](FMassExecutionContext& Context)
-		{
-			const TConstArrayView<FTransformFragment> LocationList = Context.GetFragmentView<FTransformFragment>();
-			const TArrayView<FMassReplicationViewerInfoFragment> ViewersInfoList = Context.GetMutableFragmentView<FMassReplicationViewerInfoFragment>();
-			FMassReplicationSharedFragment& RepSharedFragment = Context.GetMutableSharedFragment<FMassReplicationSharedFragment>();
-			RepSharedFragment.LODCollector.CollectLODInfo(Context, LocationList, ViewersInfoList, ViewersInfoList);
-		});
+		FMassClientReplicationInfo& ClientReplicationInfo = ReplicationSubsystem->GetMutableClientReplicationInfoChecked(ClientHandle);
 
-		CalculateLODQuery.ForEachEntityChunk(EntitySubsystem, Context, [this](FMassExecutionContext& Context)
+		// Figure out all viewer of this client
+		TArray<FViewerInfo> Viewers;
+		for(const FMassViewerHandle ClientViewerHandle : ClientReplicationInfo.Handles)
 		{
-			const TConstArrayView<FMassReplicationViewerInfoFragment> ViewersInfoList = Context.GetFragmentView<FMassReplicationViewerInfoFragment>();
-			const TArrayView<FMassReplicationLODFragment> ViewerLODList = Context.GetMutableFragmentView<FMassReplicationLODFragment>();
-			FMassReplicationSharedFragment& RepSharedFragment = Context.GetMutableSharedFragment<FMassReplicationSharedFragment>();
-			RepSharedFragment.LODCalculator.CalculateLOD(Context, ViewersInfoList, ViewerLODList, ViewersInfoList);
-		});
+			const FViewerInfo* ViewerInfo = AllViewersInfo.FindByPredicate([ClientViewerHandle](const FViewerInfo& ViewerInfo) { return ClientViewerHandle == ViewerInfo.Handle; });
+			if (ensureMsgf(ViewerInfo, TEXT("Expecting to find the client viewer handle in the all viewers info list")))
+			{
+				Viewers.Add(*ViewerInfo);
+			}
+		}
 
-		EntitySubsystem.ForEachSharedFragment<FMassReplicationSharedFragment>([](FMassReplicationSharedFragment& RepSharedFragment)
+		// Prepare LOD collector and calculator
+		// Remember the max LOD distance from each
+		float MaxLODDistance = 0.0f;
+		EntitySubsystem.ForEachSharedFragment<FMassReplicationSharedFragment>([&Viewers,&MaxLODDistance](FMassReplicationSharedFragment& RepSharedFragment)
 		{
-			RepSharedFragment.bHasAdjustedDistancesFromCount = RepSharedFragment.LODCalculator.AdjustDistancesFromCount();
+			RepSharedFragment.LODCollector.PrepareExecution(Viewers);
+			RepSharedFragment.LODCalculator.PrepareExecution(Viewers);
+			MaxLODDistance = FMath::Max(MaxLODDistance, RepSharedFragment.LODCalculator.GetMaxLODDistance());
 		});
 
-		AdjustLODDistancesQuery.ForEachEntityChunk(EntitySubsystem, Context, [this](FMassExecutionContext& Context)
+		// Fetch all entities to process
+		const FVector HalfExtent(MaxLODDistance, MaxLODDistance, 0.0f);
+		TArray<FMassEntityHandle> EntitiesInRange;
+		for (const FViewerInfo& Viewer : Viewers)
 		{
-			const TConstArrayView<FMassReplicationViewerInfoFragment> ViewersInfoList = Context.GetFragmentView<FMassReplicationViewerInfoFragment>();
-			const TArrayView<FMassReplicationLODFragment> ViewerLODList = Context.GetMutableFragmentView<FMassReplicationLODFragment>();
-			FMassReplicationSharedFragment& RepSharedFragment = Context.GetMutableSharedFragment<FMassReplicationSharedFragment>();
-			RepSharedFragment.LODCalculator.AdjustLODFromCount(Context, ViewersInfoList, ViewerLODList, ViewersInfoList);
-		});
-	}
+			FBox Bounds(Viewer.Location - HalfExtent, Viewer.Location + HalfExtent);
+			ReplicationSubsystem->GetGrid().Query(Bounds, EntitiesInRange);
+		}
 
-	{
-		QUICK_SCOPE_CYCLE_COUNTER(UMassReplicationProcessor_ProcessClientReplication);
-		FMassReplicationContext ReplicationContext(*World, *LODSubsystem, *ReplicationSubsystem);
-		EntitySubsystem.ForEachSharedFragment<FMassReplicationSharedFragment>([this, &EntitySubsystem, &Context, &ReplicationContext](FMassReplicationSharedFragment& RepSharedFragment)
+		EntityQuery.CacheArchetypes(EntitySubsystem);
+		if (EntityQuery.GetArchetypes().Num() > 0)
 		{
-			RepSharedFragment.EntityQuery.ForEachEntityChunk(EntitySubsystem, Context, [&ReplicationContext, &RepSharedFragment](FMassExecutionContext& Context)
+			// EntitySet stores array of entities per specified archetype, may contain duplicates.
+			struct FEntitySet
 			{
-				RepSharedFragment.CachedReplicator->ProcessClientReplication(Context, ReplicationContext);
-			});
-		});
-	}
+				void Reset()
+				{
+					Entities.Reset();
+				}
 
-	// Optional debug display
-	if (UE::Mass::Replication::bDebugReplicationLOD)
-	{
-		EntityQuery.ForEachEntityChunk(EntitySubsystem, Context, [this](FMassExecutionContext& Context)
+				FMassArchetypeHandle Archetype;
+				TArray<FMassEntityHandle> Entities;
+			};
+			TArray<FEntitySet> EntitySets;
+
+			for (const FMassArchetypeHandle& Archetype : EntityQuery.GetArchetypes())
+			{
+				FEntitySet& Set = EntitySets.AddDefaulted_GetRef();
+				Set.Archetype = Archetype;
+			}
+
+			auto BuildEntitySet = [&EntitySets, &EntitySubsystem](const TArray<FMassEntityHandle>& Entities)
+			{
+				FEntitySet* PrevSet = Entities.Num() ? &EntitySets[0] : nullptr;
+				for (const FMassEntityHandle Entity : Entities)
+				{
+					// Add to set of supported archetypes. Dont process if we don't care about the type.
+					const FMassArchetypeHandle Archetype = EntitySubsystem.GetArchetypeForEntity(Entity);
+					FEntitySet* Set = PrevSet && PrevSet->Archetype == Archetype ? PrevSet : EntitySets.FindByPredicate([&Archetype](const FEntitySet& Set) { return Archetype == Set.Archetype; });
+					if (Set != nullptr)
+					{
+						// We don't care about duplicates here, the FMassArchetypeEntityCollection creation below will handle it
+						Set->Entities.Add(Entity);
+						PrevSet = Set;
+					}
+				}
+			};
+
+			BuildEntitySet(ClientReplicationInfo.HandledEntities);
+			BuildEntitySet(EntitiesInRange);
+
+			for (FEntitySet& Set : EntitySets)
+			{
+				if (Set.Entities.Num() == 0)
+				{
+					continue;
+				}
+
+				Context.SetEntityCollection(FMassArchetypeEntityCollection(Set.Archetype, Set.Entities, FMassArchetypeEntityCollection::FoldDuplicates));
+
+				{
+					QUICK_SCOPE_CYCLE_COUNTER(UMassReplicationProcessor_SyncToMass);
+					SyncClientData.ForEachEntityChunk(EntitySubsystem, Context, [&ClientReplicationInfo](FMassExecutionContext& Context)
+					{
+						const TArrayView<FMassReplicationLODFragment> ViewerLODList = Context.GetMutableFragmentView<FMassReplicationLODFragment>();
+						TArrayView<FMassReplicatedAgentFragment> ReplicatedAgentList = Context.GetMutableFragmentView<FMassReplicatedAgentFragment>();
+
+						const int32 NumEntities = Context.GetNumEntities();
+						for (int EntityIdx = 0; EntityIdx < NumEntities; EntityIdx++)
+						{
+							FMassEntityHandle EntityHandle = Context.GetEntity(EntityIdx);
+							FMassReplicatedAgentFragment& AgentFragment = ReplicatedAgentList[EntityIdx];
+							FMassReplicationLODFragment& LODFragment = ViewerLODList[EntityIdx];
+
+							if (FMassReplicatedAgentData* AgentData = ClientReplicationInfo.AgentsData.Find(EntityHandle))
+							{
+								LODFragment.LOD = AgentData->LOD;
+								AgentFragment.AgentData = *AgentData;
+							}
+							else
+							{
+								LODFragment.LOD = EMassLOD::Off;
+								AgentFragment.AgentData.Invalidate();
+							}
+						}
+					});
+				}
+
+				{
+					QUICK_SCOPE_CYCLE_COUNTER(UMassReplicationProcessor_LODCollection);
+					CollectViewerInfoQuery.ForEachEntityChunk(EntitySubsystem, Context, [](FMassExecutionContext& Context)
+					{
+						const TConstArrayView<FTransformFragment> LocationList = Context.GetFragmentView<FTransformFragment>();
+						const TArrayView<FMassReplicationViewerInfoFragment> ViewersInfoList = Context.GetMutableFragmentView<FMassReplicationViewerInfoFragment>();
+						FMassReplicationSharedFragment& RepSharedFragment = Context.GetMutableSharedFragment<FMassReplicationSharedFragment>();
+						RepSharedFragment.LODCollector.CollectLODInfo(Context, LocationList, ViewersInfoList, ViewersInfoList);
+					});
+				}
+
+				{
+					QUICK_SCOPE_CYCLE_COUNTER(UMassReplicationProcessor_LODCaculation);
+					CalculateLODQuery.ForEachEntityChunk(EntitySubsystem, Context, [](FMassExecutionContext& Context)
+					{
+						const TConstArrayView<FMassReplicationViewerInfoFragment> ViewersInfoList = Context.GetFragmentView<FMassReplicationViewerInfoFragment>();
+						const TArrayView<FMassReplicationLODFragment> ViewerLODList = Context.GetMutableFragmentView<FMassReplicationLODFragment>();
+						FMassReplicationSharedFragment& RepSharedFragment = Context.GetMutableSharedFragment<FMassReplicationSharedFragment>();
+						RepSharedFragment.LODCalculator.CalculateLOD(Context, ViewersInfoList, ViewerLODList, ViewersInfoList);
+					});
+				}
+				Context.ClearEntityCollection();
+			}
+
+			{
+				QUICK_SCOPE_CYCLE_COUNTER(UMassReplicationProcessor_LODAdjustDistance);
+				EntitySubsystem.ForEachSharedFragment<FMassReplicationSharedFragment>([](FMassReplicationSharedFragment& RepSharedFragment)
+				{
+					RepSharedFragment.bHasAdjustedDistancesFromCount = RepSharedFragment.LODCalculator.AdjustDistancesFromCount();
+				});
+			}
+
+			for (FEntitySet& Set : EntitySets)
+			{
+				if (Set.Entities.Num() == 0)
+				{
+					continue;
+				}
+				Context.SetEntityCollection(FMassArchetypeEntityCollection(Set.Archetype, Set.Entities, FMassArchetypeEntityCollection::FoldDuplicates));
+
+				{
+					QUICK_SCOPE_CYCLE_COUNTER(UMassReplicationProcessor_LODAdjustLODFromCount);
+					AdjustLODDistancesQuery.ForEachEntityChunk(EntitySubsystem, Context, [](FMassExecutionContext& Context)
+					{
+						const TConstArrayView<FMassReplicationViewerInfoFragment> ViewersInfoList = Context.GetFragmentView<FMassReplicationViewerInfoFragment>();
+						const TArrayView<FMassReplicationLODFragment> ViewerLODList = Context.GetMutableFragmentView<FMassReplicationLODFragment>();
+						FMassReplicationSharedFragment& RepSharedFragment = Context.GetMutableSharedFragment<FMassReplicationSharedFragment>();
+						RepSharedFragment.LODCalculator.AdjustLODFromCount(Context, ViewersInfoList, ViewerLODList, ViewersInfoList);
+					});
+				}
+
+				{
+					QUICK_SCOPE_CYCLE_COUNTER(UMassReplicationProcessor_ProcessClientReplication);
+					FMassReplicationContext ReplicationContext(*World, *LODSubsystem, *ReplicationSubsystem);
+					EntitySubsystem.ForEachSharedFragment<FMassReplicationSharedFragment>([&EntitySubsystem, &Context, &ReplicationContext, &ClientHandle](FMassReplicationSharedFragment& RepSharedFragment)
+					{
+						RepSharedFragment.CurrentClientHandle = ClientHandle;
+
+						RepSharedFragment.EntityQuery.ForEachEntityChunk(EntitySubsystem, Context, [&ReplicationContext, &RepSharedFragment](FMassExecutionContext& Context)
+						{
+							RepSharedFragment.CachedReplicator->ProcessClientReplication(Context, ReplicationContext);
+						});
+					});
+				}
+
+				{
+					QUICK_SCOPE_CYCLE_COUNTER(UMassReplicationProcessor_SyncFromMass);
+					SyncClientData.ForEachEntityChunk(EntitySubsystem, Context, [&ClientReplicationInfo](FMassExecutionContext& Context)
+					{
+						TArrayView<FMassReplicatedAgentFragment> ReplicatedAgentList = Context.GetMutableFragmentView<FMassReplicatedAgentFragment>();
+
+						const int32 NumEntities = Context.GetNumEntities();
+						for (int EntityIdx = 0; EntityIdx < NumEntities; EntityIdx++)
+						{
+							FMassEntityHandle EntityHandle = Context.GetEntity(EntityIdx);
+							FMassReplicatedAgentFragment& AgentFragment = ReplicatedAgentList[EntityIdx];
+							ClientReplicationInfo.AgentsData.Add(EntityHandle, AgentFragment.AgentData);
+						}
+					});
+				}
+
+				// Optional debug display
+				if (UE::Mass::Replication::DebugClientReplicationLOD == ClientHandle.GetIndex())
+				{
+					EntitySubsystem.ForEachSharedFragment<FMassReplicationSharedFragment>([this, &EntitySubsystem, &Context](FMassReplicationSharedFragment& RepSharedFragment)
+					{
+						RepSharedFragment.EntityQuery.ForEachEntityChunk(EntitySubsystem, Context, [this, &RepSharedFragment](FMassExecutionContext& Context)
+						{
+							const TConstArrayView<FTransformFragment> TransformList = Context.GetFragmentView<FTransformFragment>();
+							const TConstArrayView<FMassReplicationLODFragment> ViewerLODList = Context.GetFragmentView<FMassReplicationLODFragment>();
+							RepSharedFragment.LODCalculator.DebugDisplayLOD(Context, ViewerLODList, TransformList, World);
+						});
+					});
+				}
+
+				Context.ClearEntityCollection();
+			}
+		}
+		ClientReplicationInfo.HandledEntities = MoveTemp(EntitiesInRange);
+
+		// Cleanup any AgentData that isn't relevant anymore (that is EMassLOD::OFF)
+		for (FMassReplicationAgentDataMap::TIterator It = ClientReplicationInfo.AgentsData.CreateIterator(); It; ++It)
 		{
-			const TConstArrayView<FTransformFragment> TransformList = Context.GetFragmentView<FTransformFragment>();
-			const TConstArrayView<FMassReplicationLODFragment> ViewerLODList = Context.GetFragmentView<FMassReplicationLODFragment>();
-			FMassReplicationSharedFragment& RepSharedFragment = Context.GetMutableSharedFragment<FMassReplicationSharedFragment>();
-			RepSharedFragment.LODCalculator.DebugDisplayLOD(Context, ViewerLODList, TransformList, World);
-		});
+			FMassReplicatedAgentData& AgentData = It.Value();
+			if (AgentData.LOD == EMassLOD::Off)
+			{
+				checkf(!AgentData.Handle.IsValid(), TEXT("This replicated agent should have been removed from this client and was not"));
+				It.RemoveCurrent();
+			}
+		}
 	}
 #endif //UE_REPLICATION_COMPILE_SERVER_CODE
-}
\ No newline at end of file
+}
diff --git a/Engine/Plugins/Runtime/MassGameplay/Source/MassReplication/Private/MassReplicationSubsystem.cpp b/Engine/Plugins/Runtime/MassGameplay/Source/MassReplication/Private/MassReplicationSubsystem.cpp
--- a/Engine/Plugins/Runtime/MassGameplay/Source/MassReplication/Private/MassReplicationSubsystem.cpp
+++ b/Engine/Plugins/Runtime/MassGameplay/Source/MassReplication/Private/MassReplicationSubsystem.cpp
@@ -8,6 +8,7 @@
 #include "MassEntitySubsystem.h"
 #include "MassClientBubbleHandler.h"
 #include "MassClientBubbleInfoBase.h"
+#include "MassReplicationSettings.h"
 
 
 uint32 UMassReplicationSubsystem::CurrentNetMassCounter = 0;
@@ -41,7 +42,7 @@ void UMassReplicationSubsystem::Deinitialize()
 	// make sure all the other data is reset
 	ClientHandleManager.Reset();
 	BubbleInfoArray.Reset();
-	ClientsViewerHandles.Reset();
+	ClientsReplicationInfo.Reset();
 	ClientToViewerHandleArray.Reset();
 	ViewerToClientHandleArray.Reset();
 
@@ -50,6 +51,11 @@ void UMassReplicationSubsystem::Deinitialize()
 	EntitySystem = nullptr;
 }
 
+UMassReplicationSubsystem::UMassReplicationSubsystem()
+ : ReplicationGrid(GetDefault<UMassReplicationSettings>()->GetReplicationGridCellSize())
+{
+}
+
 FMassNetworkID UMassReplicationSubsystem::GetNetIDFromHandle(const FMassEntityHandle Handle) const
 {
 	check(EntitySystem);
@@ -249,7 +255,7 @@ void UMassReplicationSubsystem::SynchronizeClientViewers(const TArray<FViewerInf
 		FMassViewerHandle ViewerHandle;
 	};
 
-	// go through the ClientsViewerHandles and check validity and store the valid ones into a map
+	// go through the ClientReplicationInfo and check validity and store the valid ones into a map
 	typedef TMap<APlayerController*, FMassClientViewerHandle> FViewerMap;
 
 	FViewerMap ClientViewerMap;
@@ -261,13 +267,13 @@ void UMassReplicationSubsystem::SynchronizeClientViewers(const TArray<FViewerInf
 			//as this is a fresh handle then we only need to check against !IsInvalid
 			if (ClientHandle.IsValid())
 			{
-				FClientViewerHandles& ClientViewers = ClientsViewerHandles[ClientHandle.GetIndex()];
+				FMassClientReplicationInfo& ClientReplicationInfo = ClientsReplicationInfo[ClientHandle.GetIndex()];
 
 				int32 ViewerIdx = 0;
 
-				while (ViewerIdx < ClientViewers.Handles.Num())
+				while (ViewerIdx < ClientReplicationInfo.Handles.Num())
 				{
-					const FMassViewerHandle& ClientViewer = ClientViewers.Handles[ViewerIdx];
+					const FMassViewerHandle& ClientViewer = ClientReplicationInfo.Handles[ViewerIdx];
 
 					APlayerController* Controller = MassLODSubsystem->GetPlayerControllerFromViewerHandle(ClientViewer);
 
@@ -285,7 +291,7 @@ void UMassReplicationSubsystem::SynchronizeClientViewers(const TArray<FViewerInf
 					}
 					else //remove invalid ClientViewer, but dont increment the ViewerIdx
 					{
-						ClientViewers.Handles.RemoveAt(ViewerIdx, 1, /* bAllowShrinking */ false);
+						ClientReplicationInfo.Handles.RemoveAt(ViewerIdx, 1, /* bAllowShrinking */ false);
 					}
 				}
 			}
@@ -315,12 +321,12 @@ void UMassReplicationSubsystem::SynchronizeClientViewers(const TArray<FViewerInf
 						check(MassLODSubsystem->IsValidViewer(ParentViewerClientPair.ViewerHandle));
 						check(ClientHandleManager.IsValidHandle(ParentViewerClientPair.ClientHandle));
 
-						// remove APlayerController from the ClientViewerMap and Add the viewer to the ClientsViewerHandles
+						// remove APlayerController from the ClientViewerMap and Add the viewer to the ClientsReplicationInfo
 						ClientViewerMap.Remove(Viewer.PlayerController);
 
-						FClientViewerHandles& ClientViewers = ClientsViewerHandles[ParentViewerClientPair.ClientHandle.GetIndex()];
+						FMassClientReplicationInfo& ClientReplicationInfo = ClientsReplicationInfo[ParentViewerClientPair.ClientHandle.GetIndex()];
 
-						ClientViewers.Handles.Add(Viewer.Handle);
+						ClientReplicationInfo.Handles.Add(Viewer.Handle);
 					}
 				}
 			}
@@ -333,9 +339,9 @@ void UMassReplicationSubsystem::SynchronizeClientViewers(const TArray<FViewerInf
 		{
 			const FMassClientViewerHandle& HandleData = Itr->Value;
 
-			FClientViewerHandles& ClientViewers = ClientsViewerHandles[HandleData.ClientHandle.GetIndex()];
+			FMassClientReplicationInfo& ClientReplicationInfo = ClientsReplicationInfo[HandleData.ClientHandle.GetIndex()];
 
-			ClientViewers.Handles.RemoveSingle(HandleData.ViewerHandle);
+			ClientReplicationInfo.Handles.RemoveSingle(HandleData.ViewerHandle);
 		}
 	}
 }
@@ -369,7 +375,7 @@ void UMassReplicationSubsystem::SynchronizeClientsAndViewers()
 	{
 		const int32 NumItems = ClientHandleManager.ShrinkHandles();
 		 
-		ClientsViewerHandles.RemoveAt(NumItems, ClientsViewerHandles.Num() - NumItems, /* bAllowShrinking */ false);
+		ClientsReplicationInfo.RemoveAt(NumItems, ClientsReplicationInfo.Num() - NumItems, /* bAllowShrinking */ false);
 		ClientToViewerHandleArray.RemoveAt(NumItems, ClientToViewerHandleArray.Num() - NumItems, /* bAllowShrinking */ false);
 
 		for (FMassClientBubbleInfoData& InfoData : BubbleInfoArray)
@@ -527,7 +533,7 @@ void UMassReplicationSubsystem::DebugCheckArraysAreInSync()
 {
 #if UE_DEBUG_REPLICATION
 
-	checkf((ClientToViewerHandleArray.Num() == ClientsViewerHandles.Num()), TEXT("Client arrays out of sync with each other!"));
+	checkf((ClientToViewerHandleArray.Num() == ClientsReplicationInfo.Num()), TEXT("Client arrays out of sync with each other!"));
 
 	const int32 NumEntries = (BubbleInfoArray.Num() > 0) ? BubbleInfoArray[0].Bubbles.Num() : 0;
 
@@ -535,7 +541,7 @@ void UMassReplicationSubsystem::DebugCheckArraysAreInSync()
 	{
 		FMassClientBubbleInfoData& InfoDataOuter = BubbleInfoArray[IdxOuter];
 
-		checkf((InfoDataOuter.Bubbles.Num() == ClientsViewerHandles.Num()), TEXT("BubbleInfoArray arrays out of sync with ClientsViewerHandles!"));
+		checkf((InfoDataOuter.Bubbles.Num() == ClientsReplicationInfo.Num()), TEXT("BubbleInfoArray arrays out of sync with ClientsReplicationInfo!"));
 		checkf(InfoDataOuter.Bubbles.Num() == NumEntries, TEXT("Bubbles have different numbers of items!"));
 
 		for (int32 IdxInner = IdxOuter + 1; IdxInner < BubbleInfoArray.Num(); ++IdxInner)
@@ -613,12 +619,12 @@ void UMassReplicationSubsystem::AddClient(FMassViewerHandle ViewerHandle, APlaye
 	// check if the handle is a new entry in the free list arrays or uses an existing entry
 	if (ClientHandle.GetIndex() == ClientToViewerHandleArray.Num())
 	{
-		ClientsViewerHandles.AddDefaulted();
+		ClientsReplicationInfo.AddDefaulted();
 		ClientToViewerHandleArray.Emplace(ViewerHandle, ClientHandle);
 	}
 	else
 	{
-		checkf(ClientsViewerHandles[ClientHandle.GetIndex()].Handles.Num() == 0, TEXT("ClientsViewerHandles being replaced must have zero entries, they should have been removed first!"));
+		checkf(ClientsReplicationInfo[ClientHandle.GetIndex()].Handles.Num() == 0, TEXT("ClientsReplicationInfo being replaced must have zero entries, they should have been removed first!"));
 
 		FViewerClientPair& ClientToViewerHandleItem = ClientToViewerHandleArray[ClientHandle.GetIndex()];
 
@@ -627,8 +633,8 @@ void UMassReplicationSubsystem::AddClient(FMassViewerHandle ViewerHandle, APlaye
 		ClientToViewerHandleItem = ViewerClientPair;
 	}
 
-	FClientViewerHandles& ClientViewers = ClientsViewerHandles[ClientHandle.GetIndex()];
-	ClientViewers.Handles.Add(ViewerHandle);
+	FMassClientReplicationInfo& ClientReplicationInfo = ClientsReplicationInfo[ClientHandle.GetIndex()];
+	ClientReplicationInfo.Handles.Add(ViewerHandle);
 
 	for (FMassClientBubbleInfoData& InfoData : BubbleInfoArray)
 	{
@@ -665,7 +671,7 @@ void UMassReplicationSubsystem::RemoveClient(FMassClientHandle ClientHandle)
 	checkf(ClientToViewerHandleItem.ViewerHandle.IsValid(), TEXT("Invalid ViewerHandle! ClientHandle is out of sync with ClientToViewerHandleArray!"));
 	checkf(ClientToViewerHandleItem.ClientHandle == ClientHandle, TEXT("ClientHandle is out of sync with ClientToViewerHandleArray!"));
 
-	FClientViewerHandles& ClientViewerHandles = ClientsViewerHandles[ClientHandle.GetIndex()];
+	FMassClientReplicationInfo& ClientViewerHandles = ClientsReplicationInfo[ClientHandle.GetIndex()];
 
 	checkf(ClientViewerHandles.Handles.Num() > 0, TEXT("There should always be atleast one client viewer handle (the parent NetConnection)"));
 	ClientViewerHandles.Handles.Reset();
diff --git a/Engine/Plugins/Runtime/MassGameplay/Source/MassReplication/Private/MassReplicationTrait.cpp b/Engine/Plugins/Runtime/MassGameplay/Source/MassReplication/Private/MassReplicationTrait.cpp
--- a/Engine/Plugins/Runtime/MassGameplay/Source/MassReplication/Private/MassReplicationTrait.cpp
+++ b/Engine/Plugins/Runtime/MassGameplay/Source/MassReplication/Private/MassReplicationTrait.cpp
@@ -23,6 +23,7 @@ void UMassReplicationTrait::BuildTemplate(FMassEntityTemplateBuildContext& Build
 	BuildContext.AddFragment<FMassReplicatedAgentFragment>();
 	BuildContext.AddFragment<FMassReplicationViewerInfoFragment>();
 	BuildContext.AddFragment<FMassReplicationLODFragment>();
+	BuildContext.AddFragment<FMassReplicationGridCellLocationFragment>();
 
 	UMassEntitySubsystem* EntitySubsystem = UWorld::GetSubsystem<UMassEntitySubsystem>(&World);
 	check(EntitySubsystem);
diff --git a/Engine/Plugins/Runtime/MassGameplay/Source/MassReplication/Public/MassReplicationFragments.h b/Engine/Plugins/Runtime/MassGameplay/Source/MassReplication/Public/MassReplicationFragments.h
--- a/Engine/Plugins/Runtime/MassGameplay/Source/MassReplication/Public/MassReplicationFragments.h
+++ b/Engine/Plugins/Runtime/MassGameplay/Source/MassReplication/Public/MassReplicationFragments.h
@@ -23,25 +23,6 @@ struct MASSREPLICATION_API FMassNetworkIDFragment : public FMassFragment
 	FMassNetworkID NetID;
 };
 
-struct FMassReplicatedAgentArrayData
-{
-	void Invalidate()
-	{
-		Handle.Invalidate();
-		LastUpdateTime = 0.f;
-#if UE_ALLOW_DEBUG_REPLICATION
-		LOD = EMassLOD::Off;
-#endif // UE_DO_DEBUG_REPLICATION
-	}
-
-	FMassReplicatedAgentHandle Handle;
-	float LastUpdateTime = 0.f;
-
-#if UE_ALLOW_DEBUG_REPLICATION
-	EMassLOD::Type LOD = EMassLOD::Off;
-#endif // UE_DO_DEBUG_REPLICATION
-};
-
 /**
  * Agent handle per client, these will be at TArray indices of the Client handles indicies (used as a free list array)
  */
@@ -50,8 +31,7 @@ struct MASSREPLICATION_API FMassReplicatedAgentFragment : public FMassFragment
 {
 	GENERATED_BODY()
 
-	//these will be maintained so that the TArray entries are at the same indices as Client Handles
-	TArray<FMassReplicatedAgentArrayData, TInlineAllocator<UE::Mass::Replication::MaxNumOfClients>> AgentsData;
+	FMassReplicatedAgentData AgentData;
 };
 
 /*
@@ -79,12 +59,6 @@ struct MASSREPLICATION_API FMassReplicationLODFragment : public FMassFragment
 
 	/** Previous LOD information*/
 	TEnumAsByte<EMassLOD::Type> PrevLOD = EMassLOD::Max;
-
-	/** Per viewer LOD information */
-	TArray<TEnumAsByte<EMassLOD::Type>> LODPerViewer;
-
-	/** Per viewer previous LOD information */
-	TArray<TEnumAsByte<EMassLOD::Type>> PrevLODPerViewer;
 };
 
 UCLASS()
@@ -150,8 +124,8 @@ public:
 
 	FMassBubbleInfoClassHandle BubbleInfoClassHandle;
 
-	TArray<FMassClientHandle, TInlineAllocator<UE::Mass::Replication::MaxNumOfClients>> CachedClientHandles;
-	TArray<bool, TInlineAllocator<UE::Mass::Replication::MaxNumOfClients>> bBubbleChanged;
+	FMassClientHandle CurrentClientHandle;
+	TArray<FMassClientHandle> CachedClientHandles;
 
 	//TODO review if we need to have this as a UPROPERTY at all and also if we can make this use a TInlineAllocator
 	//Can not use TInlineAllocator with UPROPERTY()
@@ -179,4 +153,19 @@ public:
 
 		return *static_cast<T*>(BubbleInfo);
 	}
+};
+
+/** Cell location for replicated mass agents, used to fetch quickly the agents around each clients */
+USTRUCT()
+struct MASSREPLICATION_API FMassReplicationGridCellLocationFragment : public FMassFragment
+{
+	GENERATED_BODY()
+	FReplicationHashGrid2D::FCellLocation CellLoc;
+};
+
+/** Component Tag to tell if the entity is in the replication grid */
+USTRUCT()
+struct MASSREPLICATION_API FMassInReplicationGridTag : public FMassTag
+{
+	GENERATED_BODY()
 };
\ No newline at end of file
diff --git a/Engine/Plugins/Runtime/MassGameplay/Source/MassReplication/Public/MassReplicationGridProcessor.h b/Engine/Plugins/Runtime/MassGameplay/Source/MassReplication/Public/MassReplicationGridProcessor.h
new file mode 100644
--- /dev/null
+++ b/Engine/Plugins/Runtime/MassGameplay/Source/MassReplication/Public/MassReplicationGridProcessor.h
@@ -0,0 +1,53 @@
+// Copyright Epic Games, Inc. All Rights Reserved.
+
+#pragma once
+
+#include "MassProcessor.h"
+#include "MassEntityQuery.h"
+#include "MassObserverProcessor.h"
+
+#include "MassReplicationGridProcessor.generated.h"
+
+class UMassReplicationSubsystem;
+
+/** Processor to update entity in the replication grid used to fetch entities close to clients */
+UCLASS()
+class MASSREPLICATION_API UMassReplicationGridProcessor : public UMassProcessor
+{
+	GENERATED_BODY()
+
+public:
+	UMassReplicationGridProcessor();
+
+protected:
+	virtual void ConfigureQueries() override;
+	virtual void Initialize(UObject& Owner) override;
+	virtual void Execute(UMassEntitySubsystem& EntitySubsystem, FMassExecutionContext& Context) override;
+
+private:
+	UPROPERTY(Transient)
+	UMassReplicationSubsystem* ReplicationSubsystem;
+
+	FMassEntityQuery AddToGridEntityQuery;
+	FMassEntityQuery UpdateGridEntityQuery;
+	FMassEntityQuery RemoveFromGridEntityQuery;
+};
+
+/** Deinitializer processor to remove entity from the replication grid */
+UCLASS()
+class MASSREPLICATION_API UMassReplicationGridRemoverProcessor : public UMassObserverProcessor
+{
+	GENERATED_BODY()
+
+	UMassReplicationGridRemoverProcessor();
+
+protected:
+	virtual void ConfigureQueries() override;
+	virtual void Initialize(UObject& Owner) override;
+	virtual void Execute(UMassEntitySubsystem& EntitySubsystem, FMassExecutionContext& Context) override;
+
+	FMassEntityQuery EntityQuery;
+
+	UPROPERTY(Transient)
+	UMassReplicationSubsystem* ReplicationSubsystem;
+};
\ No newline at end of file
diff --git a/Engine/Plugins/Runtime/MassGameplay/Source/MassReplication/Public/MassReplicationProcessor.h b/Engine/Plugins/Runtime/MassGameplay/Source/MassReplication/Public/MassReplicationProcessor.h
--- a/Engine/Plugins/Runtime/MassGameplay/Source/MassReplication/Public/MassReplicationProcessor.h
+++ b/Engine/Plugins/Runtime/MassGameplay/Source/MassReplication/Public/MassReplicationProcessor.h
@@ -41,6 +41,7 @@ protected:
 	UPROPERTY(Transient)
 	UMassReplicationSubsystem* ReplicationSubsystem = nullptr;
 
+	FMassEntityQuery SyncClientData;
 	FMassEntityQuery CollectViewerInfoQuery;
 	FMassEntityQuery CalculateLODQuery;
 	FMassEntityQuery AdjustLODDistancesQuery;
@@ -109,87 +110,46 @@ void UMassReplicatorBase::CalculateClientReplication(FMassExecutionContext& Cont
 	{
 		FMassReplicatedAgentFragment& AgentFragment = ReplicatedAgentList[EntityIdx];
 
-		if (AgentFragment.AgentsData.Num() < RepSharedFragment.CachedClientHandles.Num())
-		{
-			AgentFragment.AgentsData.AddDefaulted(RepSharedFragment.CachedClientHandles.Num() - AgentFragment.AgentsData.Num());
-		}
-		else if (AgentFragment.AgentsData.Num() > RepSharedFragment.CachedClientHandles.Num())
-		{
-			AgentFragment.AgentsData.RemoveAt(RepSharedFragment.CachedClientHandles.Num(), AgentFragment.AgentsData.Num() - RepSharedFragment.CachedClientHandles.Num(), /* bAllowShrinking */ false);
-		}
+		const FMassClientHandle& ClientHandle = RepSharedFragment.CurrentClientHandle;
+		check(ClientHandle.IsValid());
+
+		checkSlow(RepSharedFragment.BubbleInfos[ClientHandle.GetIndex()] != nullptr);
 
-		for (int32 ClientIdx = 0; ClientIdx < RepSharedFragment.CachedClientHandles.Num(); ++ClientIdx)
+		FMassReplicatedAgentData& AgentData = AgentFragment.AgentData;
+
+		const EMassLOD::Type LOD = ViewerLODList[EntityIdx].LOD;
+
+		if (LOD < EMassLOD::Off)
 		{
-			const FMassClientHandle& ClientHandle = RepSharedFragment.CachedClientHandles[ClientIdx];
+			AgentData.LOD = LOD;
+
+			//if the handle isn't valid we need to add the agent
+			if (!AgentData.Handle.IsValid())
+			{
+				typename AgentArrayItem::FReplicatedAgentType ReplicatedAgent;
 
-			if (ClientHandle.IsValid())
+				const FMassNetworkIDFragment& NetIDFragment = NetworkIDList[EntityIdx];
+				const FReplicationTemplateIDFragment& TemplateIDFragment = TemplateIDList[EntityIdx];
+
+				ReplicatedAgent.SetNetID(NetIDFragment.NetID);
+				ReplicatedAgent.SetTemplateID(TemplateIDFragment.ID);
+
+				AgentData.Handle = AddEntity(Context, EntityIdx, ReplicatedAgent, ClientHandle);
+
+				AgentData.LastUpdateTime = Time;
+			}
+			else
+			{
+				ModifyEntity(Context, EntityIdx, LOD, Time, AgentData.Handle, ClientHandle);
+			}
+		}
+		else
+		{
+			// as this is a fresh handle, if its valid then we can use the unsafe remove function
+			if (AgentData.Handle.IsValid())
 			{
-				checkSlow(RepSharedFragment.BubbleInfos[ClientHandle.GetIndex()] != nullptr);
-
-				FMassReplicatedAgentArrayData& AgentData = AgentFragment.AgentsData[ClientIdx];
-
-				//if the bubble has changed, set the handle Invalid. We will set this to something valid if the agent is going to replicate to the bubble
-				//When a bubble has changed the existing AMassCrowdClientBubbleInfo will reset all the data associated with it.
-				if (RepSharedFragment.bBubbleChanged[ClientIdx])
-				{
-					AgentData.Invalidate();
-				}
-
-				//now we need to see what our highest viewer LOD is on this client (split screen etc), we can use the unsafe version as we have already checked all the
-				//CachedClientHandles for validity
-				const FClientViewerHandles& ClientViewers = ReplicationContext.ReplicationSubsystem.GetClientViewersChecked(ClientHandle);
-
-				EMassLOD::Type HighestLOD = EMassLOD::Off;
-
-				for (const FMassViewerHandle& ViewerHandle : ClientViewers.Handles)
-				{
-					//this should always we valid as we synchronized the viewers just previously
-					check(ReplicationContext.LODSubsystem.IsValidViewer(ViewerHandle));
-
-					const EMassLOD::Type MassLOD = ViewerLODList[EntityIdx].LODPerViewer[ViewerHandle.GetIndex()];
-					check(MassLOD <= EMassLOD::Off);
-
-					if (HighestLOD > MassLOD)
-					{
-						HighestLOD = MassLOD;
-					}
-				}
-
-				if (HighestLOD < EMassLOD::Off)
-				{
-#if UE_ALLOW_DEBUG_REPLICATION
-					AgentData.LOD = HighestLOD;
-#endif // UE_ALLOW_DEBUG_REPLICATION
-
-					//if the handle isn't valid we need to add the agent
-					if (!AgentData.Handle.IsValid())
-					{
-						typename AgentArrayItem::FReplicatedAgentType ReplicatedAgent;
-
-						const FMassNetworkIDFragment& NetIDFragment = NetworkIDList[EntityIdx];
-						const FReplicationTemplateIDFragment& TemplateIDFragment = TemplateIDList[EntityIdx];
-
-						ReplicatedAgent.SetNetID(NetIDFragment.NetID);
-						ReplicatedAgent.SetTemplateID(TemplateIDFragment.ID);
-
-						AgentData.Handle = AddEntity(Context, EntityIdx, ReplicatedAgent, ClientHandle);
-
-						AgentData.LastUpdateTime = Time;
-					}
-					else
-					{
-						ModifyEntity(Context, EntityIdx, HighestLOD, Time, AgentData.Handle, ClientHandle);
-					}
-				}
-				else
-				{
-					// as this is a fresh handle, if its valid then we can use the unsafe remove function
-					if (AgentData.Handle.IsValid())
-					{
-						RemoveEntity(Context, AgentData.Handle, ClientHandle);
-						AgentData.Invalidate();
-					}
-				}
+				RemoveEntity(Context, AgentData.Handle, ClientHandle);
+				AgentData.Invalidate();
 			}
 		}
 	}
diff --git a/Engine/Plugins/Runtime/MassGameplay/Source/MassReplication/Public/MassReplicationSettings.h b/Engine/Plugins/Runtime/MassGameplay/Source/MassReplication/Public/MassReplicationSettings.h
new file mode 100644
--- /dev/null
+++ b/Engine/Plugins/Runtime/MassGameplay/Source/MassReplication/Public/MassReplicationSettings.h
@@ -0,0 +1,28 @@
+// Copyright Epic Games, Inc. All Rights Reserved.
+
+#pragma once
+
+#include "MassSettings.h"
+#include "MassReplicationSettings.generated.h"
+
+UCLASS(config = Mass, defaultconfig, meta = (DisplayName = "Mass Replication"))
+class MASSREPLICATION_API UMassReplicationSettings : public UMassModuleSettings
+{
+	GENERATED_BODY()
+
+public:
+	static const UMassReplicationSettings* Get()
+	{
+		return GetDefault<UMassReplicationSettings>();
+	}
+
+	float GetReplicationGridCellSize() const
+	{ 
+		return ReplicationGridCellSize; 
+	}
+	
+private:
+
+	UPROPERTY(EditAnywhere, config, Category = Replication);
+	float ReplicationGridCellSize = 5000.0f; // 50 meters;
+};
diff --git a/Engine/Plugins/Runtime/MassGameplay/Source/MassReplication/Public/MassReplicationSubsystem.h b/Engine/Plugins/Runtime/MassGameplay/Source/MassReplication/Public/MassReplicationSubsystem.h
--- a/Engine/Plugins/Runtime/MassGameplay/Source/MassReplication/Public/MassReplicationSubsystem.h
+++ b/Engine/Plugins/Runtime/MassGameplay/Source/MassReplication/Public/MassReplicationSubsystem.h
@@ -14,9 +14,19 @@ class AMassCrowdClientBubbleInfo;
 class UMassLODSubsystem;
 class AMassClientBubbleInfoBase;
 
-struct FClientViewerHandles
+typedef TMap<FMassEntityHandle, FMassReplicatedAgentData> FMassReplicationAgentDataMap;
+
+struct FMassClientReplicationInfo
 {
+	/** Array of all the viewer of this client */
 	TArray<FMassViewerHandle> Handles;
+
+	/** Array of all the entities handled by this client 
+	  * This array might contains duplicates if there is more than one viewer per client as this concatenates the entities from all viewer of a client) */
+	TArray<FMassEntityHandle> HandledEntities;
+
+	/** The saved agent data of the entities handle by this client */
+	FMassReplicationAgentDataMap AgentsData;
 };
 
 struct FViewerClientPair
@@ -62,12 +72,11 @@ namespace UE::MassReplication
 
 } // UE::MassReplication
 
-
 /**
- *  Manages the creation of NetworkIDs, ClientBubbles and ClientViewers.
+ *  Manages the creation of NetworkIDs, ClientBubbles and ClientReplicationInfo.
  *  NetworkIDs are per replicated Agent Entity and are unique and replicated between server and clients.
  *  ClientBubbles relate to the player controller that owns the parent UNetConnection to a Client machine.
- *  ClientViewers relate to all the player controllers that have a parent or child UNetConnection to a single Client machine (split screen etc).
+ *  ClientReplicationInfo relate to all the player controllers that have a parent or child UNetConnection to a single Client machine (split screen etc).
  */
 UCLASS()
 class MASSREPLICATION_API UMassReplicationSubsystem : public UWorldSubsystem
@@ -76,6 +85,8 @@ class MASSREPLICATION_API UMassReplicationSubsystem : public UWorldSubsystem
 
 public:
 
+	UMassReplicationSubsystem();
+
 	enum class EFindOrAddMassEntityInfo : uint8
 	{
 		FoundOlderReplicationID,
@@ -136,18 +147,26 @@ public:
 		return static_cast<TType *>(ClientBubble);
 	}
 
-	/** Gets the client viewers safely */
-	const FClientViewerHandles* GetClientViewers(FMassClientHandle Handle) const
+	/** Gets the client replication info safely */
+	const FMassClientReplicationInfo* GetClientReplicationInfo(FMassClientHandle Handle) const
+	{
+		return ClientHandleManager.IsValidHandle(Handle) ? &(ClientsReplicationInfo[Handle.GetIndex()]) : nullptr;
+	}
+
+	/** Gets the client replication info. Faster version using check()s */
+	const FMassClientReplicationInfo& GetClientReplicationInfoChecked(FMassClientHandle Handle) const
 	{
-		return ClientHandleManager.IsValidHandle(Handle) ? &(ClientsViewerHandles[Handle.GetIndex()]) : nullptr;
+		check(ClientHandleManager.IsValidHandle(Handle));
+
+		return ClientsReplicationInfo[Handle.GetIndex()];
 	}
 
-	/** Gets the client viewers. Faster version using check()s */
-	const FClientViewerHandles& GetClientViewersChecked(FMassClientHandle Handle) const
+	/** Gets the client replication info. Faster version using check()s */
+	FMassClientReplicationInfo& GetMutableClientReplicationInfoChecked(FMassClientHandle Handle)
 	{
 		check(ClientHandleManager.IsValidHandle(Handle));
 
-		return ClientsViewerHandles[Handle.GetIndex()];
+		return ClientsReplicationInfo[Handle.GetIndex()];
 	}
 
 	void SynchronizeClientsAndViewers();
@@ -164,6 +183,9 @@ public:
 	/** @return FMassBubbleInfoClassHandle Handle to the BubbleInfoClass, this will be an invalid handle if BubbleInfoClass can not be found */
 	FMassBubbleInfoClassHandle GetBubbleInfoClassHandle(const TSubclassOf<AMassClientBubbleInfoBase>& BubbleInfoClass) const;
 
+	const FReplicationHashGrid2D& GetGrid() const { return ReplicationGrid; }
+	FReplicationHashGrid2D& GetGridMutable() { return ReplicationGrid; }
+
 #if UE_REPLICATION_COMPILE_CLIENT_CODE
 
 	/** 
@@ -255,16 +277,16 @@ protected:
 
 	/** An Array of each Clients viewer handles (split screen players sharing the same client connections). 
 	 *  This will include both the parent and child NetConnections per client.
-	 *	The array is organised so the array index is the same as the index of the client FMassClientHandle.
+	 *	The array is organized so the array index is the same as the index of the client FMassClientHandle.
 	 */
-	TArray<FClientViewerHandles> ClientsViewerHandles;
+	TArray<FMassClientReplicationInfo> ClientsReplicationInfo;
 
-	/** A free list array of FViewerClientPairs. This is organised so the index is that of the FMassViewerHandle for fast lookup of the related FMassClientHandle. 
+	/** A free list array of FViewerClientPairs. This is organized so the index is that of the FMassViewerHandle for fast lookup of the related FMassClientHandle. 
 	 *  This only contains viewers that are also clients (ie only parent NetConnections not child ones).
 	 */
 	TArray<FViewerClientPair> ViewerToClientHandleArray;
 
-	/** A free list array of FViewerClientPairs. This is organised so the index is that of the FMassClientHandle.
+	/** A free list array of FViewerClientPairs. This is organized so the index is that of the FMassClientHandle.
 	 *  For fast lookup of the related FMassViewerHandle. This only contains viewers that are also clients (ie only parent NetConnections not child ones).
 	 */
 	TArray<FViewerClientPair> ClientToViewerHandleArray;
@@ -291,4 +313,6 @@ protected:
 	 * @todo this comment is no longer accurate, needs fixing
 	 */
 	uint64 LastSynchronizedFrame = 0;
+
+	FReplicationHashGrid2D ReplicationGrid;
 };
\ No newline at end of file
diff --git a/Engine/Plugins/Runtime/MassGameplay/Source/MassReplication/Public/MassReplicationTypes.h b/Engine/Plugins/Runtime/MassGameplay/Source/MassReplication/Public/MassReplicationTypes.h
--- a/Engine/Plugins/Runtime/MassGameplay/Source/MassReplication/Public/MassReplicationTypes.h
+++ b/Engine/Plugins/Runtime/MassGameplay/Source/MassReplication/Public/MassReplicationTypes.h
@@ -8,6 +8,7 @@
 #include "MassEntityTemplate.h"
 #include "MassObserverProcessor.h"
 #include "MassLODLogic.h"
+#include "HierarchicalHashGrid2D.h"
 
 #include "MassReplicationTypes.generated.h"
 
@@ -62,11 +63,6 @@
 
 class UMassReplicationSubsystem;
 
-namespace UE { namespace Mass { namespace Replication
-{
-	const int32 MaxNumOfClients = UE::MassLOD::MaxNumOfViewers;
-}}}
-
 DECLARE_LOG_CATEGORY_EXTERN(LogMassReplication, Log, All);
 
 /** Handle of each Agent per Client (bubble), these handles are not consistent between different client bubbles.
@@ -145,12 +141,20 @@ private:
 
 struct FReplicationLODLogic : public FLODDefaultLogic
 {
-	enum
+};
+
+struct FMassReplicatedAgentData
+{
+	void Invalidate()
 	{
-		bStoreInfoPerViewer = true,
-		bCalculateLODPerViewer = true,
-		bMaximizeCountPerViewer = true,
-	};
+		Handle.Invalidate();
+		LastUpdateTime = 0.f;
+		LOD = EMassLOD::Off;
+	}
+
+	FMassReplicatedAgentHandle Handle;
+	float LastUpdateTime = 0.f;
+	EMassLOD::Type LOD = EMassLOD::Off;
 };
 
 #if UE_REPLICATION_COMPILE_CLIENT_CODE
@@ -169,4 +173,6 @@ struct FMassReplicationEntityInfo
 	/** This is stored between removes and adds, however this item in the UMassReplicationSubsystem::EntityInfoMap will eventually get cleaned up if Entity.IsSet() == false for a fairly substantial length of time. */
 	int32 ReplicationID = INDEX_NONE;
 };
-#endif //UE_REPLICATION_COMPILE_CLIENT_CODE
\ No newline at end of file
+#endif //UE_REPLICATION_COMPILE_CLIENT_CODE
+
+typedef THierarchicalHashGrid2D<2, 4, FMassEntityHandle> FReplicationHashGrid2D;	// 2 levels of hierarchy, 4 ratio between levels
\ No newline at end of file
