NoPCH compilation fix for MassSpawner

diff --git a/Engine/Plugins/Runtime/MassGameplay/Source/MassSpawner/Private/MassSpawner.cpp b/Engine/Plugins/Runtime/MassGameplay/Source/MassSpawner/Private/MassSpawner.cpp
--- a/Engine/Plugins/Runtime/MassGameplay/Source/MassSpawner/Private/MassSpawner.cpp
+++ b/Engine/Plugins/Runtime/MassGameplay/Source/MassSpawner/Private/MassSpawner.cpp
@@ -17,6 +17,9 @@
 #include "Engine/AssetManager.h"
 #include "MassSpawnLocationProcessor.h"
 #include "MassExecutor.h"
+#if WITH_EDITOR
+#include "Engine/Texture2D.h"
+#endif
 
 
 namespace UE::MassSpawner
