Trait fragment ownership and dependencies

diff --git a/Engine/Plugins/Runtime/MassGameplay/Source/MassSpawner/Private/MassEntityConfigAsset.cpp b/Engine/Plugins/Runtime/MassGameplay/Source/MassSpawner/Private/MassEntityConfigAsset.cpp
--- a/Engine/Plugins/Runtime/MassGameplay/Source/MassSpawner/Private/MassEntityConfigAsset.cpp
+++ b/Engine/Plugins/Runtime/MassGameplay/Source/MassSpawner/Private/MassEntityConfigAsset.cpp
@@ -61,17 +61,7 @@ const FMassEntityTemplate* FMassEntityConfig::GetOrCreateEntityTemplate(AActor&
 	FMassEntityTemplate& Template = TemplateRegistry.CreateTemplate(Hash, TemplateID);
 	FMassEntityTemplateBuildContext BuildContext(Template);
 
-	for (const UMassEntityTraitBase* Trait : CombinedTraits)
-	{
-		check(Trait);
-		Trait->BuildTemplate(BuildContext, *World);
-	}
-
-	for (const UMassEntityTraitBase* Trait : CombinedTraits)
-	{
-		check(Trait);
-		Trait->ValidateTemplate(BuildContext, *World);
-	}
+	BuildContext.BuildFromTraits(CombinedTraits, *World);
 
 	if (ensureMsgf(!Template.IsEmpty(), TEXT("Need at least one fragment to create an Archetype")))
 	{
diff --git a/Engine/Plugins/Runtime/MassGameplay/Source/MassSpawner/Private/MassEntityTemplateRegistry.cpp b/Engine/Plugins/Runtime/MassGameplay/Source/MassSpawner/Private/MassEntityTemplateRegistry.cpp
--- a/Engine/Plugins/Runtime/MassGameplay/Source/MassSpawner/Private/MassEntityTemplateRegistry.cpp
+++ b/Engine/Plugins/Runtime/MassGameplay/Source/MassSpawner/Private/MassEntityTemplateRegistry.cpp
@@ -7,6 +7,8 @@
 #include "VisualLogger/VisualLogger.h"
 #include "HAL/IConsoleManager.h"
 #include "MassSpawnerSubsystem.h"
+#include "MassEntityTypes.h"
+#include "MassEntityTraitBase.h"
 
 namespace FTemplateRegistryHelpers
 {
@@ -189,3 +191,74 @@ void UMassEntityTemplateRegistry::DestroyTemplate(const uint32 HashLookup, FMass
 	TemplateIDToTemplateMap.Remove(TemplateID);
 }
 
+//----------------------------------------------------------------------//
+// FMassEntityTemplateBuildContext 
+//----------------------------------------------------------------------//
+void FMassEntityTemplateBuildContext::BuildFromTraits(TConstArrayView<UMassEntityTraitBase*> Traits, UWorld& World)
+{
+	TraitAddedTypes.Reset();
+	TraitsDependencies.Reset();
+
+	for (const UMassEntityTraitBase* Trait : Traits)
+	{
+		check(Trait);
+		BuildingTrait = Trait;
+		BuildingTrait->BuildTemplate(*this, World);
+	}
+
+	BuildingTrait = nullptr;
+
+	ValidateBuildContext(World);
+}
+
+void FMassEntityTemplateBuildContext::ValidateBuildContext(UWorld& World)
+{
+	// Group same types(key) together
+	TraitAddedTypes.KeySort( [](const UStruct& LHS, const UStruct& RHS) { return LHS.GetName() < RHS.GetName(); } );
+
+	// Loop through all the registered fragments and make sure only one trait registered them.
+	const UStruct* CurrentStruct = nullptr;
+	const UMassEntityTraitBase* CurrentTrait = nullptr;
+	bool bHeaderOutputed = false;
+	for (const auto& Pair : TraitAddedTypes)
+	{
+		if (CurrentStruct != Pair.Key)
+		{
+			CurrentStruct = Pair.Key;
+			CurrentTrait = Pair.Value;
+			CurrentTrait->ValidateTemplate(*this, World);
+			bHeaderOutputed = false;
+		}
+		else
+		{
+			if (!bHeaderOutputed)
+			{
+				UE_LOG(LogMass, Error, TEXT("Fragment(%s) was added multiple time and can only be added by one trait. Fragment was added by:"), *CurrentStruct->GetName());
+				UE_LOG(LogMass, Error, TEXT("\t\t%s"), *CurrentTrait->GetClass()->GetName());
+				bHeaderOutputed = true;
+			}
+			UE_LOG(LogMass, Error, TEXT("\t\t%s"), *Pair.Value->GetClass()->GetName());
+		}
+ 	}
+
+	// Loop through all the traits dependencies and check if they have been added
+	CurrentTrait = nullptr;
+	bHeaderOutputed = false;
+	for (const auto& Dependency : TraitsDependencies)
+	{
+		if (CurrentTrait != Dependency.Get<1>())
+		{
+			CurrentTrait = Dependency.Get<1>();
+			bHeaderOutputed = false;
+		}
+		if (!TraitAddedTypes.Contains(Dependency.Get<0>()))
+		{
+			if (!bHeaderOutputed)
+			{
+				UE_LOG(LogMass, Error, TEXT("Trait(%s) has missing dependency:"),  *CurrentTrait->GetClass()->GetName() );
+				bHeaderOutputed = true;
+			}
+			UE_LOG(LogMass, Error, TEXT("\t\t%s"), *Dependency.Get<0>()->GetName());
+		}
+	}
+}
diff --git a/Engine/Plugins/Runtime/MassGameplay/Source/MassSpawner/Public/MassEntityTemplateRegistry.h b/Engine/Plugins/Runtime/MassGameplay/Source/MassSpawner/Public/MassEntityTemplateRegistry.h
--- a/Engine/Plugins/Runtime/MassGameplay/Source/MassSpawner/Public/MassEntityTemplateRegistry.h
+++ b/Engine/Plugins/Runtime/MassGameplay/Source/MassSpawner/Public/MassEntityTemplateRegistry.h
@@ -12,6 +12,7 @@
 
 class UWorld;
 class UMassEntitySubsystem;
+class UMassEntityTraitBase;
 
 enum class EFragmentInitialization : uint8
 {
@@ -31,55 +32,68 @@ struct FMassEntityTemplateBuildContext
 	template<typename T>
 	T& AddFragment_GetRef()
 	{
+		TypeAdded(*T::StaticStruct());
 		return Template.AddFragment_GetRef<T>();
 	}
 
 	template<typename T>
 	void AddFragment()
 	{
+		TypeAdded(*T::StaticStruct());
 		Template.AddFragment<T>();
 	}
 
 	void AddFragment(FConstStructView InFragment)
 	{ 
+		checkf(InFragment.GetScriptStruct(), TEXT("Expecting a valid fragment type"));
+		TypeAdded(*InFragment.GetScriptStruct());
 		Template.AddFragment(InFragment);
 	}
 
 	template<typename T>
 	void AddTag()
 	{
+		// Tags can be added by multiple traits, so they do not follow the same rules as fragments
 		Template.AddTag<T>();
 	}
 
 	void AddTag(const UScriptStruct& TagType)
 	{
+		// Tags can be added by multiple traits, so they do not follow the same rules as fragments
 		Template.AddTag(TagType);
 	}
 
 	template<typename T>
 	void AddChunkFragment()
 	{
+		TypeAdded(*T::StaticStruct());
 		Template.AddChunkFragment<T>();
 	}
 
 	void AddConstSharedFragment(const FConstSharedStruct& InSharedFragment)
 	{
+		checkf(InSharedFragment.GetScriptStruct(), TEXT("Expecting a valide shared fragment type"));
+		TypeAdded(*InSharedFragment.GetScriptStruct());
 		Template.AddConstSharedFragment(InSharedFragment);
 	}
 
 	void AddSharedFragment(const FSharedStruct& InSharedFragment)
 	{
+		checkf(InSharedFragment.GetScriptStruct(), TEXT("Expecting a valide shared fragment type"));
+		TypeAdded(*InSharedFragment.GetScriptStruct());
 		Template.AddSharedFragment(InSharedFragment);
 	}
 
 	template<typename T>
 	bool HasFragment() const
 	{
+		ensureMsgf(!BuildingTrait, TEXT("This method is not expected to be called within the build from trait call."));
 		return Template.HasFragment<T>();
 	}
 	
 	bool HasFragment(const UScriptStruct& ScriptStruct) const
 	{
+		ensureMsgf(!BuildingTrait, TEXT("This method is not expected to be called within the build from trait call."));
 		return Template.HasFragment(ScriptStruct);
 	}
 
@@ -92,17 +106,20 @@ struct FMassEntityTemplateBuildContext
 	template<typename T>
 	bool HasChunkFragment() const
 	{
+		ensureMsgf(!BuildingTrait, TEXT("This method is not expected to be called within the build from trait call."));
 		return Template.HasChunkFragment<T>();
 	}
 
 	template<typename T>
 	bool HasSharedFragment() const
 	{
+		ensureMsgf(!BuildingTrait, TEXT("This method is not expected to be called within the build from trait call."));
 		return Template.HasSharedFragment<T>();
 	}
 
 	bool HasSharedFragment(const UScriptStruct& ScriptStruct) const
 	{
+		ensureMsgf(!BuildingTrait, TEXT("This method is not expected to be called within the build from trait call."));
 		return Template.HasSharedFragment(ScriptStruct);
 	}
 
@@ -112,17 +129,50 @@ struct FMassEntityTemplateBuildContext
 	template<typename T>
 	void AddTranslator()
 	{
-		TArray<UScriptStruct*> TagTypes;
+		TypeAdded(*T::StaticClass());
 		GetDefault<T>()->AppendRequiredTags(Template.GetMutableTags());
 	}
 
+	//----------------------------------------------------------------------//
+	// Dependencies
+	//----------------------------------------------------------------------//
+	template<typename T>
+	void RequireFragment()
+	{
+		AddDependency(T::StaticStruct());
+	}
+
+	void AddDependency(const UStruct* Dependency)
+	{
+		TraitsDependencies.Add( {Dependency, BuildingTrait} );
+	}
+
 	//----------------------------------------------------------------------//
 	// Template access
 	//----------------------------------------------------------------------//
 	FMassEntityTemplateID GetTemplateID() const { return Template.GetTemplateID(); }
 	TArray<FMassEntityTemplate::FObjectFragmentInitializerFunction>& GetMutableObjectFragmentInitializers() { return Template.GetMutableObjectFragmentInitializers(); }
 
+	//----------------------------------------------------------------------//
+	// Build methods
+	//----------------------------------------------------------------------//
+	void BuildFromTraits(TConstArrayView<UMassEntityTraitBase*> Traits, UWorld& World);
+
 protected:
+
+	void TypeAdded(const UStruct& Type)
+	{
+		if (ensureMsgf(BuildingTrait, TEXT("Expected to be called within the BuildTemplateFromTrait method")))
+		{
+			TraitAddedTypes.Add(&Type, BuildingTrait);
+		}
+	}
+	void ValidateBuildContext(UWorld& World);
+
+	const UMassEntityTraitBase* BuildingTrait = nullptr;
+	TMultiMap<const UStruct*, const UMassEntityTraitBase*> TraitAddedTypes;
+	TArray< TTuple<const UStruct*, const UMassEntityTraitBase*> > TraitsDependencies;
+
 	FMassEntityTemplate& Template;
 };
 
