Added an button in property grid to validate the entity config asset

diff --git a/Engine/Plugins/Runtime/MassGameplay/Source/MassActors/Private/MassAgentSubsystem.cpp b/Engine/Plugins/Runtime/MassGameplay/Source/MassActors/Private/MassAgentSubsystem.cpp
--- a/Engine/Plugins/Runtime/MassGameplay/Source/MassActors/Private/MassAgentSubsystem.cpp
+++ b/Engine/Plugins/Runtime/MassGameplay/Source/MassActors/Private/MassAgentSubsystem.cpp
@@ -82,15 +82,11 @@ FMassEntityTemplateID UMassAgentSubsystem::RegisterAgentComponent(UMassAgentComp
 
 	AActor* AgentActor = AgentComp.GetOwner();
 	check(AgentActor);
+	UWorld* World = AgentActor->GetWorld();
+	check(World);
 
 	const FMassEntityConfig& EntityConfig = AgentComp.GetEntityConfig();
-	const FMassEntityTemplate* EntityTemplate = EntityConfig.GetOrCreateEntityTemplate(*AgentActor, AgentComp);
-
-	if (!ensureMsgf(EntityTemplate, TEXT("Failed to fetch entity template for %s."), *AgentActor->GetName()))
-	{
-		UE_VLOG(this, LogMassActor, Warning, TEXT("Failed to fetch entity template for %s."), *AgentActor->GetName());
-		return FMassEntityTemplateID();
-	}
+	const FMassEntityTemplate& EntityTemplate = EntityConfig.GetOrCreateEntityTemplate(*World, AgentComp);
 
 #if UE_REPLICATION_COMPILE_CLIENT_CODE
 	if (AgentComp.IsNetSimulating())
@@ -100,7 +96,7 @@ FMassEntityTemplateID UMassAgentSubsystem::RegisterAgentComponent(UMassAgentComp
 	else
 #endif // UE_REPLICATION_COMPILE_CLIENT_CODE
 	{
-		FMassAgentInitializationQueue& AgentQueue = PendingAgentEntities.FindOrAdd(EntityTemplate->GetTemplateID());
+		FMassAgentInitializationQueue& AgentQueue = PendingAgentEntities.FindOrAdd(EntityTemplate.GetTemplateID());
 		// Agent already in the queue! Earlier conditions should have failed or data is inconsistent.
 		check(AgentQueue.AgentComponents.Find(&AgentComp) == INDEX_NONE);
 		AgentQueue.AgentComponents.Add(&AgentComp);
@@ -109,7 +105,7 @@ FMassEntityTemplateID UMassAgentSubsystem::RegisterAgentComponent(UMassAgentComp
    		AgentComp.EntityCreationPending();
 	}
 
-	return EntityTemplate->GetTemplateID();
+	return EntityTemplate.GetTemplateID();
 }
 
 void UMassAgentSubsystem::UpdateAgentComponent(const UMassAgentComponent& AgentComp)
@@ -126,19 +122,15 @@ void UMassAgentSubsystem::UpdateAgentComponent(const UMassAgentComponent& AgentC
 
 	AActor* AgentActor = AgentComp.GetOwner();
 	check(AgentActor);
+	UWorld* World = AgentActor->GetWorld();
+	check(World);
 
 	const FMassEntityConfig& EntityConfig = AgentComp.GetEntityConfig();
-	const FMassEntityTemplate* EntityTemplate = EntityConfig.GetOrCreateEntityTemplate(*AgentActor, AgentComp);
-
-	if (!ensureMsgf(EntityTemplate, TEXT("Failed to fetch entity template for %s."), *AgentActor->GetName()))
-	{
-		UE_VLOG(this, LogMassActor, Warning, TEXT("Failed to fetch entity template for %s."), *AgentActor->GetName());
-		return;
-	}
+	const FMassEntityTemplate& EntityTemplate = EntityConfig.GetOrCreateEntityTemplate(*World, AgentComp);
 
 	const FMassEntityHandle Entity = AgentComp.GetEntityHandle();
 	const FMassArchetypeHandle CurrentArchetypeHandle = EntitySystem->GetArchetypeForEntity(Entity);
-	if (CurrentArchetypeHandle == EntityTemplate->GetArchetype())
+	if (CurrentArchetypeHandle == EntityTemplate.GetArchetype())
 	{
 		UE_VLOG(this, LogMassActor, Log, TEXT("%s called for %s but no archetype changes have been found")
 			, ANSI_TO_TCHAR(__FUNCTION__), *AgentActor->GetName());
@@ -210,10 +202,11 @@ void UMassAgentSubsystem::UnregisterAgentComponent(UMassAgentComponent& AgentCom
 			// the entity has already been created. Destroy!
 			const FMassEntityTemplate* EntityTemplate = nullptr;
 			AActor* AgentActor = AgentComp.GetOwner();
-			if (ensure(AgentActor))
+			UWorld* World = AgentActor ? AgentActor->GetWorld() : nullptr;
+			if (ensure(World))
 			{
 				const FMassEntityConfig& EntityConfig = AgentComp.GetEntityConfig();
-				EntityTemplate = &EntityConfig.GetEntityTemplateChecked(*AgentActor, AgentComp);
+				EntityTemplate = &EntityConfig.GetEntityTemplateChecked(*World, AgentComp);
 			}
 
 			FMassEntityHandle Entity = AgentComp.GetEntityHandle();
diff --git a/Engine/Plugins/Runtime/MassGameplay/Source/MassSpawner/Private/MassEntityConfigAsset.cpp b/Engine/Plugins/Runtime/MassGameplay/Source/MassSpawner/Private/MassEntityConfigAsset.cpp
--- a/Engine/Plugins/Runtime/MassGameplay/Source/MassSpawner/Private/MassEntityConfigAsset.cpp
+++ b/Engine/Plugins/Runtime/MassGameplay/Source/MassSpawner/Private/MassEntityConfigAsset.cpp
@@ -8,6 +8,11 @@
 #include "Serialization/ArchiveObjectCrc32.h"
 #include "VisualLogger/VisualLogger.h"
 #include "Engine/World.h"
+#if WITH_EDITOR
+#include "Editor.h"
+#endif // WITH_EDITOR
+
+#define LOCTEXT_NAMESPACE "Mass"
 
 namespace UE::MassSpawner
 {
@@ -39,18 +44,31 @@ FMassEntityConfig::FMassEntityConfig(UMassEntityConfigAsset& InParent)
 
 }
 
-const FMassEntityTemplate* FMassEntityConfig::GetOrCreateEntityTemplate(AActor& OwnerActor, const UObject& ConfigOwner) const
+#if WITH_EDITOR
+void UMassEntityConfigAsset::ValidateEntityConfig()
+{
+	if (UWorld* EditorWorld = GEditor->GetEditorWorldContext().World())
+	{
+		if (Config.ValidateEntityTemplate(*EditorWorld, *this))
+		{
+			FMessageLog EditorInfo("LogMass");
+			EditorInfo.Info(LOCTEXT("MassEntityConfigAssetNoErrorsDetected", "There were no error detected during validation of the EntityConfigAsset"));
+		}
+	}
+}
+#endif // WITH_EDITOR
+
+const FMassEntityTemplate& FMassEntityConfig::GetOrCreateEntityTemplate(UWorld& World, const UObject& ConfigOwner) const
 {
 	uint32 Hash;
 	FMassEntityTemplateID TemplateID;
 	TArray<UMassEntityTraitBase*> CombinedTraits;
-	if (const FMassEntityTemplate* ExistingTemplate = GetEntityTemplateInternal(OwnerActor, ConfigOwner, Hash, TemplateID, CombinedTraits))
+	if (const FMassEntityTemplate* ExistingTemplate = GetEntityTemplateInternal(World, ConfigOwner, Hash, TemplateID, CombinedTraits))
 	{
-		return ExistingTemplate;
+		return *ExistingTemplate;
 	}
 
-	UWorld* World = OwnerActor.GetWorld();
-	UMassSpawnerSubsystem* SpawnerSystem = UWorld::GetSubsystem<UMassSpawnerSubsystem>(World);
+	UMassSpawnerSubsystem* SpawnerSystem = UWorld::GetSubsystem<UMassSpawnerSubsystem>(&World);
 	check(SpawnerSystem);
 	UMassEntityTemplateRegistry& TemplateRegistry = SpawnerSystem->GetTemplateRegistryInstance();
 
@@ -61,29 +79,28 @@ const FMassEntityTemplate* FMassEntityConfig::GetOrCreateEntityTemplate(AActor&
 	FMassEntityTemplate& Template = TemplateRegistry.CreateTemplate(Hash, TemplateID);
 	FMassEntityTemplateBuildContext BuildContext(Template);
 
-	BuildContext.BuildFromTraits(CombinedTraits, *World);
+	BuildContext.BuildFromTraits(CombinedTraits, World);
 
 	if (ensureMsgf(!Template.IsEmpty(), TEXT("Need at least one fragment to create an Archetype")))
 	{
 		TemplateRegistry.InitializeEntityTemplate(Template);
 	}
 
-	return &Template;
+	return Template;
 }
 
-void FMassEntityConfig::DestroyEntityTemplate(const AActor& OwnerActor, const UObject& ConfigOwner) const
+void FMassEntityConfig::DestroyEntityTemplate(UWorld& World, const UObject& ConfigOwner) const
 {
 	uint32 Hash;
 	FMassEntityTemplateID TemplateID;
 	TArray<UMassEntityTraitBase*> CombinedTraits;
-	const FMassEntityTemplate* Template = GetEntityTemplateInternal(OwnerActor, ConfigOwner, Hash, TemplateID, CombinedTraits);
+	const FMassEntityTemplate* Template = GetEntityTemplateInternal(World, ConfigOwner, Hash, TemplateID, CombinedTraits);
 	if (Template == nullptr)
 	{
 		return;
 	}
 
-	const UWorld* World = OwnerActor.GetWorld();
-	UMassSpawnerSubsystem* SpawnerSystem = UWorld::GetSubsystem<UMassSpawnerSubsystem>(World);
+	UMassSpawnerSubsystem* SpawnerSystem = UWorld::GetSubsystem<UMassSpawnerSubsystem>(&World);
 	check(SpawnerSystem);
 	UMassEntityTemplateRegistry& TemplateRegistry = SpawnerSystem->GetTemplateRegistryInstance();
 
@@ -99,20 +116,19 @@ void FMassEntityConfig::DestroyEntityTemplate(const AActor& OwnerActor, const UO
 	TemplateRegistry.DestroyTemplate(Hash, TemplateID);
 }
 
-const FMassEntityTemplate& FMassEntityConfig::GetEntityTemplateChecked(const AActor& OwnerActor, const UObject& ConfigOwner) const
+const FMassEntityTemplate& FMassEntityConfig::GetEntityTemplateChecked(UWorld& World, const UObject& ConfigOwner) const
 {
 	uint32 Hash;
 	FMassEntityTemplateID TemplateID;
 	TArray<UMassEntityTraitBase*> CombinedTraits;
-	const FMassEntityTemplate* ExistingTemplate = GetEntityTemplateInternal(OwnerActor, ConfigOwner, Hash, TemplateID, CombinedTraits);
+	const FMassEntityTemplate* ExistingTemplate = GetEntityTemplateInternal(World, ConfigOwner, Hash, TemplateID, CombinedTraits);
 	check(ExistingTemplate);
 	return *ExistingTemplate;
 }
 
-const FMassEntityTemplate* FMassEntityConfig::GetEntityTemplateInternal(const AActor& OwnerActor, const UObject& ConfigOwner, uint32& HashOut, FMassEntityTemplateID& TemplateIDOut, TArray<UMassEntityTraitBase*>& CombinedTraitsOut) const
+const FMassEntityTemplate* FMassEntityConfig::GetEntityTemplateInternal(UWorld& World, const UObject& ConfigOwner, uint32& HashOut, FMassEntityTemplateID& TemplateIDOut, TArray<UMassEntityTraitBase*>& CombinedTraitsOut) const
 {
-	const UWorld* World = OwnerActor.GetWorld();
-	UMassSpawnerSubsystem* SpawnerSystem = UWorld::GetSubsystem<UMassSpawnerSubsystem>(World);
+	UMassSpawnerSubsystem* SpawnerSystem = UWorld::GetSubsystem<UMassSpawnerSubsystem>(&World);
 	check(SpawnerSystem);
 	const UMassEntityTemplateRegistry& TemplateRegistry = SpawnerSystem->GetTemplateRegistryInstance();
 
@@ -175,3 +191,18 @@ void FMassEntityConfig::AddTrait(UMassEntityTraitBase& Trait)
 {
 	Traits.Add(&Trait);
 }
+
+bool FMassEntityConfig::ValidateEntityTemplate(UWorld& World, const UObject& ConfigOwner)
+{
+	TArray<const UObject*> Visited;
+	TArray<UMassEntityTraitBase*> CombinedTraits;
+	Visited.Add(&ConfigOwner);
+	GetCombinedTraits(CombinedTraits, Visited, ConfigOwner);
+
+	FMassEntityTemplate Template;
+	FMassEntityTemplateBuildContext BuildContext(Template);
+
+	return BuildContext.BuildFromTraits(CombinedTraits, World);
+}
+
+#undef LOCTEXT_NAMESPACE 
diff --git a/Engine/Plugins/Runtime/MassGameplay/Source/MassSpawner/Private/MassEntityTemplateRegistry.cpp b/Engine/Plugins/Runtime/MassGameplay/Source/MassSpawner/Private/MassEntityTemplateRegistry.cpp
--- a/Engine/Plugins/Runtime/MassGameplay/Source/MassSpawner/Private/MassEntityTemplateRegistry.cpp
+++ b/Engine/Plugins/Runtime/MassGameplay/Source/MassSpawner/Private/MassEntityTemplateRegistry.cpp
@@ -9,6 +9,9 @@
 #include "MassSpawnerSubsystem.h"
 #include "MassEntityTypes.h"
 #include "MassEntityTraitBase.h"
+#include "Logging/MessageLog.h"
+
+#define LOCTEXT_NAMESPACE "Mass"
 
 namespace FTemplateRegistryHelpers
 {
@@ -194,7 +197,7 @@ void UMassEntityTemplateRegistry::DestroyTemplate(const uint32 HashLookup, FMass
 //----------------------------------------------------------------------//
 // FMassEntityTemplateBuildContext 
 //----------------------------------------------------------------------//
-void FMassEntityTemplateBuildContext::BuildFromTraits(TConstArrayView<UMassEntityTraitBase*> Traits, UWorld& World)
+bool FMassEntityTemplateBuildContext::BuildFromTraits(TConstArrayView<UMassEntityTraitBase*> Traits, UWorld& World)
 {
 	TraitAddedTypes.Reset();
 	TraitsDependencies.Reset();
@@ -208,10 +211,10 @@ void FMassEntityTemplateBuildContext::BuildFromTraits(TConstArrayView<UMassEntit
 
 	BuildingTrait = nullptr;
 
-	ValidateBuildContext(World);
+	return ValidateBuildContext(World);
 }
 
-void FMassEntityTemplateBuildContext::ValidateBuildContext(UWorld& World)
+bool FMassEntityTemplateBuildContext::ValidateBuildContext(UWorld& World)
 {
 	// Group same types(key) together
 	TraitAddedTypes.KeySort( [](const UStruct& LHS, const UStruct& RHS) { return LHS.GetName() < RHS.GetName(); } );
@@ -220,6 +223,7 @@ void FMassEntityTemplateBuildContext::ValidateBuildContext(UWorld& World)
 	const UStruct* CurrentStruct = nullptr;
 	const UMassEntityTraitBase* CurrentTrait = nullptr;
 	bool bHeaderOutputed = false;
+	bool bFragmentHasMultipleOwners = false;
 	for (const auto& Pair : TraitAddedTypes)
 	{
 		if (CurrentStruct != Pair.Key)
@@ -238,12 +242,14 @@ void FMassEntityTemplateBuildContext::ValidateBuildContext(UWorld& World)
 				bHeaderOutputed = true;
 			}
 			UE_LOG(LogMass, Error, TEXT("\t\t%s"), *Pair.Value->GetClass()->GetName());
+			bFragmentHasMultipleOwners = true;
 		}
  	}
 
 	// Loop through all the traits dependencies and check if they have been added
 	CurrentTrait = nullptr;
 	bHeaderOutputed = false;
+	bool bMissingFragmentDependencies = false;
 	for (const auto& Dependency : TraitsDependencies)
 	{
 		if (CurrentTrait != Dependency.Get<1>())
@@ -259,6 +265,29 @@ void FMassEntityTemplateBuildContext::ValidateBuildContext(UWorld& World)
 				bHeaderOutputed = true;
 			}
 			UE_LOG(LogMass, Error, TEXT("\t\t%s"), *Dependency.Get<0>()->GetName());
+			bMissingFragmentDependencies = true;
+		}
+	}
+
+#if WITH_UNREAL_DEVELOPER_TOOLS
+	if (bFragmentHasMultipleOwners || bMissingFragmentDependencies)
+	{
+		FMessageLog EditorErrors("LogMass");
+		if (bFragmentHasMultipleOwners)
+		{
+			EditorErrors.Error(LOCTEXT("MassEntityTraitsFragmentOwnershipError", "Some fragments are added by multiple traits and can only be added by one!"));
+			EditorErrors.Notify(LOCTEXT("MassEntityTraitsFragmentOwnershipError", "Some fragments are added by multiple traits and can only be added by one!"));
+		}
+		if (bMissingFragmentDependencies)
+		{
+			EditorErrors.Error(LOCTEXT("MassEntityTraitsMissingFragment", "Some traits are requiring the presence of fragments which are missing!"));
+			EditorErrors.Notify(LOCTEXT("MassEntityTraitsMissingFragment", "Some traits are requiring the presence of fragments which are missing!"));
 		}
+		EditorErrors.Info(FText::FromString(TEXT("See the log for details")));
 	}
+#endif // WITH_UNREAL_DEVELOPER_TOOLS
+
+	return !bFragmentHasMultipleOwners && !bMissingFragmentDependencies;
 }
+
+#undef LOCTEXT_NAMESPACE 
diff --git a/Engine/Plugins/Runtime/MassGameplay/Source/MassSpawner/Private/MassSpawner.cpp b/Engine/Plugins/Runtime/MassGameplay/Source/MassSpawner/Private/MassSpawner.cpp
--- a/Engine/Plugins/Runtime/MassGameplay/Source/MassSpawner/Private/MassSpawner.cpp
+++ b/Engine/Plugins/Runtime/MassGameplay/Source/MassSpawner/Private/MassSpawner.cpp
@@ -264,11 +264,13 @@ void AMassSpawner::DEBUG_Clear()
 
 void AMassSpawner::RegisterEntityTemplates()
 {
+	UWorld* World = GetWorld();
+	check(World);
 	for (FMassSpawnedEntityType& EntityType : EntityTypes)
 	{
 		if (const UMassEntityConfigAsset* EntityConfig = EntityType.GetEntityConfig())
 		{
-			EntityConfig->GetConfig().GetOrCreateEntityTemplate(*this, *EntityConfig);
+			EntityConfig->GetOrCreateEntityTemplate(*World);
 		}
 	}
 }
@@ -420,6 +422,9 @@ void AMassSpawner::SpawnGeneratedEntities(TConstArrayView<FMassEntitySpawnDataGe
 		return;
 	}
 
+	UWorld* World = GetWorld();
+	check(World);
+
 	for (const FMassEntitySpawnDataGeneratorResult& Result : Results)
 	{
 		if (Result.NumEntities <= 0)
@@ -434,12 +439,12 @@ void AMassSpawner::SpawnGeneratedEntities(TConstArrayView<FMassEntitySpawnDataGe
 
 		if (const UMassEntityConfigAsset* EntityConfig = EntityType.GetEntityConfig())
 		{
-			const FMassEntityTemplate* EntityTemplate = EntityConfig->GetConfig().GetOrCreateEntityTemplate(*this, *EntityConfig);
-			if (EntityTemplate && EntityTemplate->IsValid())
+			const FMassEntityTemplate& EntityTemplate = EntityConfig->GetOrCreateEntityTemplate(*World);
+			if (EntityTemplate.IsValid())
 			{
 				FSpawnedEntities& SpawnedEntities = AllSpawnedEntities.AddDefaulted_GetRef();
-				SpawnedEntities.TemplateID = EntityTemplate->GetTemplateID();
-				SpawnerSystem->SpawnEntities(EntityTemplate->GetTemplateID(), Result.NumEntities, Result.SpawnData, Result.SpawnDataProcessor, SpawnedEntities.Entities);
+				SpawnedEntities.TemplateID = EntityTemplate.GetTemplateID();
+				SpawnerSystem->SpawnEntities(EntityTemplate.GetTemplateID(), Result.NumEntities, Result.SpawnData, Result.SpawnDataProcessor, SpawnedEntities.Entities);
 			}
 		}
 	}
@@ -564,11 +569,14 @@ float AMassSpawner::GetSpawningCountScale() const
 
 void AMassSpawner::ClearTemplates()
 {
+	UWorld* World = GetWorld();
+	check(World);
+
 	for (FMassSpawnedEntityType& EntityType : EntityTypes)
 	{
 		if (const UMassEntityConfigAsset* EntityConfig = EntityType.GetEntityConfig())
 		{
-			EntityConfig->GetConfig().DestroyEntityTemplate(*this, *EntityConfig);
+			EntityConfig->DestroyEntityTemplate(*World);
 		}
 	}
 }
diff --git a/Engine/Plugins/Runtime/MassGameplay/Source/MassSpawner/Public/MassEntityConfigAsset.h b/Engine/Plugins/Runtime/MassGameplay/Source/MassSpawner/Public/MassEntityConfigAsset.h
--- a/Engine/Plugins/Runtime/MassGameplay/Source/MassSpawner/Public/MassEntityConfigAsset.h
+++ b/Engine/Plugins/Runtime/MassGameplay/Source/MassSpawner/Public/MassEntityConfigAsset.h
@@ -27,17 +27,17 @@ struct MASSSPAWNER_API FMassEntityConfig
 	FMassEntityConfig(UMassEntityConfigAsset& InParent);
 
 	/** Create entity template based on the features included in this config.
-	 *  @param OwnerActor Actor who is creating the template.
+	 *  @param World in which we are creating the template.
 	 *  @param ConfigOwner Owner of the FMassEntityConfig used for error reporting.
 	 */
-	const FMassEntityTemplate* GetOrCreateEntityTemplate(AActor& OwnerActor, const UObject& ConfigOwner) const;
+	const FMassEntityTemplate& GetOrCreateEntityTemplate(UWorld& World, const UObject& ConfigOwner) const;
 
-	void DestroyEntityTemplate(const AActor& OwnerActor, const UObject& ConfigOwner) const;
+	void DestroyEntityTemplate(UWorld& World, const UObject& ConfigOwner) const;
 
 	/** 
-	 * Fetches the EntityTemplate for given OwnerActor, fails a check if one cannot be found.
+	 * Fetches the EntityTemplate for given World, fails a check if one cannot be found.
 	 */
-	const FMassEntityTemplate& GetEntityTemplateChecked(const AActor& OwnerActor, const UObject& ConfigOwner) const;
+	const FMassEntityTemplate& GetEntityTemplateChecked(UWorld& World, const UObject& ConfigOwner) const;
 
 	/** @return Parent config */
 	const UMassEntityConfigAsset* GetParent() const { return Parent; }
@@ -48,6 +48,9 @@ struct MASSSPAWNER_API FMassEntityConfig
 	/** Adds Trait to the collection of traits hosted by this FMassEntityConfig instance */
 	void AddTrait(UMassEntityTraitBase& Trait);
 
+	/** Validates if the entity template is well built */
+	bool ValidateEntityTemplate(UWorld& World, const UObject& ConfigOwner);
+
 protected:
 	/** Combines traits based on the config hierarchy and returns list of unique traits */
 	void GetCombinedTraits(TArray<UMassEntityTraitBase*>& OutTraits, TArray<const UObject*>& Visited, const UObject& ConfigOwner) const;
@@ -61,7 +64,7 @@ protected:
 	TArray<UMassEntityTraitBase*> Traits;
 
 private:
-	const FMassEntityTemplate* GetEntityTemplateInternal(const AActor& OwnerActor, const UObject& ConfigOwner, uint32& HashOut, FMassEntityTemplateID& TemplateIDOut, TArray<UMassEntityTraitBase*>& CombinedTraitsOut) const;
+	const FMassEntityTemplate* GetEntityTemplateInternal(UWorld& World, const UObject& ConfigOwner, uint32& HashOut, FMassEntityTemplateID& TemplateIDOut, TArray<UMassEntityTraitBase*>& CombinedTraitsOut) const;
 };
 
 /**
@@ -80,6 +83,23 @@ public:
 	/** @return Mutable agent config stored in this asset */
 	FMassEntityConfig& GetMutableConfig() { return Config; }
 
+	const FMassEntityTemplate& GetOrCreateEntityTemplate(UWorld& World) const
+	{
+		return Config.GetOrCreateEntityTemplate(World, *this);
+	}
+
+	void DestroyEntityTemplate(UWorld& World) const
+	{
+		Config.DestroyEntityTemplate(World, *this);
+	}
+
+#if WITH_EDITOR
+
+	UFUNCTION(CallInEditor, Category = "Entity Config")
+	void ValidateEntityConfig();
+
+#endif // WITH_EDITOR
+
 protected:
 	/** The config described in this asset. */
 	UPROPERTY(Category = "Entity Config", EditAnywhere)
diff --git a/Engine/Plugins/Runtime/MassGameplay/Source/MassSpawner/Public/MassEntityTemplateRegistry.h b/Engine/Plugins/Runtime/MassGameplay/Source/MassSpawner/Public/MassEntityTemplateRegistry.h
--- a/Engine/Plugins/Runtime/MassGameplay/Source/MassSpawner/Public/MassEntityTemplateRegistry.h
+++ b/Engine/Plugins/Runtime/MassGameplay/Source/MassSpawner/Public/MassEntityTemplateRegistry.h
@@ -156,10 +156,24 @@ struct FMassEntityTemplateBuildContext
 	//----------------------------------------------------------------------//
 	// Build methods
 	//----------------------------------------------------------------------//
-	void BuildFromTraits(TConstArrayView<UMassEntityTraitBase*> Traits, UWorld& World);
+
+	/**
+	 * Builds context from a list of traits
+	 * @param Traits is the list of all the traits to build an entity
+	 * @param World owning the MassEntitySubsystem for which the entity template is built
+	 * @return true if there were no validation errors
+	 */
+	bool BuildFromTraits(TConstArrayView<UMassEntityTraitBase*> Traits, UWorld& World);
 
 protected:
 
+	/**
+	 * Validate the build context for fragment trait ownership and trait fragment missing dependency
+	 * @param World owning the MassEntitySubsystem for which the entity template is validated against
+	 * @return true if there were no validation errors
+	 */
+	bool ValidateBuildContext(UWorld& World);
+
 	void TypeAdded(const UStruct& Type)
 	{
 		if (ensureMsgf(BuildingTrait, TEXT("Expected to be called within the BuildTemplateFromTrait method")))
@@ -167,7 +181,6 @@ protected:
 			TraitAddedTypes.Add(&Type, BuildingTrait);
 		}
 	}
-	void ValidateBuildContext(UWorld& World);
 
 	const UMassEntityTraitBase* BuildingTrait = nullptr;
 	TMultiMap<const UStruct*, const UMassEntityTraitBase*> TraitAddedTypes;
