Constifying the world parameter for EntityTemplate Building

#rnx

diff --git a/Engine/Plugins/AI/MassAI/Source/MassAIBehavior/Private/MassLookAtTargetTrait.cpp b/Engine/Plugins/AI/MassAI/Source/MassAIBehavior/Private/MassLookAtTargetTrait.cpp
--- a/Engine/Plugins/AI/MassAI/Source/MassAIBehavior/Private/MassLookAtTargetTrait.cpp
+++ b/Engine/Plugins/AI/MassAI/Source/MassAIBehavior/Private/MassLookAtTargetTrait.cpp
@@ -5,7 +5,7 @@
 #include "MassCommonFragments.h"
 #include "MassEntityTemplateRegistry.h"
 
-void UMassLookAtTargetTrait::BuildTemplate(FMassEntityTemplateBuildContext& BuildContext, UWorld& World) const
+void UMassLookAtTargetTrait::BuildTemplate(FMassEntityTemplateBuildContext& BuildContext, const UWorld& World) const
 {
 	BuildContext.AddTag<FMassLookAtTargetTag>();
 	BuildContext.AddFragment<FTransformFragment>();
diff --git a/Engine/Plugins/AI/MassAI/Source/MassAIBehavior/Private/MassLookAtTrait.cpp b/Engine/Plugins/AI/MassAI/Source/MassAIBehavior/Private/MassLookAtTrait.cpp
--- a/Engine/Plugins/AI/MassAI/Source/MassAIBehavior/Private/MassLookAtTrait.cpp
+++ b/Engine/Plugins/AI/MassAI/Source/MassAIBehavior/Private/MassLookAtTrait.cpp
@@ -4,7 +4,7 @@
 #include "MassLookAtFragments.h"
 #include "MassEntityTemplateRegistry.h"
 
-void UMassLookAtTrait::BuildTemplate(FMassEntityTemplateBuildContext& BuildContext, UWorld& World) const
+void UMassLookAtTrait::BuildTemplate(FMassEntityTemplateBuildContext& BuildContext, const UWorld& World) const
 {
 	BuildContext.AddFragment<FMassLookAtFragment>();
 	BuildContext.AddFragment<FMassLookAtTrajectoryFragment>();
diff --git a/Engine/Plugins/AI/MassAI/Source/MassAIBehavior/Private/MassStateTreeTrait.cpp b/Engine/Plugins/AI/MassAI/Source/MassAIBehavior/Private/MassStateTreeTrait.cpp
--- a/Engine/Plugins/AI/MassAI/Source/MassAIBehavior/Private/MassStateTreeTrait.cpp
+++ b/Engine/Plugins/AI/MassAI/Source/MassAIBehavior/Private/MassStateTreeTrait.cpp
@@ -9,7 +9,7 @@
 #include "MassEntityTemplateRegistry.h"
 #include "Engine/World.h"
 
-void UMassStateTreeTrait::BuildTemplate(FMassEntityTemplateBuildContext& BuildContext, UWorld& World) const
+void UMassStateTreeTrait::BuildTemplate(FMassEntityTemplateBuildContext& BuildContext, const UWorld& World) const
 {
 	UMassEntitySubsystem* EntitySubsystem = UWorld::GetSubsystem<UMassEntitySubsystem>(&World);
 	check(EntitySubsystem);
@@ -41,7 +41,7 @@ void UMassStateTreeTrait::BuildTemplate(FMassEntityTemplateBuildContext& BuildCo
 	BuildContext.AddFragment<FMassStateTreeInstanceFragment>();
 }
 
-void UMassStateTreeTrait::ValidateTemplate(FMassEntityTemplateBuildContext& BuildContext, UWorld& World) const
+void UMassStateTreeTrait::ValidateTemplate(FMassEntityTemplateBuildContext& BuildContext, const UWorld& World) const
 {
 	UMassStateTreeSubsystem* MassStateTreeSubsystem = World.GetSubsystem<UMassStateTreeSubsystem>();
 	if (!MassStateTreeSubsystem)
diff --git a/Engine/Plugins/AI/MassAI/Source/MassAIBehavior/Private/MassZoneGraphAnnotationTrait.cpp b/Engine/Plugins/AI/MassAI/Source/MassAIBehavior/Private/MassZoneGraphAnnotationTrait.cpp
--- a/Engine/Plugins/AI/MassAI/Source/MassAIBehavior/Private/MassZoneGraphAnnotationTrait.cpp
+++ b/Engine/Plugins/AI/MassAI/Source/MassAIBehavior/Private/MassZoneGraphAnnotationTrait.cpp
@@ -5,7 +5,7 @@
 #include "Engine/World.h"
 #include "MassZoneGraphAnnotationFragments.h"
 
-void UMassZoneGraphAnnotationTrait::BuildTemplate(FMassEntityTemplateBuildContext& BuildContext, UWorld& World) const
+void UMassZoneGraphAnnotationTrait::BuildTemplate(FMassEntityTemplateBuildContext& BuildContext, const UWorld& World) const
 {
 	BuildContext.AddFragment<FMassZoneGraphAnnotationFragment>();
 	BuildContext.AddChunkFragment<FMassZoneGraphAnnotationVariableTickChunkFragment>();
diff --git a/Engine/Plugins/AI/MassAI/Source/MassAIBehavior/Public/MassLookAtTargetTrait.h b/Engine/Plugins/AI/MassAI/Source/MassAIBehavior/Public/MassLookAtTargetTrait.h
--- a/Engine/Plugins/AI/MassAI/Source/MassAIBehavior/Public/MassLookAtTargetTrait.h
+++ b/Engine/Plugins/AI/MassAI/Source/MassAIBehavior/Public/MassLookAtTargetTrait.h
@@ -13,5 +13,5 @@ class UMassLookAtTargetTrait : public UMassEntityTraitBase
 
 protected:
 
-	virtual void BuildTemplate(FMassEntityTemplateBuildContext& BuildContext, UWorld& World) const override;
+	virtual void BuildTemplate(FMassEntityTemplateBuildContext& BuildContext, const UWorld& World) const override;
 };
diff --git a/Engine/Plugins/AI/MassAI/Source/MassAIBehavior/Public/MassLookAtTrait.h b/Engine/Plugins/AI/MassAI/Source/MassAIBehavior/Public/MassLookAtTrait.h
--- a/Engine/Plugins/AI/MassAI/Source/MassAIBehavior/Public/MassLookAtTrait.h
+++ b/Engine/Plugins/AI/MassAI/Source/MassAIBehavior/Public/MassLookAtTrait.h
@@ -13,5 +13,5 @@ class UMassLookAtTrait : public UMassEntityTraitBase
 
 protected:
 
-	virtual void BuildTemplate(FMassEntityTemplateBuildContext& BuildContext, UWorld& World) const override;
+	virtual void BuildTemplate(FMassEntityTemplateBuildContext& BuildContext, const UWorld& World) const override;
 };
diff --git a/Engine/Plugins/AI/MassAI/Source/MassAIBehavior/Public/MassStateTreeTrait.h b/Engine/Plugins/AI/MassAI/Source/MassAIBehavior/Public/MassStateTreeTrait.h
--- a/Engine/Plugins/AI/MassAI/Source/MassAIBehavior/Public/MassStateTreeTrait.h
+++ b/Engine/Plugins/AI/MassAI/Source/MassAIBehavior/Public/MassStateTreeTrait.h
@@ -17,9 +17,9 @@ class MASSAIBEHAVIOR_API UMassStateTreeTrait : public UMassEntityTraitBase
 
 protected:
 
-	virtual void BuildTemplate(FMassEntityTemplateBuildContext& BuildContext, UWorld& World) const override;
+	virtual void BuildTemplate(FMassEntityTemplateBuildContext& BuildContext, const UWorld& World) const override;
 
-	virtual void ValidateTemplate(FMassEntityTemplateBuildContext& BuildContext, UWorld& World) const override;
+	virtual void ValidateTemplate(FMassEntityTemplateBuildContext& BuildContext, const UWorld& World) const override;
 
 	UPROPERTY(Category="StateTree", EditAnywhere, meta=(RequiredAssetDataTags="Schema=MassStateTreeSchema"))
 	UStateTree* StateTree;
diff --git a/Engine/Plugins/AI/MassAI/Source/MassAIBehavior/Public/MassZoneGraphAnnotationTrait.h b/Engine/Plugins/AI/MassAI/Source/MassAIBehavior/Public/MassZoneGraphAnnotationTrait.h
--- a/Engine/Plugins/AI/MassAI/Source/MassAIBehavior/Public/MassZoneGraphAnnotationTrait.h
+++ b/Engine/Plugins/AI/MassAI/Source/MassAIBehavior/Public/MassZoneGraphAnnotationTrait.h
@@ -11,5 +11,5 @@ class MASSAIBEHAVIOR_API UMassZoneGraphAnnotationTrait : public UMassEntityTrait
 	GENERATED_BODY()
 
 protected:
-	virtual void BuildTemplate(FMassEntityTemplateBuildContext& BuildContext, UWorld& World) const override;
+	virtual void BuildTemplate(FMassEntityTemplateBuildContext& BuildContext, const UWorld& World) const override;
 };
diff --git a/Engine/Plugins/AI/MassAI/Source/MassNavigation/Private/MassAvoidanceTrait.cpp b/Engine/Plugins/AI/MassAI/Source/MassNavigation/Private/MassAvoidanceTrait.cpp
--- a/Engine/Plugins/AI/MassAI/Source/MassNavigation/Private/MassAvoidanceTrait.cpp
+++ b/Engine/Plugins/AI/MassAI/Source/MassNavigation/Private/MassAvoidanceTrait.cpp
@@ -8,7 +8,7 @@
 #include "MassNavigationFragments.h"
 #include "Engine/World.h"
 
-void UMassObstacleAvoidanceTrait::BuildTemplate(FMassEntityTemplateBuildContext& BuildContext, UWorld& World) const
+void UMassObstacleAvoidanceTrait::BuildTemplate(FMassEntityTemplateBuildContext& BuildContext, const UWorld& World) const
 {
 	UMassEntitySubsystem* EntitySubsystem = UWorld::GetSubsystem<UMassEntitySubsystem>(&World);
 	check(EntitySubsystem);
diff --git a/Engine/Plugins/AI/MassAI/Source/MassNavigation/Private/MassNavigationObstacleTrait.cpp b/Engine/Plugins/AI/MassAI/Source/MassNavigation/Private/MassNavigationObstacleTrait.cpp
--- a/Engine/Plugins/AI/MassAI/Source/MassNavigation/Private/MassNavigationObstacleTrait.cpp
+++ b/Engine/Plugins/AI/MassAI/Source/MassNavigation/Private/MassNavigationObstacleTrait.cpp
@@ -4,7 +4,7 @@
 #include "MassNavigationFragments.h"
 #include "MassCommonFragments.h"
 
-void UMassNavigationObstacleTrait::BuildTemplate(FMassEntityTemplateBuildContext& BuildContext, UWorld& World) const
+void UMassNavigationObstacleTrait::BuildTemplate(FMassEntityTemplateBuildContext& BuildContext, const UWorld& World) const
 {
 	BuildContext.RequireFragment<FAgentRadiusFragment>();
 
diff --git a/Engine/Plugins/AI/MassAI/Source/MassNavigation/Private/MassSmoothOrientationTrait.cpp b/Engine/Plugins/AI/MassAI/Source/MassNavigation/Private/MassSmoothOrientationTrait.cpp
--- a/Engine/Plugins/AI/MassAI/Source/MassNavigation/Private/MassSmoothOrientationTrait.cpp
+++ b/Engine/Plugins/AI/MassAI/Source/MassNavigation/Private/MassSmoothOrientationTrait.cpp
@@ -6,7 +6,7 @@
 #include "MassNavigationFragments.h"
 #include "Engine/World.h"
 
-void UMassSmoothOrientationTrait::BuildTemplate(FMassEntityTemplateBuildContext& BuildContext, UWorld& World) const
+void UMassSmoothOrientationTrait::BuildTemplate(FMassEntityTemplateBuildContext& BuildContext, const UWorld& World) const
 {
 	UMassEntitySubsystem* EntitySubsystem = UWorld::GetSubsystem<UMassEntitySubsystem>(&World);
 	check(EntitySubsystem);
diff --git a/Engine/Plugins/AI/MassAI/Source/MassNavigation/Private/MassSteeringTrait.cpp b/Engine/Plugins/AI/MassAI/Source/MassNavigation/Private/MassSteeringTrait.cpp
--- a/Engine/Plugins/AI/MassAI/Source/MassNavigation/Private/MassSteeringTrait.cpp
+++ b/Engine/Plugins/AI/MassAI/Source/MassNavigation/Private/MassSteeringTrait.cpp
@@ -8,7 +8,7 @@
 #include "Steering/MassSteeringFragments.h"
 #include "Engine/World.h"
 
-void UMassSteeringTrait::BuildTemplate(FMassEntityTemplateBuildContext& BuildContext, UWorld& World) const
+void UMassSteeringTrait::BuildTemplate(FMassEntityTemplateBuildContext& BuildContext, const UWorld& World) const
 {
 	UMassEntitySubsystem* EntitySubsystem = UWorld::GetSubsystem<UMassEntitySubsystem>(&World);
 	check(EntitySubsystem);
diff --git a/Engine/Plugins/AI/MassAI/Source/MassNavigation/Public/Avoidance/MassAvoidanceTrait.h b/Engine/Plugins/AI/MassAI/Source/MassNavigation/Public/Avoidance/MassAvoidanceTrait.h
--- a/Engine/Plugins/AI/MassAI/Source/MassNavigation/Public/Avoidance/MassAvoidanceTrait.h
+++ b/Engine/Plugins/AI/MassAI/Source/MassNavigation/Public/Avoidance/MassAvoidanceTrait.h
@@ -13,7 +13,7 @@ class MASSNAVIGATION_API UMassObstacleAvoidanceTrait : public UMassEntityTraitBa
 
 protected:
 	
-	virtual void BuildTemplate(FMassEntityTemplateBuildContext& BuildContext, UWorld& World) const override;
+	virtual void BuildTemplate(FMassEntityTemplateBuildContext& BuildContext, const UWorld& World) const override;
 
 	UPROPERTY(EditAnywhere, Category="")
 	FMassMovingAvoidanceParameters MovingParameters;
diff --git a/Engine/Plugins/AI/MassAI/Source/MassNavigation/Public/Avoidance/MassNavigationObstacleTrait.h b/Engine/Plugins/AI/MassAI/Source/MassNavigation/Public/Avoidance/MassNavigationObstacleTrait.h
--- a/Engine/Plugins/AI/MassAI/Source/MassNavigation/Public/Avoidance/MassNavigationObstacleTrait.h
+++ b/Engine/Plugins/AI/MassAI/Source/MassNavigation/Public/Avoidance/MassNavigationObstacleTrait.h
@@ -11,6 +11,6 @@ class MASSNAVIGATION_API UMassNavigationObstacleTrait : public UMassEntityTraitB
 	GENERATED_BODY()
 
 protected:
-	virtual void BuildTemplate(FMassEntityTemplateBuildContext& BuildContext, UWorld& World) const override;
+	virtual void BuildTemplate(FMassEntityTemplateBuildContext& BuildContext, const UWorld& World) const override;
 
 };
\ No newline at end of file
diff --git a/Engine/Plugins/AI/MassAI/Source/MassNavigation/Public/SmoothOrientation/MassSmoothOrientationTrait.h b/Engine/Plugins/AI/MassAI/Source/MassNavigation/Public/SmoothOrientation/MassSmoothOrientationTrait.h
--- a/Engine/Plugins/AI/MassAI/Source/MassNavigation/Public/SmoothOrientation/MassSmoothOrientationTrait.h
+++ b/Engine/Plugins/AI/MassAI/Source/MassNavigation/Public/SmoothOrientation/MassSmoothOrientationTrait.h
@@ -12,7 +12,7 @@ class MASSNAVIGATION_API UMassSmoothOrientationTrait : public UMassEntityTraitBa
 	GENERATED_BODY()
 
 protected:
-	virtual void BuildTemplate(FMassEntityTemplateBuildContext& BuildContext, UWorld& World) const override;
+	virtual void BuildTemplate(FMassEntityTemplateBuildContext& BuildContext, const UWorld& World) const override;
 
 	UPROPERTY(EditAnywhere, Category="")
 	FMassSmoothOrientationParameters Orientation;
diff --git a/Engine/Plugins/AI/MassAI/Source/MassNavigation/Public/Steering/MassSteeringTrait.h b/Engine/Plugins/AI/MassAI/Source/MassNavigation/Public/Steering/MassSteeringTrait.h
--- a/Engine/Plugins/AI/MassAI/Source/MassNavigation/Public/Steering/MassSteeringTrait.h
+++ b/Engine/Plugins/AI/MassAI/Source/MassNavigation/Public/Steering/MassSteeringTrait.h
@@ -13,7 +13,7 @@ class MASSNAVIGATION_API UMassSteeringTrait : public UMassEntityTraitBase
 	GENERATED_BODY()
 
 protected:
-	virtual void BuildTemplate(FMassEntityTemplateBuildContext& BuildContext, UWorld& World) const override;
+	virtual void BuildTemplate(FMassEntityTemplateBuildContext& BuildContext, const UWorld& World) const override;
 
 	UPROPERTY(Category="Steering", EditAnywhere, meta=(EditInline))
 	FMassMovingSteeringParameters MovingSteering;
diff --git a/Engine/Plugins/AI/MassAI/Source/MassZoneGraphNavigation/Private/MassZoneGraphNavigationTrait.cpp b/Engine/Plugins/AI/MassAI/Source/MassZoneGraphNavigation/Private/MassZoneGraphNavigationTrait.cpp
--- a/Engine/Plugins/AI/MassAI/Source/MassZoneGraphNavigation/Private/MassZoneGraphNavigationTrait.cpp
+++ b/Engine/Plugins/AI/MassAI/Source/MassZoneGraphNavigation/Private/MassZoneGraphNavigationTrait.cpp
@@ -7,7 +7,7 @@
 #include "MassZoneGraphNavigationFragments.h"
 #include "Engine/World.h"
 
-void UMassZoneGraphNavigationTrait::BuildTemplate(FMassEntityTemplateBuildContext& BuildContext, UWorld& World) const
+void UMassZoneGraphNavigationTrait::BuildTemplate(FMassEntityTemplateBuildContext& BuildContext, const UWorld& World) const
 {
 	UMassEntitySubsystem* EntitySubsystem = UWorld::GetSubsystem<UMassEntitySubsystem>(&World);
 	check(EntitySubsystem);
diff --git a/Engine/Plugins/AI/MassAI/Source/MassZoneGraphNavigation/Public/MassZoneGraphNavigationTrait.h b/Engine/Plugins/AI/MassAI/Source/MassZoneGraphNavigation/Public/MassZoneGraphNavigationTrait.h
--- a/Engine/Plugins/AI/MassAI/Source/MassZoneGraphNavigation/Public/MassZoneGraphNavigationTrait.h
+++ b/Engine/Plugins/AI/MassAI/Source/MassZoneGraphNavigation/Public/MassZoneGraphNavigationTrait.h
@@ -13,7 +13,7 @@ class MASSZONEGRAPHNAVIGATION_API UMassZoneGraphNavigationTrait : public UMassEn
 	GENERATED_BODY()
 
 protected:
-	virtual void BuildTemplate(FMassEntityTemplateBuildContext& BuildContext, UWorld& World) const override;
+	virtual void BuildTemplate(FMassEntityTemplateBuildContext& BuildContext, const UWorld& World) const override;
 
 	UPROPERTY(Category="Movement", EditAnywhere)
 	FMassZoneGraphNavigationParameters NavigationParameters;
diff --git a/Engine/Plugins/AI/MassCrowd/Source/MassCrowd/Private/MassCrowdMemberTrait.cpp b/Engine/Plugins/AI/MassCrowd/Source/MassCrowd/Private/MassCrowdMemberTrait.cpp
--- a/Engine/Plugins/AI/MassCrowd/Source/MassCrowd/Private/MassCrowdMemberTrait.cpp
+++ b/Engine/Plugins/AI/MassCrowd/Source/MassCrowd/Private/MassCrowdMemberTrait.cpp
@@ -4,7 +4,7 @@
 #include "MassCrowdFragments.h"
 #include "MassEntityTemplateRegistry.h"
 
-void UMassCrowdMemberTrait::BuildTemplate(FMassEntityTemplateBuildContext& BuildContext, UWorld& World) const
+void UMassCrowdMemberTrait::BuildTemplate(FMassEntityTemplateBuildContext& BuildContext, const UWorld& World) const
 {
 	BuildContext.AddTag<FMassCrowdTag>();
 	BuildContext.AddFragment<FMassCrowdLaneTrackingFragment>();	
diff --git a/Engine/Plugins/AI/MassCrowd/Source/MassCrowd/Private/MassCrowdServerRepresentationTrait.cpp b/Engine/Plugins/AI/MassCrowd/Source/MassCrowd/Private/MassCrowdServerRepresentationTrait.cpp
--- a/Engine/Plugins/AI/MassCrowd/Source/MassCrowd/Private/MassCrowdServerRepresentationTrait.cpp
+++ b/Engine/Plugins/AI/MassCrowd/Source/MassCrowd/Private/MassCrowdServerRepresentationTrait.cpp
@@ -24,7 +24,7 @@ UMassCrowdServerRepresentationTrait::UMassCrowdServerRepresentationTrait()
 	Params.NotVisibleUpdateRate = 0.5f;
 }
 
-void UMassCrowdServerRepresentationTrait::BuildTemplate(FMassEntityTemplateBuildContext& BuildContext, UWorld& World) const
+void UMassCrowdServerRepresentationTrait::BuildTemplate(FMassEntityTemplateBuildContext& BuildContext, const UWorld& World) const
 {
 	// This should only be ran on NM_DedicatedServer network mode
 	if (!World.IsNetMode(NM_DedicatedServer))
diff --git a/Engine/Plugins/AI/MassCrowd/Source/MassCrowd/Public/MassCrowdMemberTrait.h b/Engine/Plugins/AI/MassCrowd/Source/MassCrowd/Public/MassCrowdMemberTrait.h
--- a/Engine/Plugins/AI/MassCrowd/Source/MassCrowd/Public/MassCrowdMemberTrait.h
+++ b/Engine/Plugins/AI/MassCrowd/Source/MassCrowd/Public/MassCrowdMemberTrait.h
@@ -14,5 +14,5 @@ class MASSCROWD_API UMassCrowdMemberTrait : public UMassEntityTraitBase
 	GENERATED_BODY()
 
 protected:
-	virtual void BuildTemplate(FMassEntityTemplateBuildContext& BuildContext, UWorld& World) const override;
+	virtual void BuildTemplate(FMassEntityTemplateBuildContext& BuildContext, const UWorld& World) const override;
 };
diff --git a/Engine/Plugins/AI/MassCrowd/Source/MassCrowd/Public/MassCrowdServerRepresentationTrait.h b/Engine/Plugins/AI/MassCrowd/Source/MassCrowd/Public/MassCrowdServerRepresentationTrait.h
--- a/Engine/Plugins/AI/MassCrowd/Source/MassCrowd/Public/MassCrowdServerRepresentationTrait.h
+++ b/Engine/Plugins/AI/MassCrowd/Source/MassCrowd/Public/MassCrowdServerRepresentationTrait.h
@@ -19,7 +19,7 @@ class MASSCROWD_API UMassCrowdServerRepresentationTrait : public UMassEntityTrai
 
 protected:
 
-	virtual void BuildTemplate(FMassEntityTemplateBuildContext& BuildContext, UWorld& World) const override;
+	virtual void BuildTemplate(FMassEntityTemplateBuildContext& BuildContext, const UWorld& World) const override;
 
 	/** Actor class of this agent when spawned on server */
 	UPROPERTY(EditAnywhere, Category = "Mass|Visual")
diff --git a/Engine/Plugins/Runtime/MassGameplay/Source/MassActors/Private/MassAgentTraits.cpp b/Engine/Plugins/Runtime/MassGameplay/Source/MassActors/Private/MassAgentTraits.cpp
--- a/Engine/Plugins/Runtime/MassGameplay/Source/MassActors/Private/MassAgentTraits.cpp
+++ b/Engine/Plugins/Runtime/MassGameplay/Source/MassActors/Private/MassAgentTraits.cpp
@@ -39,7 +39,7 @@ namespace FMassAgentTraitsHelper
 //----------------------------------------------------------------------//
 //  UMassAgentCapsuleCollisionSyncTrait
 //----------------------------------------------------------------------//
-void UMassAgentCapsuleCollisionSyncTrait::BuildTemplate(FMassEntityTemplateBuildContext& BuildContext, UWorld& World) const
+void UMassAgentCapsuleCollisionSyncTrait::BuildTemplate(FMassEntityTemplateBuildContext& BuildContext, const UWorld& World) const
 {
 	BuildContext.AddFragment<FCapsuleComponentWrapperFragment>();
 	BuildContext.AddFragment<FAgentRadiusFragment>();
@@ -83,7 +83,7 @@ void UMassAgentCapsuleCollisionSyncTrait::BuildTemplate(FMassEntityTemplateBuild
 //----------------------------------------------------------------------//
 //  UMassAgentMovementSyncTrait
 //----------------------------------------------------------------------//
-void UMassAgentMovementSyncTrait::BuildTemplate(FMassEntityTemplateBuildContext& BuildContext, UWorld& World) const
+void UMassAgentMovementSyncTrait::BuildTemplate(FMassEntityTemplateBuildContext& BuildContext, const UWorld& World) const
 {
 	BuildContext.AddFragment<FCharacterMovementComponentWrapperFragment>();
 	BuildContext.AddFragment<FMassVelocityFragment>();
@@ -126,7 +126,7 @@ void UMassAgentMovementSyncTrait::BuildTemplate(FMassEntityTemplateBuildContext&
 //----------------------------------------------------------------------//
 //  UMassAgentOrientationSyncTrait
 //----------------------------------------------------------------------//
-void UMassAgentOrientationSyncTrait::BuildTemplate(FMassEntityTemplateBuildContext& BuildContext, UWorld& World) const
+void UMassAgentOrientationSyncTrait::BuildTemplate(FMassEntityTemplateBuildContext& BuildContext, const UWorld& World) const
 {
 	BuildContext.RequireFragment<FCharacterMovementComponentWrapperFragment>();
 
@@ -144,7 +144,7 @@ void UMassAgentOrientationSyncTrait::BuildTemplate(FMassEntityTemplateBuildConte
 //----------------------------------------------------------------------//
 //  UMassAgentFeetLocationSyncTrait
 //----------------------------------------------------------------------//
-void UMassAgentFeetLocationSyncTrait::BuildTemplate(FMassEntityTemplateBuildContext& BuildContext, UWorld& World) const
+void UMassAgentFeetLocationSyncTrait::BuildTemplate(FMassEntityTemplateBuildContext& BuildContext, const UWorld& World) const
 {
 	BuildContext.AddFragment<FMassSceneComponentWrapperFragment>();
 	BuildContext.AddFragment<FTransformFragment>();
diff --git a/Engine/Plugins/Runtime/MassGameplay/Source/MassActors/Public/MassAgentTraits.h b/Engine/Plugins/Runtime/MassGameplay/Source/MassActors/Public/MassAgentTraits.h
--- a/Engine/Plugins/Runtime/MassGameplay/Source/MassActors/Public/MassAgentTraits.h
+++ b/Engine/Plugins/Runtime/MassGameplay/Source/MassActors/Public/MassAgentTraits.h
@@ -31,7 +31,7 @@ class MASSACTORS_API UMassAgentCapsuleCollisionSyncTrait : public UMassAgentSync
 	GENERATED_BODY()
 
 protected:
-	virtual void BuildTemplate(FMassEntityTemplateBuildContext& BuildContext, UWorld& World) const override;
+	virtual void BuildTemplate(FMassEntityTemplateBuildContext& BuildContext, const UWorld& World) const override;
 
 	UPROPERTY(EditAnywhere, Category = Mass)
 	bool bSyncTransform = true;
@@ -43,7 +43,7 @@ class MASSACTORS_API UMassAgentMovementSyncTrait : public UMassAgentSyncTrait
 	GENERATED_BODY()
 
 protected:
-	virtual void BuildTemplate(FMassEntityTemplateBuildContext& BuildContext, UWorld& World) const override;
+	virtual void BuildTemplate(FMassEntityTemplateBuildContext& BuildContext, const UWorld& World) const override;
 };
 
 UCLASS(BlueprintType, EditInlineNew, CollapseCategories, meta = (DisplayName = "Agent Orientation Sync"))
@@ -52,7 +52,7 @@ class MASSACTORS_API UMassAgentOrientationSyncTrait : public UMassAgentSyncTrait
 	GENERATED_BODY()
 
 	protected:
-	virtual void BuildTemplate(FMassEntityTemplateBuildContext& BuildContext, UWorld& World) const override;
+	virtual void BuildTemplate(FMassEntityTemplateBuildContext& BuildContext, const UWorld& World) const override;
 };
 
 UCLASS(BlueprintType, EditInlineNew, CollapseCategories, meta = (DisplayName = "Agent Feet Location Sync"))
@@ -61,6 +61,6 @@ class MASSACTORS_API UMassAgentFeetLocationSyncTrait : public UMassAgentSyncTrai
 	GENERATED_BODY()
 
 protected:
-	virtual void BuildTemplate(FMassEntityTemplateBuildContext& BuildContext, UWorld& World) const override;
+	virtual void BuildTemplate(FMassEntityTemplateBuildContext& BuildContext, const UWorld& World) const override;
 };
 
diff --git a/Engine/Plugins/Runtime/MassGameplay/Source/MassGameplayDebug/Private/MassDebugVisualizationTrait.cpp b/Engine/Plugins/Runtime/MassGameplay/Source/MassGameplayDebug/Private/MassDebugVisualizationTrait.cpp
--- a/Engine/Plugins/Runtime/MassGameplay/Source/MassGameplayDebug/Private/MassDebugVisualizationTrait.cpp
+++ b/Engine/Plugins/Runtime/MassGameplay/Source/MassGameplayDebug/Private/MassDebugVisualizationTrait.cpp
@@ -7,7 +7,7 @@
 #include "MassCommonFragments.h"
 #include "Engine/World.h"
 
-void UMassDebugVisualizationTrait::BuildTemplate(FMassEntityTemplateBuildContext& BuildContext, UWorld& World) const
+void UMassDebugVisualizationTrait::BuildTemplate(FMassEntityTemplateBuildContext& BuildContext, const UWorld& World) const
 {
 #if WITH_EDITORONLY_DATA
 	const UStaticMesh* const DebugMesh = DebugShape.Mesh;
diff --git a/Engine/Plugins/Runtime/MassGameplay/Source/MassGameplayDebug/Public/MassDebugVisualizationTrait.h b/Engine/Plugins/Runtime/MassGameplay/Source/MassGameplayDebug/Public/MassDebugVisualizationTrait.h
--- a/Engine/Plugins/Runtime/MassGameplay/Source/MassGameplayDebug/Public/MassDebugVisualizationTrait.h
+++ b/Engine/Plugins/Runtime/MassGameplay/Source/MassGameplayDebug/Public/MassDebugVisualizationTrait.h
@@ -14,7 +14,7 @@ class MASSGAMEPLAYDEBUG_API UMassDebugVisualizationTrait : public UMassEntityTra
 
 protected:
 
-	virtual void BuildTemplate(FMassEntityTemplateBuildContext& BuildContext, UWorld& World) const override;
+	virtual void BuildTemplate(FMassEntityTemplateBuildContext& BuildContext, const UWorld& World) const override;
 
 #if WITH_EDITORONLY_DATA
 	UPROPERTY(EditAnywhere, Category = "Mass|Spawn")
diff --git a/Engine/Plugins/Runtime/MassGameplay/Source/MassLOD/Private/MassLODTrait.cpp b/Engine/Plugins/Runtime/MassGameplay/Source/MassLOD/Private/MassLODTrait.cpp
--- a/Engine/Plugins/Runtime/MassGameplay/Source/MassLOD/Private/MassLODTrait.cpp
+++ b/Engine/Plugins/Runtime/MassGameplay/Source/MassLOD/Private/MassLODTrait.cpp
@@ -7,14 +7,14 @@
 #include "StructUtilsTypes.h"
 #include "MassLODFragments.h"
 
-void UMassLODCollectorTrait::BuildTemplate(FMassEntityTemplateBuildContext& BuildContext, UWorld& World) const
+void UMassLODCollectorTrait::BuildTemplate(FMassEntityTemplateBuildContext& BuildContext, const UWorld& World) const
 {
 	BuildContext.AddFragment<FMassViewerInfoFragment>();
 	BuildContext.AddTag<FMassCollectLODViewerInfoTag>();
 	BuildContext.RequireFragment<FTransformFragment>();
 }
 
-void UMassSimulationLODTrait::BuildTemplate(FMassEntityTemplateBuildContext& BuildContext, UWorld& World) const
+void UMassSimulationLODTrait::BuildTemplate(FMassEntityTemplateBuildContext& BuildContext, const UWorld& World) const
 {
 	BuildContext.RequireFragment<FMassViewerInfoFragment>();
 	BuildContext.RequireFragment<FTransformFragment>();
diff --git a/Engine/Plugins/Runtime/MassGameplay/Source/MassLOD/Public/MassLODTrait.h b/Engine/Plugins/Runtime/MassGameplay/Source/MassLOD/Public/MassLODTrait.h
--- a/Engine/Plugins/Runtime/MassGameplay/Source/MassLOD/Public/MassLODTrait.h
+++ b/Engine/Plugins/Runtime/MassGameplay/Source/MassLOD/Public/MassLODTrait.h
@@ -13,7 +13,7 @@ class MASSLOD_API UMassLODCollectorTrait : public UMassEntityTraitBase
 	GENERATED_BODY()
 
 public:
-	virtual void BuildTemplate(FMassEntityTemplateBuildContext& BuildContext, UWorld& World) const override;
+	virtual void BuildTemplate(FMassEntityTemplateBuildContext& BuildContext, const UWorld& World) const override;
 };
 
 UCLASS(meta = (DisplayName = "SimulationLOD"))
@@ -31,5 +31,5 @@ class MASSLOD_API UMassSimulationLODTrait : public UMassEntityTraitBase
 	FMassSimulationVariableTickParameters VariableTickParams;
 
 protected:
-	virtual void BuildTemplate(FMassEntityTemplateBuildContext& BuildContext, UWorld& World) const override;
+	virtual void BuildTemplate(FMassEntityTemplateBuildContext& BuildContext, const UWorld& World) const override;
 };
diff --git a/Engine/Plugins/Runtime/MassGameplay/Source/MassMovement/Private/MassMovementTrait.cpp b/Engine/Plugins/Runtime/MassGameplay/Source/MassMovement/Private/MassMovementTrait.cpp
--- a/Engine/Plugins/Runtime/MassGameplay/Source/MassMovement/Private/MassMovementTrait.cpp
+++ b/Engine/Plugins/Runtime/MassGameplay/Source/MassMovement/Private/MassMovementTrait.cpp
@@ -6,7 +6,7 @@
 #include "MassMovementTypes.h"
 #include "Engine/World.h"
 
-void UMassMovementTrait::BuildTemplate(FMassEntityTemplateBuildContext& BuildContext, UWorld& World) const
+void UMassMovementTrait::BuildTemplate(FMassEntityTemplateBuildContext& BuildContext, const UWorld& World) const
 {
 	UMassEntitySubsystem* EntitySubsystem = UWorld::GetSubsystem<UMassEntitySubsystem>(&World);
 	check(EntitySubsystem);
diff --git a/Engine/Plugins/Runtime/MassGameplay/Source/MassMovement/Private/MassSimpleMovementTrait.cpp b/Engine/Plugins/Runtime/MassGameplay/Source/MassMovement/Private/MassSimpleMovementTrait.cpp
--- a/Engine/Plugins/Runtime/MassGameplay/Source/MassMovement/Private/MassSimpleMovementTrait.cpp
+++ b/Engine/Plugins/Runtime/MassGameplay/Source/MassMovement/Private/MassSimpleMovementTrait.cpp
@@ -11,7 +11,7 @@
 //----------------------------------------------------------------------//
 //  UMassSimpleMovementTrait
 //----------------------------------------------------------------------//
-void UMassSimpleMovementTrait::BuildTemplate(FMassEntityTemplateBuildContext& BuildContext, UWorld& World) const
+void UMassSimpleMovementTrait::BuildTemplate(FMassEntityTemplateBuildContext& BuildContext, const UWorld& World) const
 {
 	BuildContext.AddFragment<FTransformFragment>();
 	BuildContext.AddFragment<FMassVelocityFragment>();
diff --git a/Engine/Plugins/Runtime/MassGameplay/Source/MassMovement/Private/MassVelocityRandomizerTrait.cpp b/Engine/Plugins/Runtime/MassGameplay/Source/MassMovement/Private/MassVelocityRandomizerTrait.cpp
--- a/Engine/Plugins/Runtime/MassGameplay/Source/MassMovement/Private/MassVelocityRandomizerTrait.cpp
+++ b/Engine/Plugins/Runtime/MassGameplay/Source/MassMovement/Private/MassVelocityRandomizerTrait.cpp
@@ -9,7 +9,7 @@
 //----------------------------------------------------------------------//
 //  UMassVelocityRandomizerTrait
 //----------------------------------------------------------------------//
-void UMassVelocityRandomizerTrait::BuildTemplate(FMassEntityTemplateBuildContext& BuildContext, UWorld& World) const
+void UMassVelocityRandomizerTrait::BuildTemplate(FMassEntityTemplateBuildContext& BuildContext, const UWorld& World) const
 {
 	FMassVelocityFragment& VelocityTemplate = BuildContext.AddFragment_GetRef<FMassVelocityFragment>();
 	// This is a small @hack to support sending parameters to initializer.
diff --git a/Engine/Plugins/Runtime/MassGameplay/Source/MassMovement/Public/Example/MassSimpleMovementTrait.h b/Engine/Plugins/Runtime/MassGameplay/Source/MassMovement/Public/Example/MassSimpleMovementTrait.h
--- a/Engine/Plugins/Runtime/MassGameplay/Source/MassMovement/Public/Example/MassSimpleMovementTrait.h
+++ b/Engine/Plugins/Runtime/MassGameplay/Source/MassMovement/Public/Example/MassSimpleMovementTrait.h
@@ -21,7 +21,7 @@ class MASSMOVEMENT_API UMassSimpleMovementTrait : public UMassEntityTraitBase
 	GENERATED_BODY()
 
 protected:
-	virtual void BuildTemplate(FMassEntityTemplateBuildContext& BuildContext, UWorld& World) const override;
+	virtual void BuildTemplate(FMassEntityTemplateBuildContext& BuildContext, const UWorld& World) const override;
 };
 
 
diff --git a/Engine/Plugins/Runtime/MassGameplay/Source/MassMovement/Public/Example/MassVelocityRandomizerTrait.h b/Engine/Plugins/Runtime/MassGameplay/Source/MassMovement/Public/Example/MassVelocityRandomizerTrait.h
--- a/Engine/Plugins/Runtime/MassGameplay/Source/MassMovement/Public/Example/MassVelocityRandomizerTrait.h
+++ b/Engine/Plugins/Runtime/MassGameplay/Source/MassMovement/Public/Example/MassVelocityRandomizerTrait.h
@@ -15,7 +15,7 @@ class MASSMOVEMENT_API UMassVelocityRandomizerTrait : public UMassEntityTraitBas
 	GENERATED_BODY()
 
 protected:
-	virtual void BuildTemplate(FMassEntityTemplateBuildContext& BuildContext, UWorld& World) const override;
+	virtual void BuildTemplate(FMassEntityTemplateBuildContext& BuildContext, const UWorld& World) const override;
 
 	/** The speed is expressed in UnrealUnits per second, which usually translates to 0.01m/s */
 	UPROPERTY(Category = "Velocity", EditAnywhere, meta = (UIMin = 0.0, ClampMin = 0.0))
diff --git a/Engine/Plugins/Runtime/MassGameplay/Source/MassMovement/Public/Movement/MassMovementTrait.h b/Engine/Plugins/Runtime/MassGameplay/Source/MassMovement/Public/Movement/MassMovementTrait.h
--- a/Engine/Plugins/Runtime/MassGameplay/Source/MassMovement/Public/Movement/MassMovementTrait.h
+++ b/Engine/Plugins/Runtime/MassGameplay/Source/MassMovement/Public/Movement/MassMovementTrait.h
@@ -12,7 +12,7 @@ class MASSMOVEMENT_API UMassMovementTrait : public UMassEntityTraitBase
 	GENERATED_BODY()
 
 protected:
-	virtual void BuildTemplate(FMassEntityTemplateBuildContext& BuildContext, UWorld& World) const override;
+	virtual void BuildTemplate(FMassEntityTemplateBuildContext& BuildContext, const UWorld& World) const override;
 
 	UPROPERTY(Category="Movement", EditAnywhere)
 	FMassMovementParameters Movement;
diff --git a/Engine/Plugins/Runtime/MassGameplay/Source/MassReplication/Private/MassReplicationTrait.cpp b/Engine/Plugins/Runtime/MassGameplay/Source/MassReplication/Private/MassReplicationTrait.cpp
--- a/Engine/Plugins/Runtime/MassGameplay/Source/MassReplication/Private/MassReplicationTrait.cpp
+++ b/Engine/Plugins/Runtime/MassGameplay/Source/MassReplication/Private/MassReplicationTrait.cpp
@@ -9,7 +9,7 @@
 #include "MassReplicationSubsystem.h"
 
 
-void UMassReplicationTrait::BuildTemplate(FMassEntityTemplateBuildContext& BuildContext, UWorld& World) const
+void UMassReplicationTrait::BuildTemplate(FMassEntityTemplateBuildContext& BuildContext, const UWorld& World) const
 {
 	if (World.IsNetMode(NM_Standalone))
 	{
diff --git a/Engine/Plugins/Runtime/MassGameplay/Source/MassReplication/Public/MassReplicationTrait.h b/Engine/Plugins/Runtime/MassGameplay/Source/MassReplication/Public/MassReplicationTrait.h
--- a/Engine/Plugins/Runtime/MassGameplay/Source/MassReplication/Public/MassReplicationTrait.h
+++ b/Engine/Plugins/Runtime/MassGameplay/Source/MassReplication/Public/MassReplicationTrait.h
@@ -14,7 +14,7 @@ class MASSREPLICATION_API UMassReplicationTrait : public UMassEntityTraitBase
 
 public:
 
-	virtual void BuildTemplate(FMassEntityTemplateBuildContext& BuildContext, UWorld& World) const override;
+	virtual void BuildTemplate(FMassEntityTemplateBuildContext& BuildContext, const UWorld& World) const override;
 
 	UPROPERTY(EditAnywhere, Category = "Mass|Replication")
 	FMassReplicationParameters Params;
diff --git a/Engine/Plugins/Runtime/MassGameplay/Source/MassRepresentation/Private/MassVisualizationTrait.cpp b/Engine/Plugins/Runtime/MassGameplay/Source/MassRepresentation/Private/MassVisualizationTrait.cpp
--- a/Engine/Plugins/Runtime/MassGameplay/Source/MassRepresentation/Private/MassVisualizationTrait.cpp
+++ b/Engine/Plugins/Runtime/MassGameplay/Source/MassRepresentation/Private/MassVisualizationTrait.cpp
@@ -40,7 +40,7 @@ UMassVisualizationTrait::UMassVisualizationTrait()
 	LODParams.DistanceToFrustumHysteresis = 0.0f;
 }
 
-void UMassVisualizationTrait::BuildTemplate(FMassEntityTemplateBuildContext& BuildContext, UWorld& World) const
+void UMassVisualizationTrait::BuildTemplate(FMassEntityTemplateBuildContext& BuildContext, const UWorld& World) const
 {
 	// This should not be ran on NM_Server network mode
 	if (World.IsNetMode(NM_DedicatedServer))
diff --git a/Engine/Plugins/Runtime/MassGameplay/Source/MassRepresentation/Public/MassVisualizationTrait.h b/Engine/Plugins/Runtime/MassGameplay/Source/MassRepresentation/Public/MassVisualizationTrait.h
--- a/Engine/Plugins/Runtime/MassGameplay/Source/MassRepresentation/Public/MassVisualizationTrait.h
+++ b/Engine/Plugins/Runtime/MassGameplay/Source/MassRepresentation/Public/MassVisualizationTrait.h
@@ -20,7 +20,7 @@ class MASSREPRESENTATION_API UMassVisualizationTrait : public UMassEntityTraitBa
 public:
 	UMassVisualizationTrait();
 
-	virtual void BuildTemplate(FMassEntityTemplateBuildContext& BuildContext, UWorld& World) const override;
+	virtual void BuildTemplate(FMassEntityTemplateBuildContext& BuildContext, const UWorld& World) const override;
 
 	/** Instanced static mesh information for this agent */
 	UPROPERTY(EditAnywhere, Category = "Mass|Visual")
diff --git a/Engine/Plugins/Runtime/MassGameplay/Source/MassSmartObjects/Private/MassSmartObjectUserTrait.cpp b/Engine/Plugins/Runtime/MassGameplay/Source/MassSmartObjects/Private/MassSmartObjectUserTrait.cpp
--- a/Engine/Plugins/Runtime/MassGameplay/Source/MassSmartObjects/Private/MassSmartObjectUserTrait.cpp
+++ b/Engine/Plugins/Runtime/MassGameplay/Source/MassSmartObjects/Private/MassSmartObjectUserTrait.cpp
@@ -5,7 +5,7 @@
 #include "MassSmartObjectFragments.h"
 #include "MassEntityTemplateRegistry.h"
 
-void UMassSmartObjectUserTrait::BuildTemplate(FMassEntityTemplateBuildContext& BuildContext, UWorld& World) const
+void UMassSmartObjectUserTrait::BuildTemplate(FMassEntityTemplateBuildContext& BuildContext, const UWorld& World) const
 {
 	BuildContext.AddFragment<FMassSmartObjectUserFragment>();
 }
diff --git a/Engine/Plugins/Runtime/MassGameplay/Source/MassSmartObjects/Public/MassSmartObjectUserTrait.h b/Engine/Plugins/Runtime/MassGameplay/Source/MassSmartObjects/Public/MassSmartObjectUserTrait.h
--- a/Engine/Plugins/Runtime/MassGameplay/Source/MassSmartObjects/Public/MassSmartObjectUserTrait.h
+++ b/Engine/Plugins/Runtime/MassGameplay/Source/MassSmartObjects/Public/MassSmartObjectUserTrait.h
@@ -14,5 +14,5 @@ class MASSSMARTOBJECTS_API UMassSmartObjectUserTrait : public UMassEntityTraitBa
 	GENERATED_BODY()
 
 protected:
-	virtual void BuildTemplate(FMassEntityTemplateBuildContext& BuildContext, UWorld& World) const override;
+	virtual void BuildTemplate(FMassEntityTemplateBuildContext& BuildContext, const UWorld& World) const override;
 };
diff --git a/Engine/Plugins/Runtime/MassGameplay/Source/MassSpawner/Private/MassAssortedFragmentsTrait.cpp b/Engine/Plugins/Runtime/MassGameplay/Source/MassSpawner/Private/MassAssortedFragmentsTrait.cpp
--- a/Engine/Plugins/Runtime/MassGameplay/Source/MassSpawner/Private/MassAssortedFragmentsTrait.cpp
+++ b/Engine/Plugins/Runtime/MassGameplay/Source/MassSpawner/Private/MassAssortedFragmentsTrait.cpp
@@ -4,7 +4,7 @@
 #include "MassEntityTemplateRegistry.h"
 
 
-void UMassAssortedFragmentsTrait::BuildTemplate(FMassEntityTemplateBuildContext& BuildContext, UWorld& World) const
+void UMassAssortedFragmentsTrait::BuildTemplate(FMassEntityTemplateBuildContext& BuildContext, const UWorld& World) const
 {
 	for (const FInstancedStruct& Fragment : Fragments)
 	{
diff --git a/Engine/Plugins/Runtime/MassGameplay/Source/MassSpawner/Private/MassEntityConfigAsset.cpp b/Engine/Plugins/Runtime/MassGameplay/Source/MassSpawner/Private/MassEntityConfigAsset.cpp
--- a/Engine/Plugins/Runtime/MassGameplay/Source/MassSpawner/Private/MassEntityConfigAsset.cpp
+++ b/Engine/Plugins/Runtime/MassGameplay/Source/MassSpawner/Private/MassEntityConfigAsset.cpp
@@ -58,7 +58,7 @@ void UMassEntityConfigAsset::ValidateEntityConfig()
 }
 #endif // WITH_EDITOR
 
-const FMassEntityTemplate& FMassEntityConfig::GetOrCreateEntityTemplate(UWorld& World, const UObject& ConfigOwner) const
+const FMassEntityTemplate& FMassEntityConfig::GetOrCreateEntityTemplate(const UWorld& World, const UObject& ConfigOwner) const
 {
 	uint32 Hash;
 	FMassEntityTemplateID TemplateID;
@@ -89,7 +89,7 @@ const FMassEntityTemplate& FMassEntityConfig::GetOrCreateEntityTemplate(UWorld&
 	return Template;
 }
 
-void FMassEntityConfig::DestroyEntityTemplate(UWorld& World, const UObject& ConfigOwner) const
+void FMassEntityConfig::DestroyEntityTemplate(const UWorld& World, const UObject& ConfigOwner) const
 {
 	uint32 Hash;
 	FMassEntityTemplateID TemplateID;
@@ -116,7 +116,7 @@ void FMassEntityConfig::DestroyEntityTemplate(UWorld& World, const UObject& Conf
 	TemplateRegistry.DestroyTemplate(Hash, TemplateID);
 }
 
-const FMassEntityTemplate& FMassEntityConfig::GetEntityTemplateChecked(UWorld& World, const UObject& ConfigOwner) const
+const FMassEntityTemplate& FMassEntityConfig::GetEntityTemplateChecked(const UWorld& World, const UObject& ConfigOwner) const
 {
 	uint32 Hash;
 	FMassEntityTemplateID TemplateID;
@@ -126,7 +126,7 @@ const FMassEntityTemplate& FMassEntityConfig::GetEntityTemplateChecked(UWorld& W
 	return *ExistingTemplate;
 }
 
-const FMassEntityTemplate* FMassEntityConfig::GetEntityTemplateInternal(UWorld& World, const UObject& ConfigOwner, uint32& HashOut, FMassEntityTemplateID& TemplateIDOut, TArray<UMassEntityTraitBase*>& CombinedTraitsOut) const
+const FMassEntityTemplate* FMassEntityConfig::GetEntityTemplateInternal(const UWorld& World, const UObject& ConfigOwner, uint32& HashOut, FMassEntityTemplateID& TemplateIDOut, TArray<UMassEntityTraitBase*>& CombinedTraitsOut) const
 {
 	UMassSpawnerSubsystem* SpawnerSystem = UWorld::GetSubsystem<UMassSpawnerSubsystem>(&World);
 	check(SpawnerSystem);
@@ -192,7 +192,7 @@ void FMassEntityConfig::AddTrait(UMassEntityTraitBase& Trait)
 	Traits.Add(&Trait);
 }
 
-bool FMassEntityConfig::ValidateEntityTemplate(UWorld& World, const UObject& ConfigOwner)
+bool FMassEntityConfig::ValidateEntityTemplate(const UWorld& World, const UObject& ConfigOwner)
 {
 	TArray<const UObject*> Visited;
 	TArray<UMassEntityTraitBase*> CombinedTraits;
diff --git a/Engine/Plugins/Runtime/MassGameplay/Source/MassSpawner/Private/MassEntityTemplateRegistry.cpp b/Engine/Plugins/Runtime/MassGameplay/Source/MassSpawner/Private/MassEntityTemplateRegistry.cpp
--- a/Engine/Plugins/Runtime/MassGameplay/Source/MassSpawner/Private/MassEntityTemplateRegistry.cpp
+++ b/Engine/Plugins/Runtime/MassGameplay/Source/MassSpawner/Private/MassEntityTemplateRegistry.cpp
@@ -197,7 +197,7 @@ void UMassEntityTemplateRegistry::DestroyTemplate(const uint32 HashLookup, FMass
 //----------------------------------------------------------------------//
 // FMassEntityTemplateBuildContext 
 //----------------------------------------------------------------------//
-bool FMassEntityTemplateBuildContext::BuildFromTraits(TConstArrayView<UMassEntityTraitBase*> Traits, UWorld& World)
+bool FMassEntityTemplateBuildContext::BuildFromTraits(TConstArrayView<UMassEntityTraitBase*> Traits, const UWorld& World)
 {
 	TraitAddedTypes.Reset();
 	TraitsDependencies.Reset();
@@ -214,7 +214,7 @@ bool FMassEntityTemplateBuildContext::BuildFromTraits(TConstArrayView<UMassEntit
 	return ValidateBuildContext(World);
 }
 
-bool FMassEntityTemplateBuildContext::ValidateBuildContext(UWorld& World)
+bool FMassEntityTemplateBuildContext::ValidateBuildContext(const UWorld& World)
 {
 	// Group same types(key) together
 	TraitAddedTypes.KeySort( [](const UStruct& LHS, const UStruct& RHS) { return LHS.GetName() < RHS.GetName(); } );
diff --git a/Engine/Plugins/Runtime/MassGameplay/Source/MassSpawner/Public/MassAssortedFragmentsTrait.h b/Engine/Plugins/Runtime/MassGameplay/Source/MassSpawner/Public/MassAssortedFragmentsTrait.h
--- a/Engine/Plugins/Runtime/MassGameplay/Source/MassSpawner/Public/MassAssortedFragmentsTrait.h
+++ b/Engine/Plugins/Runtime/MassGameplay/Source/MassSpawner/Public/MassAssortedFragmentsTrait.h
@@ -17,7 +17,7 @@ class MASSSPAWNER_API UMassAssortedFragmentsTrait : public UMassEntityTraitBase
 
 protected:
 
-	virtual void BuildTemplate(FMassEntityTemplateBuildContext& BuildContext, UWorld& World) const override;
+	virtual void BuildTemplate(FMassEntityTemplateBuildContext& BuildContext, const UWorld& World) const override;
 
 	UPROPERTY(Category="Fragments", EditAnywhere, meta = (BaseStruct = "MassFragment", ExcludeBaseStruct))
 	TArray<FInstancedStruct> Fragments;
diff --git a/Engine/Plugins/Runtime/MassGameplay/Source/MassSpawner/Public/MassEntityConfigAsset.h b/Engine/Plugins/Runtime/MassGameplay/Source/MassSpawner/Public/MassEntityConfigAsset.h
--- a/Engine/Plugins/Runtime/MassGameplay/Source/MassSpawner/Public/MassEntityConfigAsset.h
+++ b/Engine/Plugins/Runtime/MassGameplay/Source/MassSpawner/Public/MassEntityConfigAsset.h
@@ -30,14 +30,14 @@ struct MASSSPAWNER_API FMassEntityConfig
 	 *  @param World in which we are creating the template.
 	 *  @param ConfigOwner Owner of the FMassEntityConfig used for error reporting.
 	 */
-	const FMassEntityTemplate& GetOrCreateEntityTemplate(UWorld& World, const UObject& ConfigOwner) const;
+	const FMassEntityTemplate& GetOrCreateEntityTemplate(const UWorld& World, const UObject& ConfigOwner) const;
 
-	void DestroyEntityTemplate(UWorld& World, const UObject& ConfigOwner) const;
+	void DestroyEntityTemplate(const UWorld& World, const UObject& ConfigOwner) const;
 
 	/** 
 	 * Fetches the EntityTemplate for given World, fails a check if one cannot be found.
 	 */
-	const FMassEntityTemplate& GetEntityTemplateChecked(UWorld& World, const UObject& ConfigOwner) const;
+	const FMassEntityTemplate& GetEntityTemplateChecked(const UWorld& World, const UObject& ConfigOwner) const;
 
 	/** @return Parent config */
 	const UMassEntityConfigAsset* GetParent() const { return Parent; }
@@ -49,7 +49,7 @@ struct MASSSPAWNER_API FMassEntityConfig
 	void AddTrait(UMassEntityTraitBase& Trait);
 
 	/** Validates if the entity template is well built */
-	bool ValidateEntityTemplate(UWorld& World, const UObject& ConfigOwner);
+	bool ValidateEntityTemplate(const UWorld& World, const UObject& ConfigOwner);
 
 protected:
 	/** Combines traits based on the config hierarchy and returns list of unique traits */
@@ -64,7 +64,7 @@ protected:
 	TArray<UMassEntityTraitBase*> Traits;
 
 private:
-	const FMassEntityTemplate* GetEntityTemplateInternal(UWorld& World, const UObject& ConfigOwner, uint32& HashOut, FMassEntityTemplateID& TemplateIDOut, TArray<UMassEntityTraitBase*>& CombinedTraitsOut) const;
+	const FMassEntityTemplate* GetEntityTemplateInternal(const UWorld& World, const UObject& ConfigOwner, uint32& HashOut, FMassEntityTemplateID& TemplateIDOut, TArray<UMassEntityTraitBase*>& CombinedTraitsOut) const;
 };
 
 /**
@@ -83,12 +83,12 @@ public:
 	/** @return Mutable agent config stored in this asset */
 	FMassEntityConfig& GetMutableConfig() { return Config; }
 
-	const FMassEntityTemplate& GetOrCreateEntityTemplate(UWorld& World) const
+	const FMassEntityTemplate& GetOrCreateEntityTemplate(const UWorld& World) const
 	{
 		return Config.GetOrCreateEntityTemplate(World, *this);
 	}
 
-	void DestroyEntityTemplate(UWorld& World) const
+	void DestroyEntityTemplate(const UWorld& World) const
 	{
 		Config.DestroyEntityTemplate(World, *this);
 	}
diff --git a/Engine/Plugins/Runtime/MassGameplay/Source/MassSpawner/Public/MassEntityTemplateRegistry.h b/Engine/Plugins/Runtime/MassGameplay/Source/MassSpawner/Public/MassEntityTemplateRegistry.h
--- a/Engine/Plugins/Runtime/MassGameplay/Source/MassSpawner/Public/MassEntityTemplateRegistry.h
+++ b/Engine/Plugins/Runtime/MassGameplay/Source/MassSpawner/Public/MassEntityTemplateRegistry.h
@@ -163,7 +163,7 @@ struct FMassEntityTemplateBuildContext
 	 * @param World owning the MassEntitySubsystem for which the entity template is built
 	 * @return true if there were no validation errors
 	 */
-	bool BuildFromTraits(TConstArrayView<UMassEntityTraitBase*> Traits, UWorld& World);
+	bool BuildFromTraits(TConstArrayView<UMassEntityTraitBase*> Traits, const UWorld& World);
 
 protected:
 
@@ -172,7 +172,7 @@ protected:
 	 * @param World owning the MassEntitySubsystem for which the entity template is validated against
 	 * @return true if there were no validation errors
 	 */
-	bool ValidateBuildContext(UWorld& World);
+	bool ValidateBuildContext(const UWorld& World);
 
 	void TypeAdded(const UStruct& Type)
 	{
diff --git a/Engine/Plugins/Runtime/MassGameplay/Source/MassSpawner/Public/MassEntityTraitBase.h b/Engine/Plugins/Runtime/MassGameplay/Source/MassSpawner/Public/MassEntityTraitBase.h
--- a/Engine/Plugins/Runtime/MassGameplay/Source/MassSpawner/Public/MassEntityTraitBase.h
+++ b/Engine/Plugins/Runtime/MassGameplay/Source/MassSpawner/Public/MassEntityTraitBase.h
@@ -22,10 +22,10 @@ class MASSSPAWNER_API UMassEntityTraitBase : public UObject
 public:
 
 	/** Appends items into the entity template required for the trait. */
-	virtual void BuildTemplate(FMassEntityTemplateBuildContext& BuildContext, UWorld& World) const PURE_VIRTUAL(UMassEntityTraitBase::BuildTemplate, return; );
+	virtual void BuildTemplate(FMassEntityTemplateBuildContext& BuildContext, const UWorld& World) const PURE_VIRTUAL(UMassEntityTraitBase::BuildTemplate, return; );
 
 	virtual void DestroyTemplate() const {}
 
 	/** Called when all Traits have had BuildTemplate() called. */
-	virtual void ValidateTemplate(FMassEntityTemplateBuildContext& BuildContext, UWorld& World) const {};
+	virtual void ValidateTemplate(FMassEntityTemplateBuildContext& BuildContext, const UWorld& World) const {};
 };
