Fix static analysis

diff --git a/Engine/Plugins/Runtime/MassGameplay/Source/MassSpawner/Private/MassEntityTemplateRegistry.cpp b/Engine/Plugins/Runtime/MassGameplay/Source/MassSpawner/Private/MassEntityTemplateRegistry.cpp
--- a/Engine/Plugins/Runtime/MassGameplay/Source/MassSpawner/Private/MassEntityTemplateRegistry.cpp
+++ b/Engine/Plugins/Runtime/MassGameplay/Source/MassSpawner/Private/MassEntityTemplateRegistry.cpp
@@ -237,6 +237,8 @@ bool FMassEntityTemplateBuildContext::ValidateBuildContext(const UWorld& World)
 		{
 			if (!bHeaderOutputed)
 			{
+				check(CurrentStruct);
+				check(CurrentTrait);
 				UE_LOG(LogMass, Error, TEXT("Fragment(%s) was added multiple time and can only be added by one trait. Fragment was added by:"), *CurrentStruct->GetName());
 				UE_LOG(LogMass, Error, TEXT("\t\t%s"), *CurrentTrait->GetClass()->GetName());
 				bHeaderOutputed = true;
