Fix static analysis

diff --git a/Engine/Plugins/Runtime/MassGameplay/Source/MassSpawner/Private/MassEntityTemplateRegistry.cpp b/Engine/Plugins/Runtime/MassGameplay/Source/MassSpawner/Private/MassEntityTemplateRegistry.cpp
--- a/Engine/Plugins/Runtime/MassGameplay/Source/MassSpawner/Private/MassEntityTemplateRegistry.cpp
+++ b/Engine/Plugins/Runtime/MassGameplay/Source/MassSpawner/Private/MassEntityTemplateRegistry.cpp
@@ -230,6 +230,7 @@ bool FMassEntityTemplateBuildContext::ValidateBuildContext(const UWorld& World)
 		{
 			CurrentStruct = Pair.Key;
 			CurrentTrait = Pair.Value;
+			check(CurrentTrait);
 			CurrentTrait->ValidateTemplate(*this, World);
 			bHeaderOutputed = false;
 		}
@@ -263,6 +264,7 @@ bool FMassEntityTemplateBuildContext::ValidateBuildContext(const UWorld& World)
 		{
 			if (!bHeaderOutputed)
 			{
+				check(CurrentTrait);
 				UE_LOG(LogMass, Error, TEXT("Trait(%s) has missing dependency:"),  *CurrentTrait->GetClass()->GetName() );
 				bHeaderOutputed = true;
 			}
