Uses of the new "system requirements"

The main point of the change is to get rid of processor-cached subsystem pointers and let queries declare their desired subsystem use so that we can arrange calculations accordingly and thread-safely for parallel processing (future work).

diff --git a/Engine/Plugins/AI/MassAI/Source/MassAIBehavior/Private/MassLookAtProcessors.cpp b/Engine/Plugins/AI/MassAI/Source/MassAIBehavior/Private/MassLookAtProcessors.cpp
--- a/Engine/Plugins/AI/MassAI/Source/MassAIBehavior/Private/MassLookAtProcessors.cpp
+++ b/Engine/Plugins/AI/MassAI/Source/MassAIBehavior/Private/MassLookAtProcessors.cpp
@@ -95,17 +95,6 @@ UMassLookAtProcessor::UMassLookAtProcessor()
 	ExecutionOrder.ExecuteAfter.Add(UE::Mass::ProcessorGroupNames::Representation);
 }
 
-void UMassLookAtProcessor::Initialize(UObject& Owner)
-{
-	Super::Initialize(Owner);
-
-	NavigationSubsystem = UWorld::GetSubsystem<UMassNavigationSubsystem>(Owner.GetWorld());
-	checkf(NavigationSubsystem != nullptr, TEXT("UMassNavigationSubsystem is mandatory when using MassLookAtProcessor processor."));
-
-	ZoneGraphSubsystem = UWorld::GetSubsystem<UZoneGraphSubsystem>(Owner.GetWorld());
-	checkf(ZoneGraphSubsystem != nullptr, TEXT("UMassZoneGraphSubsystem is mandatory when using MassLookAtProcessor processor."));
-}
-
 void UMassLookAtProcessor::ConfigureQueries()
 {
 	EntityQuery_Conditional.AddRequirement<FMassLookAtFragment>(EMassFragmentAccess::ReadWrite);
@@ -119,6 +108,8 @@ void UMassLookAtProcessor::ConfigureQueries()
 	EntityQuery_Conditional.AddTagRequirement<FMassOffLODTag>(EMassFragmentPresence::None);
 	EntityQuery_Conditional.AddChunkRequirement<FMassVisualizationChunkFragment>(EMassFragmentAccess::ReadOnly);
 	EntityQuery_Conditional.SetChunkFilter(&FMassVisualizationChunkFragment::AreAnyEntitiesVisibleInChunk);
+	EntityQuery_Conditional.AddSystemRequirement<UMassNavigationSubsystem>(EMassFragmentAccess::ReadOnly);
+	EntityQuery_Conditional.AddSystemRequirement<UZoneGraphSubsystem>(EMassFragmentAccess::ReadOnly);
 }
 
 void UMassLookAtProcessor::Execute(UMassEntitySubsystem& EntitySubsystem, FMassExecutionContext& Context)
@@ -127,8 +118,11 @@ void UMassLookAtProcessor::Execute(UMassEntitySubsystem& EntitySubsystem, FMassE
 
 	const float CurrentTime = GetWorld()->GetTimeSeconds();
 
-	EntityQuery_Conditional.ForEachEntityChunk(EntitySubsystem, Context, [this, &EntitySubsystem, CurrentTime](FMassExecutionContext& Context)
+	EntityQuery_Conditional.ForEachEntityChunk(EntitySubsystem, Context, [this, &EntitySubsystem, CurrentTime, World = EntitySubsystem.GetWorld()](FMassExecutionContext& Context)
 		{
+			const UMassNavigationSubsystem& MassNavSystem = Context.GetSubsystemChecked<UMassNavigationSubsystem>(World);
+			const UZoneGraphSubsystem& ZoneGraphSubsystem = Context.GetSubsystemChecked<UZoneGraphSubsystem>(World);
+
 			const int32 NumEntities = Context.GetNumEntities();
 			const TArrayView<FMassLookAtFragment> LookAtList = Context.GetMutableFragmentView<FMassLookAtFragment>();
 			const TConstArrayView<FTransformFragment> TransformList = Context.GetFragmentView<FTransformFragment>();
@@ -158,7 +152,7 @@ void UMassLookAtProcessor::Execute(UMassEntitySubsystem& EntitySubsystem, FMassE
 					const float TimeSinceUpdate = CurrentTime - LookAt.GazeStartTime;
 					if (TimeSinceUpdate >= LookAt.GazeDuration)
 					{
-						FindNewGazeTarget(EntitySubsystem, CurrentTime, TransformFragment.GetTransform(), LookAt);
+						FindNewGazeTarget(MassNavSystem, EntitySubsystem, CurrentTime, TransformFragment.GetTransform(), LookAt);
 					}
 				}
 
@@ -182,7 +176,7 @@ void UMassLookAtProcessor::Execute(UMassEntitySubsystem& EntitySubsystem, FMassE
 							const FMassZoneGraphLaneLocationFragment& LaneLocation = ZoneGraphLocationList[i];
 							const FMassZoneGraphShortPathFragment& ShortPath = ShortPathList[i];
 							
-							BuildTrajectory(LaneLocation, ShortPath, Entity, bDisplayDebug, LookAtTrajectory);
+							BuildTrajectory(ZoneGraphSubsystem, LaneLocation, ShortPath, Entity, bDisplayDebug, LookAtTrajectory);
 							LookAt.LastSeenActionID = MoveTarget.GetCurrentActionID();
 						}
 						
@@ -237,8 +231,9 @@ void UMassLookAtProcessor::Execute(UMassEntitySubsystem& EntitySubsystem, FMassE
 		});
 }
 
-void UMassLookAtProcessor::FindNewGazeTarget(const UMassEntitySubsystem& EntitySubsystem, const float CurrentTime, const FTransform& Transform, FMassLookAtFragment& LookAt) const
+void UMassLookAtProcessor::FindNewGazeTarget(const UMassNavigationSubsystem& MassNavSystem, const UMassEntitySubsystem& EntitySubsystem, const float CurrentTime, const FTransform& Transform, FMassLookAtFragment& LookAt) const
 {
+	const FNavigationObstacleHashGrid2D& ObstacleGrid = MassNavSystem.GetObstacleGrid();
 	const FMassEntityHandle LastTrackedEntity = LookAt.GazeTrackedEntity;
 	
 	LookAt.GazeTrackedEntity.Reset();
@@ -246,7 +241,7 @@ void UMassLookAtProcessor::FindNewGazeTarget(const UMassEntitySubsystem& EntityS
 
 	// Search for potential targets in front
 	bool bTargetFound = false;
-	if (LookAt.bRandomGazeEntities && NavigationSubsystem != nullptr)
+	if (LookAt.bRandomGazeEntities)
 	{
 		const float CosAngleThreshold = FMath::Cos(FMath::DegreesToRadians(AngleThresholdInDegrees));
 		const FVector Extent(QueryExtent, QueryExtent, QueryExtent);
@@ -255,7 +250,7 @@ void UMassLookAtProcessor::FindNewGazeTarget(const UMassEntitySubsystem& EntityS
 
 		TArray<FNavigationObstacleHashGrid2D::ItemIDType> NearbyEntities;
 		NearbyEntities.Reserve(16);
-		NavigationSubsystem->GetObstacleGrid().Query(QueryBox, NearbyEntities);
+		ObstacleGrid.Query(QueryBox, NearbyEntities);
 
 		// We'll pick the first entity that passes, this ensure that it's random one.
 		Algo::RandomShuffle(NearbyEntities);
@@ -403,7 +398,7 @@ bool UMassLookAtProcessor::UpdateGazeTrackedEntity(const UMassEntitySubsystem& E
 	return bHasTarget;
 }
 
-void UMassLookAtProcessor::BuildTrajectory(const FMassZoneGraphLaneLocationFragment& LaneLocation, const FMassZoneGraphShortPathFragment& ShortPath,
+void UMassLookAtProcessor::BuildTrajectory(const UZoneGraphSubsystem& ZoneGraphSubsystem, const FMassZoneGraphLaneLocationFragment& LaneLocation, const FMassZoneGraphShortPathFragment& ShortPath,
 											const FMassEntityHandle Entity, const bool bDisplayDebug, FMassLookAtTrajectoryFragment& LookAtTrajectory)
 {
 	LookAtTrajectory.Reset();
@@ -427,7 +422,7 @@ void UMassLookAtProcessor::BuildTrajectory(const FMassZoneGraphLaneLocationFragm
 	// If the path will lead to next lane, add a point from next lane too.
 	if (ShortPath.NextLaneHandle.IsValid())
 	{
-		const FZoneGraphStorage* ZoneGraphStorage = ZoneGraphSubsystem->GetZoneGraphStorage(LaneLocation.LaneHandle.DataHandle);
+		const FZoneGraphStorage* ZoneGraphStorage = ZoneGraphSubsystem.GetZoneGraphStorage(LaneLocation.LaneHandle.DataHandle);
 		if (ZoneGraphStorage != nullptr)
 		{
 			if (ShortPath.NextExitLinkType == EZoneLaneLinkType::Outgoing)
diff --git a/Engine/Plugins/AI/MassAI/Source/MassAIBehavior/Private/MassStateTreeProcessors.cpp b/Engine/Plugins/AI/MassAI/Source/MassAIBehavior/Private/MassStateTreeProcessors.cpp
--- a/Engine/Plugins/AI/MassAI/Source/MassAIBehavior/Private/MassStateTreeProcessors.cpp
+++ b/Engine/Plugins/AI/MassAI/Source/MassAIBehavior/Private/MassStateTreeProcessors.cpp
@@ -187,6 +187,7 @@ void UMassStateTreeFragmentDestructor::ConfigureQueries()
 {
 	EntityQuery.AddRequirement<FMassStateTreeInstanceFragment>(EMassFragmentAccess::ReadWrite);
 	EntityQuery.AddConstSharedRequirement<FMassStateTreeSharedFragment>();
+	EntityQuery.AddSystemRequirement<UMassStateTreeSubsystem>(EMassFragmentAccess::ReadWrite);
 }
 
 void UMassStateTreeFragmentDestructor::Execute(UMassEntitySubsystem& EntitySubsystem, FMassExecutionContext& Context)
@@ -197,17 +198,17 @@ void UMassStateTreeFragmentDestructor::Execute(UMassEntitySubsystem& EntitySubsy
 	}
 		
 	FMassStateTreeExecutionContext StateTreeContext(EntitySubsystem, *SignalSubsystem, Context);
-	UMassStateTreeSubsystem* MassStateTreeSubsystem = UWorld::GetSubsystem<UMassStateTreeSubsystem>(EntitySubsystem.GetWorld());
 
 	EntityQuery.ForEachEntityChunk(EntitySubsystem,Context,
-		[this, &StateTreeContext, MassStateTreeSubsystem](FMassExecutionContext& Context)
+		[this, &StateTreeContext, World = EntitySubsystem.GetWorld()](FMassExecutionContext& Context)
 		{
+			UMassStateTreeSubsystem& MassStateTreeSubsystem = Context.GetMutableSubsystemChecked<UMassStateTreeSubsystem>(World);
 			const TArrayView<FMassStateTreeInstanceFragment> StateTreeInstanceList = Context.GetMutableFragmentView<FMassStateTreeInstanceFragment>();
 			const FMassStateTreeSharedFragment& SharedStateTree = Context.GetConstSharedFragment<FMassStateTreeSharedFragment>();
 
 			const int32 NumEntities = Context.GetNumEntities();
 	
-			UE::MassBehavior::ForEachEntityInChunk(Context, StateTreeContext, *MassStateTreeSubsystem,
+			UE::MassBehavior::ForEachEntityInChunk(Context, StateTreeContext, MassStateTreeSubsystem,
 				[](FMassStateTreeExecutionContext& StateTreeExecutionContext, FMassStateTreeInstanceFragment& StateTreeInstance, FStateTreeInstanceData& InstanceData)
 				{
 					// Stop the tree instance
@@ -220,7 +221,7 @@ void UMassStateTreeFragmentDestructor::Execute(UMassEntitySubsystem& EntitySubsy
 				FMassStateTreeInstanceFragment& StateTreeInstance = StateTreeInstanceList[EntityIndex];
 				if (StateTreeInstance.InstanceHandle.IsValid())
 				{
-					MassStateTreeSubsystem->FreeInstanceData(StateTreeInstance.InstanceHandle);
+					MassStateTreeSubsystem.FreeInstanceData(StateTreeInstance.InstanceHandle);
 					StateTreeInstance.InstanceHandle = FMassStateTreeInstanceHandle();
 				}
 			}
@@ -247,6 +248,7 @@ void UMassStateTreeActivationProcessor::ConfigureQueries()
 	EntityQuery.AddConstSharedRequirement<FMassStateTreeSharedFragment>();
 	EntityQuery.AddTagRequirement<FMassStateTreeActivatedTag>(EMassFragmentPresence::None);
 	EntityQuery.AddChunkRequirement<FMassSimulationVariableTickChunkFragment>(EMassFragmentAccess::ReadOnly, EMassFragmentPresence::Optional);
+	EntityQuery.AddSystemRequirement<UMassStateTreeSubsystem>(EMassFragmentAccess::ReadWrite);
 }
 
 void UMassStateTreeActivationProcessor::Execute(UMassEntitySubsystem& EntitySubsystem, FMassExecutionContext& Context)
@@ -263,7 +265,6 @@ void UMassStateTreeActivationProcessor::Execute(UMassEntitySubsystem& EntitySubs
 	// The initializer provides that by sending a signal to all new entities that use StateTree.
 
 	FMassStateTreeExecutionContext StateTreeContext(EntitySubsystem, *SignalSubsystem, Context);
-	UMassStateTreeSubsystem* MassStateTreeSubsystem = UWorld::GetSubsystem<UMassStateTreeSubsystem>(EntitySubsystem.GetWorld());
 
 	const float TimeInSeconds = EntitySubsystem.GetWorld()->GetTimeSeconds();
 
@@ -271,8 +272,9 @@ void UMassStateTreeActivationProcessor::Execute(UMassEntitySubsystem& EntitySubs
 	int32 ActivationCounts[EMassLOD::Max] {0,0,0,0};
 	
 	EntityQuery.ForEachEntityChunk(EntitySubsystem, Context,
-		[&EntitiesToSignal, &ActivationCounts, MaxActivationsPerLOD = BehaviorSettings->MaxActivationsPerLOD, &StateTreeContext, TimeInSeconds, MassStateTreeSubsystem](FMassExecutionContext& Context)
+		[&EntitiesToSignal, &ActivationCounts, MaxActivationsPerLOD = BehaviorSettings->MaxActivationsPerLOD, &StateTreeContext, TimeInSeconds, World = EntitySubsystem.GetWorld()](FMassExecutionContext& Context)
 		{
+			UMassStateTreeSubsystem& MassStateTreeSubsystem = Context.GetMutableSubsystemChecked<UMassStateTreeSubsystem>(World);
 			const int32 NumEntities = Context.GetNumEntities();
 			// Check if we already reached the maximum for this frame
 			const EMassLOD::Type ChunkLOD = FMassSimulationVariableTickChunkFragment::GetChunkLOD(Context);
@@ -289,13 +291,12 @@ void UMassStateTreeActivationProcessor::Execute(UMassEntitySubsystem& EntitySubs
 			for (int32 EntityIndex = 0; EntityIndex < NumEntities; EntityIndex++)
 			{
 				FMassStateTreeInstanceFragment& StateTreeInstance = StateTreeInstanceList[EntityIndex];
-				StateTreeInstance.InstanceHandle = MassStateTreeSubsystem->AllocateInstanceData(SharedStateTree.StateTree);
+				StateTreeInstance.InstanceHandle = MassStateTreeSubsystem.AllocateInstanceData(SharedStateTree.StateTree);
 			}
 			
 			// Start StateTree. This may do substantial amount of work, as we select and enter the first state.
-			UE::MassBehavior::ForEachEntityInChunk(Context, StateTreeContext,*MassStateTreeSubsystem,
-				[TimeInSeconds, MassStateTreeSubsystem]
-				(FMassStateTreeExecutionContext& StateTreeExecutionContext, FMassStateTreeInstanceFragment& StateTreeInstance, FStateTreeInstanceData& InstanceData)
+			UE::MassBehavior::ForEachEntityInChunk(Context, StateTreeContext, MassStateTreeSubsystem,
+				[TimeInSeconds](FMassStateTreeExecutionContext& StateTreeExecutionContext, FMassStateTreeInstanceFragment& StateTreeInstance, FStateTreeInstanceData& InstanceData)
 				{
 					// Start the tree instance
 					StateTreeExecutionContext.Start(&InstanceData);
@@ -344,8 +345,7 @@ UMassStateTreeProcessor::UMassStateTreeProcessor(const FObjectInitializer& Objec
 void UMassStateTreeProcessor::Initialize(UObject& Owner)
 {
 	Super::Initialize(Owner);
-	
-	MassStateTreeSubsystem = UWorld::GetSubsystem<UMassStateTreeSubsystem>(Owner.GetWorld());
+
 	SignalSubsystem = UWorld::GetSubsystem<UMassSignalSubsystem>(Owner.GetWorld());
 
 	SubscribeToSignal(UE::Mass::Signals::StateTreeActivate);
@@ -375,11 +375,12 @@ void UMassStateTreeProcessor::ConfigureQueries()
 {
 	EntityQuery.AddRequirement<FMassStateTreeInstanceFragment>(EMassFragmentAccess::ReadWrite);
 	EntityQuery.AddConstSharedRequirement<FMassStateTreeSharedFragment>();
+	EntityQuery.AddSystemRequirement<UMassStateTreeSubsystem>(EMassFragmentAccess::ReadWrite);
 }
 
 void UMassStateTreeProcessor::SignalEntities(UMassEntitySubsystem& EntitySubsystem, FMassExecutionContext& Context, FMassSignalNameLookup& EntitySignals)
 {
-	if (MassStateTreeSubsystem == nullptr || SignalSubsystem == nullptr)
+	if (SignalSubsystem == nullptr)
 	{
 		return;
 	}
@@ -392,12 +393,14 @@ void UMassStateTreeProcessor::SignalEntities(UMassEntitySubsystem& EntitySubsyst
 	TArray<FMassEntityHandle> EntitiesToSignal;
 
 	EntityQuery.ForEachEntityChunk(EntitySubsystem, Context,
-		[this, &StateTreeContext, TimeInSeconds, &EntitiesToSignal](FMassExecutionContext& Context)
+		[this, &StateTreeContext, TimeInSeconds, &EntitiesToSignal, World = EntitySubsystem.GetWorld()](FMassExecutionContext& Context)
 		{
 			// Keep stats regarding the amount of tree instances ticked per frame
 			CSV_CUSTOM_STAT(StateTreeProcessor, NumTickedStateTree, Context.GetNumEntities(), ECsvCustomStatOp::Accumulate);
 
-			UE::MassBehavior::ForEachEntityInChunk(Context,StateTreeContext,*MassStateTreeSubsystem,
+			UMassStateTreeSubsystem& MassStateTreeSubsystem = Context.GetMutableSubsystemChecked<UMassStateTreeSubsystem>(World);
+
+			UE::MassBehavior::ForEachEntityInChunk(Context, StateTreeContext, MassStateTreeSubsystem,
 				[TimeInSeconds, &EntitiesToSignal](FMassStateTreeExecutionContext& StateTreeExecutionContext, FMassStateTreeInstanceFragment& StateTreeInstance, FStateTreeInstanceData& InstanceData)
 				{
 					// Compute adjusted delta time
diff --git a/Engine/Plugins/AI/MassAI/Source/MassAIBehavior/Public/MassLookAtProcessors.h b/Engine/Plugins/AI/MassAI/Source/MassAIBehavior/Public/MassLookAtProcessors.h
--- a/Engine/Plugins/AI/MassAI/Source/MassAIBehavior/Public/MassLookAtProcessors.h
+++ b/Engine/Plugins/AI/MassAI/Source/MassAIBehavior/Public/MassLookAtProcessors.h
@@ -28,11 +28,10 @@ class MASSAIBEHAVIOR_API UMassLookAtProcessor : public UMassProcessor
 protected:
 
 	virtual void ConfigureQueries() override;
-	virtual void Initialize(UObject& Owner) override;
 	virtual void Execute(UMassEntitySubsystem& EntitySubsystem, FMassExecutionContext& Context) override;
 
 	/** Selects a nearby target if possible or use a random fixed direction */
-	void FindNewGazeTarget(const UMassEntitySubsystem& EntitySubsystem, const float CurrentTime, const FTransform& Transform, FMassLookAtFragment& LookAt) const;
+	void FindNewGazeTarget(const UMassNavigationSubsystem& MassNavSystem, const UMassEntitySubsystem& EntitySubsystem, const float CurrentTime, const FTransform& Transform, FMassLookAtFragment& LookAt) const;
 
 	/** Updates look direction based on look at trajectory. */
 	void UpdateLookAtTrajectory(const FTransform& Transform, const FMassZoneGraphLaneLocationFragment& ZoneGraphLocation,
@@ -45,7 +44,7 @@ protected:
 	bool UpdateGazeTrackedEntity(const UMassEntitySubsystem& EntitySubsystem, const FTransform& Transform, const bool bDisplayDebug, FMassLookAtFragment& LookAt) const;
 
 	/** Builds look at trajectory along the current path. */
-	void BuildTrajectory(const FMassZoneGraphLaneLocationFragment& LaneLocation, const FMassZoneGraphShortPathFragment& ShortPath,
+	void BuildTrajectory(const UZoneGraphSubsystem& ZoneGraphSubsystem, const FMassZoneGraphLaneLocationFragment& LaneLocation, const FMassZoneGraphShortPathFragment& ShortPath,
 							const FMassEntityHandle Entity, const bool bDisplayDebug, FMassLookAtTrajectoryFragment& LookAtTrajectory);
 
 	/** Size of the query to find potential targets */
@@ -69,10 +68,4 @@ protected:
 	float AngleThresholdInDegrees = 0.f;
 
 	FMassEntityQuery EntityQuery_Conditional;
-
-	UPROPERTY(Transient)
-	UMassNavigationSubsystem* NavigationSubsystem;
-
-	UPROPERTY(Transient)
-	UZoneGraphSubsystem* ZoneGraphSubsystem;
 };
diff --git a/Engine/Plugins/AI/MassAI/Source/MassAIBehavior/Public/MassStateTreeProcessors.h b/Engine/Plugins/AI/MassAI/Source/MassAIBehavior/Public/MassStateTreeProcessors.h
--- a/Engine/Plugins/AI/MassAI/Source/MassAIBehavior/Public/MassStateTreeProcessors.h
+++ b/Engine/Plugins/AI/MassAI/Source/MassAIBehavior/Public/MassStateTreeProcessors.h
@@ -74,7 +74,4 @@ protected:
 	virtual void Initialize(UObject& Owner) override;
 	virtual void ConfigureQueries() override;
 	virtual void SignalEntities(UMassEntitySubsystem& EntitySubsystem, FMassExecutionContext& Context, FMassSignalNameLookup& EntitySignals) override;
-
-	UPROPERTY(Transient)
-	UMassStateTreeSubsystem* MassStateTreeSubsystem;
 };
diff --git a/Engine/Plugins/AI/MassAI/Source/MassNavigation/Private/MassNavigationProcessors.cpp b/Engine/Plugins/AI/MassAI/Source/MassNavigation/Private/MassNavigationProcessors.cpp
--- a/Engine/Plugins/AI/MassAI/Source/MassNavigation/Private/MassNavigationProcessors.cpp
+++ b/Engine/Plugins/AI/MassAI/Source/MassNavigation/Private/MassNavigationProcessors.cpp
@@ -169,6 +169,7 @@ void UMassNavigationObstacleGridProcessor::ConfigureQueries()
 	BaseEntityQuery.AddRequirement<FTransformFragment>(EMassFragmentAccess::ReadOnly);
 	BaseEntityQuery.AddRequirement<FAgentRadiusFragment>(EMassFragmentAccess::ReadOnly);
 	BaseEntityQuery.AddRequirement<FMassNavigationObstacleGridCellLocationFragment>(EMassFragmentAccess::ReadWrite);
+	BaseEntityQuery.AddSystemRequirement<UMassNavigationSubsystem>(EMassFragmentAccess::ReadWrite);
 
 	AddToGridEntityQuery = BaseEntityQuery;
 	AddToGridEntityQuery.AddRequirement<FMassAvoidanceColliderFragment>(EMassFragmentAccess::ReadOnly, EMassFragmentPresence::Optional);
@@ -185,23 +186,12 @@ void UMassNavigationObstacleGridProcessor::ConfigureQueries()
 	RemoveFromGridEntityQuery.AddTagRequirement<FMassInNavigationObstacleGridTag>(EMassFragmentPresence::All);
 }
 
-void UMassNavigationObstacleGridProcessor::Initialize(UObject& Owner)
-{
-	Super::Initialize(Owner);
-
-	NavigationSubsystem = UWorld::GetSubsystem<UMassNavigationSubsystem>(Owner.GetWorld());
-}
-
 void UMassNavigationObstacleGridProcessor::Execute(UMassEntitySubsystem& EntitySubsystem, FMassExecutionContext& Context)
 {
-	if (!NavigationSubsystem)
-	{
-		return;
-	}
-
 	// can't be ParallelFor due to MovementSubsystem->GetGridMutable().Move not being thread-safe
-	AddToGridEntityQuery.ForEachEntityChunk(EntitySubsystem, Context, [this, &EntitySubsystem](FMassExecutionContext& Context)
+	AddToGridEntityQuery.ForEachEntityChunk(EntitySubsystem, Context, [this, &EntitySubsystem, World = EntitySubsystem.GetWorld()](FMassExecutionContext& Context)
 	{
+		FNavigationObstacleHashGrid2D& HashGrid = Context.GetMutableSubsystemChecked<UMassNavigationSubsystem>(World).GetObstacleGridMutable();
 		const int32 NumEntities = Context.GetNumEntities();
 
 		TConstArrayView<FTransformFragment> LocationList = Context.GetFragmentView<FTransformFragment>();
@@ -220,14 +210,15 @@ void UMassNavigationObstacleGridProcessor::Execute(UMassEntitySubsystem& EntityS
 			ObstacleItem.ItemFlags |= bHasColliderData ? EMassNavigationObstacleFlags::HasColliderData : EMassNavigationObstacleFlags::None;
 			
 			const FBox NewBounds(NewPos - FVector(Radius, Radius, 0.f), NewPos + FVector(Radius, Radius, 0.f));
-			NavigationObstacleCellLocationList[EntityIndex].CellLoc = NavigationSubsystem->GetObstacleGridMutable().Add(ObstacleItem, NewBounds);
+			NavigationObstacleCellLocationList[EntityIndex].CellLoc = HashGrid.Add(ObstacleItem, NewBounds);
 
 			Context.Defer().AddTag<FMassInNavigationObstacleGridTag>(ObstacleItem.Entity);
 		}
 	});
 
-	UpdateGridEntityQuery.ForEachEntityChunk(EntitySubsystem, Context, [this, &EntitySubsystem](FMassExecutionContext& Context)
+	UpdateGridEntityQuery.ForEachEntityChunk(EntitySubsystem, Context, [this, &EntitySubsystem, World = EntitySubsystem.GetWorld()](FMassExecutionContext& Context)
 	{
+		FNavigationObstacleHashGrid2D& HashGrid = Context.GetMutableSubsystemChecked<UMassNavigationSubsystem>(World).GetObstacleGridMutable();
 		const int32 NumEntities = Context.GetNumEntities();
 
 		TConstArrayView<FTransformFragment> LocationList = Context.GetFragmentView<FTransformFragment>();
@@ -245,22 +236,23 @@ void UMassNavigationObstacleGridProcessor::Execute(UMassEntitySubsystem& EntityS
 			ObstacleItem.ItemFlags |= bHasColliderData ? EMassNavigationObstacleFlags::HasColliderData : EMassNavigationObstacleFlags::None;
 
 			const FBox NewBounds(NewPos - FVector(Radius, Radius, 0.f), NewPos + FVector(Radius, Radius, 0.f));
-			NavigationObstacleCellLocationList[EntityIndex].CellLoc = NavigationSubsystem->GetObstacleGridMutable().Move(ObstacleItem, NavigationObstacleCellLocationList[EntityIndex].CellLoc, NewBounds);
+			NavigationObstacleCellLocationList[EntityIndex].CellLoc = HashGrid.Move(ObstacleItem, NavigationObstacleCellLocationList[EntityIndex].CellLoc, NewBounds);
 
 #if WITH_MASSGAMEPLAY_DEBUG && 0
 			const FDebugContext BaseDebugContext(this, LogAvoidance, nullptr, ObstacleItem.Entity);
 			if (DebugIsSelected(ObstacleItem.Entity))
 			{
 				FBox Box = MovementSubsystem->GetGridMutable().CalcCellBounds(AvoidanceObstacleCellLocationList[EntityIndex].CellLoc);
-				Box.Max.Z += 200.f;					
+				Box.Max.Z += 200.f;
 				DebugDrawBox(BaseDebugContext, Box, FColor::Yellow);
 			}
 #endif // WITH_MASSGAMEPLAY_DEBUG
 		}
 	});
 
-	RemoveFromGridEntityQuery.ForEachEntityChunk(EntitySubsystem, Context, [this](FMassExecutionContext& Context)
+	RemoveFromGridEntityQuery.ForEachEntityChunk(EntitySubsystem, Context, [this, &EntitySubsystem, World = EntitySubsystem.GetWorld()](FMassExecutionContext& Context)
 	{
+		FNavigationObstacleHashGrid2D& HashGrid = Context.GetMutableSubsystemChecked<UMassNavigationSubsystem>(World).GetObstacleGridMutable();
 		const int32 NumEntities = Context.GetNumEntities();
 
 		TArrayView<FMassNavigationObstacleGridCellLocationFragment> AvoidanceObstacleCellLocationList = Context.GetMutableFragmentView<FMassNavigationObstacleGridCellLocationFragment>();
@@ -268,7 +260,7 @@ void UMassNavigationObstacleGridProcessor::Execute(UMassEntitySubsystem& EntityS
 		{
 			FMassNavigationObstacleItem ObstacleItem;
 			ObstacleItem.Entity = Context.GetEntity(EntityIndex);
-			NavigationSubsystem->GetObstacleGridMutable().Remove(ObstacleItem, AvoidanceObstacleCellLocationList[EntityIndex].CellLoc);
+			HashGrid.Remove(ObstacleItem, AvoidanceObstacleCellLocationList[EntityIndex].CellLoc);
 			AvoidanceObstacleCellLocationList[EntityIndex].CellLoc = FNavigationObstacleHashGrid2D::FCellLocation();
 
 			Context.Defer().RemoveTag<FMassInNavigationObstacleGridTag>(ObstacleItem.Entity);
@@ -289,23 +281,14 @@ UMassNavigationObstacleRemoverProcessor::UMassNavigationObstacleRemoverProcessor
 void UMassNavigationObstacleRemoverProcessor::ConfigureQueries()
 {
 	EntityQuery.AddRequirement<FMassNavigationObstacleGridCellLocationFragment>(EMassFragmentAccess::ReadWrite);
-}
-
-void UMassNavigationObstacleRemoverProcessor::Initialize(UObject& Owner)
-{
-	Super::Initialize(Owner);
-	NavigationSubsystem = UWorld::GetSubsystem<UMassNavigationSubsystem>(Owner.GetWorld());
+	EntityQuery.AddSystemRequirement<UMassNavigationSubsystem>(EMassFragmentAccess::ReadWrite);
 }
 
 void UMassNavigationObstacleRemoverProcessor::Execute(UMassEntitySubsystem& EntitySubsystem, FMassExecutionContext& Context)
 {
-	if (!NavigationSubsystem)
-	{
-		return;
-	}
-
-	EntityQuery.ForEachEntityChunk(EntitySubsystem, Context, [this](FMassExecutionContext& Context)
+	EntityQuery.ForEachEntityChunk(EntitySubsystem, Context, [this, World = EntitySubsystem.GetWorld()](FMassExecutionContext& Context)
 	{
+		FNavigationObstacleHashGrid2D& HashGrid = Context.GetMutableSubsystemChecked<UMassNavigationSubsystem>(World).GetObstacleGridMutable();
 		const int32 NumEntities = Context.GetNumEntities();
 		const TArrayView<FMassNavigationObstacleGridCellLocationFragment> AvoidanceObstacleCellLocationList = Context.GetMutableFragmentView<FMassNavigationObstacleGridCellLocationFragment>();
 
@@ -313,7 +296,7 @@ void UMassNavigationObstacleRemoverProcessor::Execute(UMassEntitySubsystem& Enti
 		{
 			FMassNavigationObstacleItem ObstacleItem;
 			ObstacleItem.Entity = Context.GetEntity(i);
-			NavigationSubsystem->GetObstacleGridMutable().Remove(ObstacleItem, AvoidanceObstacleCellLocationList[i].CellLoc);
+			HashGrid.Remove(ObstacleItem, AvoidanceObstacleCellLocationList[i].CellLoc);
 		}
 	});
 }
diff --git a/Engine/Plugins/AI/MassAI/Source/MassNavigation/Public/MassNavigationProcessors.h b/Engine/Plugins/AI/MassAI/Source/MassNavigation/Public/MassNavigationProcessors.h
--- a/Engine/Plugins/AI/MassAI/Source/MassNavigation/Public/MassNavigationProcessors.h
+++ b/Engine/Plugins/AI/MassAI/Source/MassNavigation/Public/MassNavigationProcessors.h
@@ -75,11 +75,9 @@ public:
 
 protected:
 	virtual void ConfigureQueries() override;
-	virtual void Initialize(UObject& Owner) override;
 	virtual void Execute(UMassEntitySubsystem& EntitySubsystem, FMassExecutionContext& Context) override;
 
 private:
-	TObjectPtr<UMassNavigationSubsystem> NavigationSubsystem;
 	FMassEntityQuery AddToGridEntityQuery;
 	FMassEntityQuery UpdateGridEntityQuery;
 	FMassEntityQuery RemoveFromGridEntityQuery;
@@ -92,11 +90,9 @@ class MASSNAVIGATION_API UMassNavigationObstacleRemoverProcessor : public UMassO
 	GENERATED_BODY()
 
 	UMassNavigationObstacleRemoverProcessor();
-	TObjectPtr<UMassNavigationSubsystem> NavigationSubsystem;
 
 protected:
 	virtual void ConfigureQueries() override;
-	virtual void Initialize(UObject& Owner) override;
 	virtual void Execute(UMassEntitySubsystem& EntitySubsystem, FMassExecutionContext& Context) override;
 
 	FMassEntityQuery EntityQuery;
diff --git a/Engine/Plugins/AI/MassAI/Source/MassZoneGraphNavigation/Private/MassZoneGraphNavigationProcessors.cpp b/Engine/Plugins/AI/MassAI/Source/MassZoneGraphNavigation/Private/MassZoneGraphNavigationProcessors.cpp
--- a/Engine/Plugins/AI/MassAI/Source/MassZoneGraphNavigation/Private/MassZoneGraphNavigationProcessors.cpp
+++ b/Engine/Plugins/AI/MassAI/Source/MassZoneGraphNavigation/Private/MassZoneGraphNavigationProcessors.cpp
@@ -48,24 +48,15 @@ void UMassZoneGraphLocationInitializer::ConfigureQueries()
 	EntityQuery.AddRequirement<FTransformFragment>(EMassFragmentAccess::ReadOnly);
 	EntityQuery.AddRequirement<FMassMoveTargetFragment>(EMassFragmentAccess::ReadWrite); // Make optional?
 	EntityQuery.AddConstSharedRequirement<FMassZoneGraphNavigationParameters>(EMassFragmentPresence::All);
-}
-
-void UMassZoneGraphLocationInitializer::Initialize(UObject& Owner)
-{
-	Super::Initialize(Owner);
-	ZoneGraphSubsystem = UWorld::GetSubsystem<UZoneGraphSubsystem>(Owner.GetWorld());
-	SignalSubsystem = UWorld::GetSubsystem<UMassSignalSubsystem>(Owner.GetWorld());
+	EntityQuery.AddSystemRequirement<UZoneGraphSubsystem>(EMassFragmentAccess::ReadOnly);
 }
 
 void UMassZoneGraphLocationInitializer::Execute(UMassEntitySubsystem& EntitySubsystem, FMassExecutionContext& Context)
 {
-	if (!ZoneGraphSubsystem || !SignalSubsystem)
+	EntityQuery.ForEachEntityChunk(EntitySubsystem, Context, [this, World = EntitySubsystem.GetWorld()](FMassExecutionContext& Context)
 	{
-		return;
-	}
+		const UZoneGraphSubsystem& ZoneGraphSubsystem = Context.GetSubsystemChecked<UZoneGraphSubsystem>(World);
 
-	EntityQuery.ForEachEntityChunk(EntitySubsystem, Context, [this](FMassExecutionContext& Context)
-	{
 		const int32 NumEntities = Context.GetNumEntities();
 		const TArrayView<FMassZoneGraphLaneLocationFragment> LaneLocationList = Context.GetMutableFragmentView<FMassZoneGraphLaneLocationFragment>();
 		const TArrayView<FMassMoveTargetFragment> MoveTargetList = Context.GetMutableFragmentView<FMassMoveTargetFragment>();
@@ -85,9 +76,9 @@ void UMassZoneGraphLocationInitializer::Execute(UMassEntitySubsystem& EntitySubs
 			FZoneGraphLaneLocation NearestLane;
 			float NearestLaneDistSqr = 0;
 			
-			if (ZoneGraphSubsystem->FindNearestLane(QueryBounds, NavigationParams.LaneFilter, NearestLane, NearestLaneDistSqr))
+			if (ZoneGraphSubsystem.FindNearestLane(QueryBounds, NavigationParams.LaneFilter, NearestLane, NearestLaneDistSqr))
 			{
-				const FZoneGraphStorage* ZoneGraphStorage = ZoneGraphSubsystem->GetZoneGraphStorage(NearestLane.LaneHandle.DataHandle);
+				const FZoneGraphStorage* ZoneGraphStorage = ZoneGraphSubsystem.GetZoneGraphStorage(NearestLane.LaneHandle.DataHandle);
 				check(ZoneGraphStorage); // Assume valid storage since we just got result.
 
 				LaneLocation.LaneHandle = NearestLane.LaneHandle;
@@ -128,7 +119,6 @@ UMassZoneGraphPathFollowProcessor::UMassZoneGraphPathFollowProcessor()
 void UMassZoneGraphPathFollowProcessor::Initialize(UObject& Owner)
 {
 	Super::Initialize(Owner);
-	ZoneGraphSubsystem = UWorld::GetSubsystem<UZoneGraphSubsystem>(Owner.GetWorld());
 	SignalSubsystem = UWorld::GetSubsystem<UMassSignalSubsystem>(Owner.GetWorld());
 }
 
@@ -142,11 +132,13 @@ void UMassZoneGraphPathFollowProcessor::ConfigureQueries()
 
 	EntityQuery_Conditional.AddChunkRequirement<FMassSimulationVariableTickChunkFragment>(EMassFragmentAccess::ReadOnly, EMassFragmentPresence::Optional);
 	EntityQuery_Conditional.SetChunkFilter(&FMassSimulationVariableTickChunkFragment::ShouldTickChunkThisFrame);
+
+	EntityQuery_Conditional.AddSystemRequirement<UZoneGraphSubsystem>(EMassFragmentAccess::ReadOnly);
 }
 
 void UMassZoneGraphPathFollowProcessor::Execute(UMassEntitySubsystem& EntitySubsystem, FMassExecutionContext& Context)
 {
-	if (!SignalSubsystem || !ZoneGraphSubsystem)
+	if (!SignalSubsystem)
 	{
 		return;
 	}
@@ -154,8 +146,10 @@ void UMassZoneGraphPathFollowProcessor::Execute(UMassEntitySubsystem& EntitySubs
 	TArray<FMassEntityHandle> EntitiesToSignalPathDone;
 	TArray<FMassEntityHandle> EntitiesToSignalLaneChanged;
 
-	EntityQuery_Conditional.ForEachEntityChunk(EntitySubsystem, Context, [this, &EntitiesToSignalPathDone, &EntitiesToSignalLaneChanged](FMassExecutionContext& Context)
+	EntityQuery_Conditional.ForEachEntityChunk(EntitySubsystem, Context, [this, &EntitiesToSignalPathDone, &EntitiesToSignalLaneChanged, World = EntitySubsystem.GetWorld()](FMassExecutionContext& Context)
 	{
+		const UZoneGraphSubsystem& ZoneGraphSubsystem = Context.GetSubsystemChecked<UZoneGraphSubsystem>(World);
+
 		const int32 NumEntities = Context.GetNumEntities();
 		const TArrayView<FMassZoneGraphShortPathFragment> ShortPathList = Context.GetMutableFragmentView<FMassZoneGraphShortPathFragment>();
 		const TArrayView<FMassZoneGraphLaneLocationFragment> LaneLocationList = Context.GetMutableFragmentView<FMassZoneGraphLaneLocationFragment>();
@@ -277,7 +271,7 @@ void UMassZoneGraphPathFollowProcessor::Execute(UMassEntitySubsystem& EntitySubs
 						// Check to see if need advance to next lane.
 						if (ShortPath.NextLaneHandle.IsValid())
 						{
-							const FZoneGraphStorage* ZoneGraphStorage = ZoneGraphSubsystem->GetZoneGraphStorage(LaneLocation.LaneHandle.DataHandle);
+							const FZoneGraphStorage* ZoneGraphStorage = ZoneGraphSubsystem.GetZoneGraphStorage(LaneLocation.LaneHandle.DataHandle);
 							if (ZoneGraphStorage != nullptr)
 							{
 								if (ShortPath.NextExitLinkType == EZoneLaneLinkType::Outgoing)
diff --git a/Engine/Plugins/AI/MassAI/Source/MassZoneGraphNavigation/Public/MassZoneGraphNavigationProcessors.h b/Engine/Plugins/AI/MassAI/Source/MassZoneGraphNavigation/Public/MassZoneGraphNavigationProcessors.h
--- a/Engine/Plugins/AI/MassAI/Source/MassZoneGraphNavigation/Public/MassZoneGraphNavigationProcessors.h
+++ b/Engine/Plugins/AI/MassAI/Source/MassZoneGraphNavigation/Public/MassZoneGraphNavigationProcessors.h
@@ -7,8 +7,6 @@
 #include "MassZoneGraphNavigationProcessors.generated.h"
 
 class UMassSignalSubsystem;
-class UZoneGraphSubsystem;
-
 
 
 /**
@@ -23,17 +21,10 @@ public:
 	UMassZoneGraphLocationInitializer();
 
 protected:
-	virtual void Initialize(UObject& Owner) override;
 	virtual void ConfigureQueries() override;
 	virtual void Execute(UMassEntitySubsystem& EntitySubsystem, FMassExecutionContext& Context) override;
 
 	FMassEntityQuery EntityQuery;
-
-	UPROPERTY(Transient)
-	UZoneGraphSubsystem* ZoneGraphSubsystem = nullptr;
-
-	UPROPERTY(Transient)
-	UMassSignalSubsystem* SignalSubsystem = nullptr;
 };
 
 /** 
@@ -53,9 +44,6 @@ protected:
 
 	FMassEntityQuery EntityQuery_Conditional;
 
-	UPROPERTY(Transient)
-	UZoneGraphSubsystem* ZoneGraphSubsystem = nullptr;
-
 	UPROPERTY(Transient)
 	UMassSignalSubsystem* SignalSubsystem = nullptr;
 };
diff --git a/Engine/Plugins/AI/MassCrowd/Source/MassCrowd/Private/MassCrowdNavigationProcessor.cpp b/Engine/Plugins/AI/MassCrowd/Source/MassCrowd/Private/MassCrowdNavigationProcessor.cpp
--- a/Engine/Plugins/AI/MassCrowd/Source/MassCrowd/Private/MassCrowdNavigationProcessor.cpp
+++ b/Engine/Plugins/AI/MassCrowd/Source/MassCrowd/Private/MassCrowdNavigationProcessor.cpp
@@ -26,27 +26,21 @@ void UMassCrowdLaneTrackingSignalProcessor::ConfigureQueries()
 	EntityQuery.AddTagRequirement<FMassCrowdTag>(EMassFragmentPresence::All);
 	EntityQuery.AddRequirement<FMassCrowdLaneTrackingFragment>(EMassFragmentAccess::ReadWrite);
 	EntityQuery.AddRequirement<FMassZoneGraphLaneLocationFragment>(EMassFragmentAccess::ReadOnly);
+	EntityQuery.AddSystemRequirement<UMassCrowdSubsystem>(EMassFragmentAccess::ReadWrite);
 }
 
 void UMassCrowdLaneTrackingSignalProcessor::Initialize(UObject& Owner)
 {
 	Super::Initialize(Owner);
 	
-	MassCrowdSubsystem = UWorld::GetSubsystem<UMassCrowdSubsystem>(Owner.GetWorld());
-	checkf(MassCrowdSubsystem != nullptr, TEXT("UMassCrowdSubsystem is mandatory when using this processor."));
-
 	SubscribeToSignal(UE::Mass::Signals::CurrentLaneChanged);
 }
 
 void UMassCrowdLaneTrackingSignalProcessor::SignalEntities(UMassEntitySubsystem& EntitySubsystem, FMassExecutionContext& Context, FMassSignalNameLookup& EntitySignals)
 {
-	if (!MassCrowdSubsystem)
-	{
-		return;
-	}
-
-	EntityQuery.ForEachEntityChunk(EntitySubsystem, Context, [this](FMassExecutionContext& Context)
+	EntityQuery.ForEachEntityChunk(EntitySubsystem, Context, [this, World = EntitySubsystem.GetWorld()](FMassExecutionContext& Context)
 	{
+		UMassCrowdSubsystem& MassCrowdSubsystem = Context.GetMutableSubsystemChecked<UMassCrowdSubsystem>(World);
 		const int32 NumEntities = Context.GetNumEntities();
 		const TConstArrayView<FMassZoneGraphLaneLocationFragment> LaneLocationList = Context.GetFragmentView<FMassZoneGraphLaneLocationFragment>();
 		const TArrayView<FMassCrowdLaneTrackingFragment> LaneTrackingList = Context.GetMutableFragmentView<FMassCrowdLaneTrackingFragment>();
@@ -57,7 +51,7 @@ void UMassCrowdLaneTrackingSignalProcessor::SignalEntities(UMassEntitySubsystem&
 			FMassCrowdLaneTrackingFragment& LaneTracking = LaneTrackingList[EntityIndex];
 			if (LaneTracking.TrackedLaneHandle != LaneLocation.LaneHandle)
 			{
-				MassCrowdSubsystem->OnEntityLaneChanged(Context.GetEntity(EntityIndex), LaneTracking.TrackedLaneHandle, LaneLocation.LaneHandle);
+				MassCrowdSubsystem.OnEntityLaneChanged(Context.GetEntity(EntityIndex), LaneTracking.TrackedLaneHandle, LaneLocation.LaneHandle);
 				LaneTracking.TrackedLaneHandle = LaneLocation.LaneHandle;
 			}
 		}
@@ -74,23 +68,18 @@ UMassCrowdLaneTrackingDestructor::UMassCrowdLaneTrackingDestructor()
 	Operation = EMassObservedOperation::Remove;
 }
 
-void UMassCrowdLaneTrackingDestructor::Initialize(UObject& Owner)
-{
-	Super::Initialize(Owner);
-	MassCrowdSubsystem = UWorld::GetSubsystem<UMassCrowdSubsystem>(Owner.GetWorld());
-	checkf(MassCrowdSubsystem != nullptr, TEXT("UMassCrowdSubsystem is mandatory when using this processor."));
-}
-
 void UMassCrowdLaneTrackingDestructor::ConfigureQueries()
 {
 	EntityQuery.AddTagRequirement<FMassCrowdTag>(EMassFragmentPresence::All);
 	EntityQuery.AddRequirement<FMassCrowdLaneTrackingFragment>(EMassFragmentAccess::ReadOnly);
+	EntityQuery.AddSystemRequirement<UMassCrowdSubsystem>(EMassFragmentAccess::ReadWrite);
 }
 
 void UMassCrowdLaneTrackingDestructor::Execute(UMassEntitySubsystem& EntitySubsystem, FMassExecutionContext& Context)
 {
-	EntityQuery.ForEachEntityChunk(EntitySubsystem, Context, [this](const FMassExecutionContext& Context)
+	EntityQuery.ForEachEntityChunk(EntitySubsystem, Context, [this, World = EntitySubsystem.GetWorld()](FMassExecutionContext& Context)
 	{
+		UMassCrowdSubsystem& MassCrowdSubsystem = Context.GetMutableSubsystemChecked<UMassCrowdSubsystem>(World);
 		const int32 NumEntities = Context.GetNumEntities();
 		const TConstArrayView<FMassCrowdLaneTrackingFragment> LaneTrackingList = Context.GetFragmentView<FMassCrowdLaneTrackingFragment>();
 
@@ -99,9 +88,8 @@ void UMassCrowdLaneTrackingDestructor::Execute(UMassEntitySubsystem& EntitySubsy
 			const FMassCrowdLaneTrackingFragment& LaneTracking = LaneTrackingList[EntityIndex];
 			if (LaneTracking.TrackedLaneHandle.IsValid())
 			{
-				MassCrowdSubsystem->OnEntityLaneChanged(Context.GetEntity(EntityIndex), LaneTracking.TrackedLaneHandle, FZoneGraphLaneHandle());
+				MassCrowdSubsystem.OnEntityLaneChanged(Context.GetEntity(EntityIndex), LaneTracking.TrackedLaneHandle, FZoneGraphLaneHandle());
 			}
-
 		}
 	});
 }
diff --git a/Engine/Plugins/AI/MassCrowd/Source/MassCrowd/Private/MassCrowdServerRepresentationLODProcessor.cpp b/Engine/Plugins/AI/MassCrowd/Source/MassCrowd/Private/MassCrowdServerRepresentationLODProcessor.cpp
--- a/Engine/Plugins/AI/MassCrowd/Source/MassCrowd/Private/MassCrowdServerRepresentationLODProcessor.cpp
+++ b/Engine/Plugins/AI/MassCrowd/Source/MassCrowd/Private/MassCrowdServerRepresentationLODProcessor.cpp
@@ -47,8 +47,9 @@ void UMassCrowdServerRepresentationLODProcessor::Execute(UMassEntitySubsystem& E
 {
 	TRACE_CPUPROFILER_EVENT_SCOPE(TEXT("CrowdServerRepresentationLOD"))
 
-	check(LODSubsystem);
-	const TArray<FViewerInfo>& Viewers = LODSubsystem->GetViewers();
+	// @todo this access patter will be transformed into "computation tasks" in near future. Do not duplicate.
+	const UMassLODSubsystem& LODSubsystem = Context.GetSubsystemChecked<UMassLODSubsystem>(EntitySubsystem.GetWorld());
+	const TArray<FViewerInfo>& Viewers = LODSubsystem.GetViewers();
 	LODCalculator.PrepareExecution(Viewers);
 	
 	{
diff --git a/Engine/Plugins/AI/MassCrowd/Source/MassCrowd/Public/MassCrowdNavigationProcessor.h b/Engine/Plugins/AI/MassCrowd/Source/MassCrowd/Public/MassCrowdNavigationProcessor.h
--- a/Engine/Plugins/AI/MassCrowd/Source/MassCrowd/Public/MassCrowdNavigationProcessor.h
+++ b/Engine/Plugins/AI/MassCrowd/Source/MassCrowd/Public/MassCrowdNavigationProcessor.h
@@ -22,9 +22,6 @@ protected:
 	virtual void Initialize(UObject& Owner) override;
 	virtual void ConfigureQueries() override;
 	virtual void SignalEntities(UMassEntitySubsystem& EntitySubsystem, FMassExecutionContext& Context, FMassSignalNameLookup& EntitySignals) override;
-
-	UPROPERTY(Transient)
-	UMassCrowdSubsystem* MassCrowdSubsystem;
 };
 
 /** Processors that cleans up the lane tracking on entity destruction. */
@@ -37,13 +34,9 @@ public:
 	UMassCrowdLaneTrackingDestructor();
 
 protected:
-	virtual void Initialize(UObject& Owner) override;
 	virtual void ConfigureQueries() override;
 	virtual void Execute(UMassEntitySubsystem& EntitySubsystem, FMassExecutionContext& Context) override;
 
-	UPROPERTY(Transient)
-	UMassCrowdSubsystem* MassCrowdSubsystem;
-
 	FMassEntityQuery EntityQuery;
 };
 
diff --git a/Engine/Plugins/Runtime/MassGameplay/Source/MassGameplayDebug/Private/AssignDebugVisProcessor.cpp b/Engine/Plugins/Runtime/MassGameplay/Source/MassGameplayDebug/Private/AssignDebugVisProcessor.cpp
--- a/Engine/Plugins/Runtime/MassGameplay/Source/MassGameplayDebug/Private/AssignDebugVisProcessor.cpp
+++ b/Engine/Plugins/Runtime/MassGameplay/Source/MassGameplayDebug/Private/AssignDebugVisProcessor.cpp
@@ -19,24 +19,30 @@ UAssignDebugVisProcessor::UAssignDebugVisProcessor()
 void UAssignDebugVisProcessor::ConfigureQueries()
 {
 	EntityQuery.AddRequirement<FSimDebugVisFragment>(EMassFragmentAccess::ReadWrite);
+	EntityQuery.AddSystemRequirement<UMassDebuggerSubsystem>(EMassFragmentAccess::ReadWrite);
 }
 
 void UAssignDebugVisProcessor::Execute(UMassEntitySubsystem& EntitySubsystem, FMassExecutionContext& Context)
 {
 #if WITH_EDITORONLY_DATA
-	UMassDebugVisualizationComponent* Visualizer = WeakVisualizer.Get();
-	if (!ensure(Visualizer))
-	{
-		return;
-	}
-
 	QUICK_SCOPE_CYCLE_COUNTER(AssignDebugVisProcessor_Execute);
 	
+	UWorld* World = EntitySubsystem.GetWorld();
+	// @todo this code bit is temporary, so is the Visualizer->DirtyVisuals at the end of the function. Will be wrapped in
+	// "executable task" once that's implemented. 
+	UMassDebugVisualizationComponent* Visualizer = nullptr;
+	UMassDebuggerSubsystem& Debugger = Context.GetMutableSubsystemChecked<UMassDebuggerSubsystem>(World);
+	check(Debugger.GetVisualizationComponent());
+	Visualizer = Debugger.GetVisualizationComponent();
 	// note that this function will create the "visual components" only it they're missing or out of sync. 
-	Visualizer->ConditionallyConstructVisualComponent();
+	Debugger.GetVisualizationComponent()->ConditionallyConstructVisualComponent();
 
-	EntityQuery.ForEachEntityChunk(EntitySubsystem, Context, [this, Visualizer](FMassExecutionContext& Context)
+	EntityQuery.ForEachEntityChunk(EntitySubsystem, Context, [this, World](FMassExecutionContext& Context)
 	{
+		UMassDebuggerSubsystem& Debugger = Context.GetMutableSubsystemChecked<UMassDebuggerSubsystem>(World);
+		UMassDebugVisualizationComponent* Visualizer = Debugger.GetVisualizationComponent();
+		check(Visualizer);
+
 		const TArrayView<FSimDebugVisFragment> DebugVisList = Context.GetMutableFragmentView<FSimDebugVisFragment>();
 		for (FSimDebugVisFragment& VisualComp : DebugVisList)
 		{
@@ -47,20 +53,10 @@ void UAssignDebugVisProcessor::Execute(UMassEntitySubsystem& EntitySubsystem, FM
 			}
 		}
 	});
-	Visualizer->DirtyVisuals();
-#endif // WITH_EDITORONLY_DATA
-}
 
-void UAssignDebugVisProcessor::Initialize(UObject& InOwner)
-{
-	Super::Initialize(InOwner);
-	
-#if WITH_EDITORONLY_DATA
-	UMassDebuggerSubsystem* Debugger = UWorld::GetSubsystem<UMassDebuggerSubsystem>(InOwner.GetWorld());
-	if (ensure(Debugger))
+	if (ensure(Visualizer))
 	{
-		WeakVisualizer = Debugger->GetVisualizationComponent();	
-		ensure(WeakVisualizer.Get());
+		Visualizer->DirtyVisuals();
 	}
 #endif // WITH_EDITORONLY_DATA
 }
diff --git a/Engine/Plugins/Runtime/MassGameplay/Source/MassGameplayDebug/Private/DebugVisLocationProcessor.cpp b/Engine/Plugins/Runtime/MassGameplay/Source/MassGameplayDebug/Private/DebugVisLocationProcessor.cpp
--- a/Engine/Plugins/Runtime/MassGameplay/Source/MassGameplayDebug/Private/DebugVisLocationProcessor.cpp
+++ b/Engine/Plugins/Runtime/MassGameplay/Source/MassGameplayDebug/Private/DebugVisLocationProcessor.cpp
@@ -21,23 +21,21 @@ void UDebugVisLocationProcessor::ConfigureQueries()
 	EntityQuery.AddRequirement<FTransformFragment>(EMassFragmentAccess::ReadOnly);
 	EntityQuery.AddRequirement<FSimDebugVisFragment>(EMassFragmentAccess::ReadOnly);
 	EntityQuery.AddTagRequirement<FMassDebuggableTag>(EMassFragmentPresence::All);
+	EntityQuery.AddSystemRequirement<UMassDebuggerSubsystem>(EMassFragmentAccess::ReadWrite);
 }
 
 void UDebugVisLocationProcessor::Execute(UMassEntitySubsystem& EntitySubsystem, FMassExecutionContext& Context)
 {
 #if WITH_EDITORONLY_DATA
-	UMassDebugVisualizationComponent* Visualizer = WeakVisualizer.Get();
-	if (!ensure(Visualizer))
-	{
-		return;
-	}
-
 	QUICK_SCOPE_CYCLE_COUNTER(DebugVisLocationProcessor_Run);
 
-	TArrayView<UHierarchicalInstancedStaticMeshComponent*> VisualDataISMCs = Visualizer->GetVisualDataISMCs();
-	if (VisualDataISMCs.Num() > 0)
+	EntityQuery.ForEachEntityChunk(EntitySubsystem, Context, [this, World = EntitySubsystem.GetWorld()](FMassExecutionContext& Context)
 	{
-		EntityQuery.ForEachEntityChunk(EntitySubsystem, Context, [this, &VisualDataISMCs](const FMassExecutionContext& Context)
+		UMassDebuggerSubsystem& Debugger = Context.GetMutableSubsystemChecked<UMassDebuggerSubsystem>(World);
+		UMassDebugVisualizationComponent* Visualizer = Debugger.GetVisualizationComponent();
+		check(Visualizer);
+		TArrayView<UHierarchicalInstancedStaticMeshComponent*> VisualDataISMCs = Visualizer->GetVisualDataISMCs();
+		if (VisualDataISMCs.Num() > 0)
 		{
 			const int32 NumEntities = Context.GetNumEntities();
 			const TConstArrayView<FTransformFragment> LocationList = Context.GetFragmentView<FTransformFragment>();
@@ -54,27 +52,17 @@ void UDebugVisLocationProcessor::Execute(UMassEntitySubsystem& EntitySubsystem,
 
 				VisualDataISMCs[VisualComp.VisualType]->UpdateInstanceTransform(VisualComp.InstanceIndex, SMTransform, true);
 			}
-		});
-
-		Visualizer->DirtyVisuals();
-	}
-	else 
-	{
-		UE_LOG(LogMassDebug, Log, TEXT("UDebugVisLocationProcessor: Trying to update InstanceStaticMeshes while none created. Check your debug visualization setup"));
-	}
-#endif // WITH_EDITORONLY_DATA
-}
-
-void UDebugVisLocationProcessor::Initialize(UObject& InOwner)
-{
-	Super::Initialize(InOwner);
+		}
+		else
+		{
+			UE_LOG(LogMassDebug, Log, TEXT("UDebugVisLocationProcessor: Trying to update InstanceStaticMeshes while none created. Check your debug visualization setup"));
+		}
+	});
 
-#if WITH_EDITORONLY_DATA
-	UMassDebuggerSubsystem* Debugger = UWorld::GetSubsystem<UMassDebuggerSubsystem>(InOwner.GetWorld());
+	UMassDebuggerSubsystem* Debugger = UWorld::GetSubsystem<UMassDebuggerSubsystem>(EntitySubsystem.GetWorld());
 	if (ensure(Debugger))
 	{
-		WeakVisualizer = Debugger->GetVisualizationComponent();
-		ensure(WeakVisualizer.Get());
+		Debugger->GetVisualizationComponent()->DirtyVisuals();
 	}
 #endif // WITH_EDITORONLY_DATA
 }
@@ -94,20 +82,17 @@ void UMassProcessor_UpdateDebugVis::ConfigureQueries()
 	EntityQuery.AddRequirement<FDataFragment_DebugVis>(EMassFragmentAccess::ReadWrite);
 	EntityQuery.AddRequirement<FAgentRadiusFragment>(EMassFragmentAccess::ReadWrite);
 	EntityQuery.AddTagRequirement<FMassDebuggableTag>(EMassFragmentPresence::All);
+	EntityQuery.AddSystemRequirement<UMassDebuggerSubsystem>(EMassFragmentAccess::ReadWrite);
 }
 
 void UMassProcessor_UpdateDebugVis::Execute(UMassEntitySubsystem& EntitySubsystem, FMassExecutionContext& Context)
 {
-	UMassDebuggerSubsystem* Debugger = UWorld::GetSubsystem<UMassDebuggerSubsystem>(GetWorld());
-	if (Debugger == nullptr)
-	{
-		return;
-	}
-
 	QUICK_SCOPE_CYCLE_COUNTER(UMassProcessor_UpdateDebugVis_Run);
 
-	EntityQuery.ForEachEntityChunk(EntitySubsystem, Context, [this, Debugger](FMassExecutionContext& Context)
+	EntityQuery.ForEachEntityChunk(EntitySubsystem, Context, [this, World = EntitySubsystem.GetWorld()](FMassExecutionContext& Context)
 		{
+			UMassDebuggerSubsystem& Debugger = Context.GetMutableSubsystemChecked<UMassDebuggerSubsystem>(World);
+
 			const int32 NumEntities = Context.GetNumEntities();
 			const TConstArrayView<FTransformFragment> LocationList = Context.GetFragmentView<FTransformFragment>();
 			const TArrayView<FDataFragment_DebugVis> DebugVisList = Context.GetMutableFragmentView<FDataFragment_DebugVis>();
@@ -115,8 +100,8 @@ void UMassProcessor_UpdateDebugVis::Execute(UMassEntitySubsystem& EntitySubsyste
 
 			for (int32 i = 0; i < NumEntities; ++i)
 			{
-				Debugger->AddShape(DebugVisList[i].Shape, LocationList[i].GetTransform().GetLocation(), RadiiList[i].Radius);
-				Debugger->AddEntityLocation(Context.GetEntity(i), LocationList[i].GetTransform().GetLocation());
+				Debugger.AddShape(DebugVisList[i].Shape, LocationList[i].GetTransform().GetLocation(), RadiiList[i].Radius);
+				Debugger.AddEntityLocation(Context.GetEntity(i), LocationList[i].GetTransform().GetLocation());
 			}
 		});
 }
diff --git a/Engine/Plugins/Runtime/MassGameplay/Source/MassGameplayDebug/Public/AssignDebugVisProcessor.h b/Engine/Plugins/Runtime/MassGameplay/Source/MassGameplayDebug/Public/AssignDebugVisProcessor.h
--- a/Engine/Plugins/Runtime/MassGameplay/Source/MassGameplayDebug/Public/AssignDebugVisProcessor.h
+++ b/Engine/Plugins/Runtime/MassGameplay/Source/MassGameplayDebug/Public/AssignDebugVisProcessor.h
@@ -17,10 +17,7 @@ public:
 	UAssignDebugVisProcessor();
 	virtual void ConfigureQueries() override;
 	virtual void Execute(UMassEntitySubsystem& EntitySubsystem, FMassExecutionContext& Context) override;
-	virtual void Initialize(UObject& InOwner) override;
 
 protected:
-	TWeakObjectPtr<UMassDebugVisualizationComponent> WeakVisualizer;
-
 	FMassEntityQuery EntityQuery;
 };
diff --git a/Engine/Plugins/Runtime/MassGameplay/Source/MassGameplayDebug/Public/DebugVisLocationProcessor.h b/Engine/Plugins/Runtime/MassGameplay/Source/MassGameplayDebug/Public/DebugVisLocationProcessor.h
--- a/Engine/Plugins/Runtime/MassGameplay/Source/MassGameplayDebug/Public/DebugVisLocationProcessor.h
+++ b/Engine/Plugins/Runtime/MassGameplay/Source/MassGameplayDebug/Public/DebugVisLocationProcessor.h
@@ -19,9 +19,7 @@ public:
 protected:
 	virtual void ConfigureQueries() override;
 	virtual void Execute(UMassEntitySubsystem& EntitySubsystem, FMassExecutionContext& Context) override;
-	virtual void Initialize(UObject& InOwner) override;
 
-	TWeakObjectPtr<UMassDebugVisualizationComponent> WeakVisualizer;
 	FMassEntityQuery EntityQuery;
 };
 
diff --git a/Engine/Plugins/Runtime/MassGameplay/Source/MassLOD/Private/MassLODCollectorProcessor.cpp b/Engine/Plugins/Runtime/MassGameplay/Source/MassLOD/Private/MassLODCollectorProcessor.cpp
--- a/Engine/Plugins/Runtime/MassGameplay/Source/MassLOD/Private/MassLODCollectorProcessor.cpp
+++ b/Engine/Plugins/Runtime/MassGameplay/Source/MassLOD/Private/MassLODCollectorProcessor.cpp
@@ -74,8 +74,9 @@ void UMassLODCollectorProcessor::ExecuteInternal(UMassEntitySubsystem& EntitySub
 
 void UMassLODCollectorProcessor::Execute(UMassEntitySubsystem& EntitySubsystem, FMassExecutionContext& Context)
 {
-	check(LODSubsystem);
-	const TArray<FViewerInfo>& Viewers = LODSubsystem->GetViewers();
+	// @todo this access patter will be transformed into "computation tasks" in near future. Do not duplicate.
+	const UMassLODSubsystem& LODSubsystem = Context.GetSubsystemChecked<UMassLODSubsystem>(EntitySubsystem.GetWorld());
+	const TArray<FViewerInfo>& Viewers = LODSubsystem.GetViewers();
 	Collector.PrepareExecution(Viewers);
 
 	check(World);
diff --git a/Engine/Plugins/Runtime/MassGameplay/Source/MassLOD/Private/MassLODSubsystem.cpp b/Engine/Plugins/Runtime/MassGameplay/Source/MassLOD/Private/MassLODSubsystem.cpp
--- a/Engine/Plugins/Runtime/MassGameplay/Source/MassLOD/Private/MassLODSubsystem.cpp
+++ b/Engine/Plugins/Runtime/MassGameplay/Source/MassLOD/Private/MassLODSubsystem.cpp
@@ -26,7 +26,6 @@ void UMassLODProcessorBase::Initialize(UObject& Owner)
 	Super::Initialize(Owner);
 
 	World = Owner.GetWorld();
-	LODSubsystem = UWorld::GetSubsystem<UMassLODSubsystem>(World);
 }
 
 void UMassLODSubsystem::Initialize(FSubsystemCollectionBase& Collection)
diff --git a/Engine/Plugins/Runtime/MassGameplay/Source/MassLOD/Private/MassSimulationLOD.cpp b/Engine/Plugins/Runtime/MassGameplay/Source/MassLOD/Private/MassSimulationLOD.cpp
--- a/Engine/Plugins/Runtime/MassGameplay/Source/MassLOD/Private/MassSimulationLOD.cpp
+++ b/Engine/Plugins/Runtime/MassGameplay/Source/MassLOD/Private/MassSimulationLOD.cpp
@@ -112,8 +112,9 @@ void UMassSimulationLODProcessor::Execute(UMassEntitySubsystem& EntitySubsystem,
 
 	{
 		TRACE_CPUPROFILER_EVENT_SCOPE(TEXT("PrepareExecution"));
-		check(LODSubsystem);
-		const TArray<FViewerInfo>& Viewers = LODSubsystem->GetViewers();
+		// @todo this access patter will be transformed into "computation tasks" in near future. Do not duplicate.
+		const UMassLODSubsystem& LODSubsystem = Context.GetSubsystemChecked<UMassLODSubsystem>(EntitySubsystem.GetWorld());
+		const TArray<FViewerInfo>& Viewers = LODSubsystem.GetViewers();
 
 		EntitySubsystem.ForEachSharedFragment<FMassSimulationLODSharedFragment>([&Viewers](FMassSimulationLODSharedFragment& LODSharedFragment)
 		{
diff --git a/Engine/Plugins/Runtime/MassGameplay/Source/MassLOD/Public/MassLODSubsystem.h b/Engine/Plugins/Runtime/MassGameplay/Source/MassLOD/Public/MassLODSubsystem.h
--- a/Engine/Plugins/Runtime/MassGameplay/Source/MassLOD/Public/MassLODSubsystem.h
+++ b/Engine/Plugins/Runtime/MassGameplay/Source/MassLOD/Public/MassLODSubsystem.h
@@ -26,9 +26,6 @@ public:
 protected:
 	UPROPERTY()
 	UWorld* World = nullptr;
-
-	UPROPERTY()
-	UMassLODSubsystem* LODSubsystem = nullptr;
 };
 
 /*
diff --git a/Engine/Plugins/Runtime/MassGameplay/Source/MassReplication/Private/MassReplicationFragments.cpp b/Engine/Plugins/Runtime/MassGameplay/Source/MassReplication/Private/MassReplicationFragments.cpp
--- a/Engine/Plugins/Runtime/MassGameplay/Source/MassReplication/Private/MassReplicationFragments.cpp
+++ b/Engine/Plugins/Runtime/MassGameplay/Source/MassReplication/Private/MassReplicationFragments.cpp
@@ -18,14 +18,7 @@ UMassNetworkIDFragmentInitializer::UMassNetworkIDFragmentInitializer()
 void UMassNetworkIDFragmentInitializer::ConfigureQueries()
 {
 	EntityQuery.AddRequirement<FMassNetworkIDFragment>(EMassFragmentAccess::ReadWrite);
-}
-
-void UMassNetworkIDFragmentInitializer::Initialize(UObject& Owner)
-{
-	Super::Initialize(Owner);
-#if UE_REPLICATION_COMPILE_SERVER_CODE
-	ReplicationSubsystem = UWorld::GetSubsystem<UMassReplicationSubsystem>(Owner.GetWorld());
-#endif //UE_REPLICATION_COMPILE_SERVER_CODE
+	EntityQuery.AddSystemRequirement<UMassReplicationSubsystem>(EMassFragmentAccess::ReadWrite);
 }
 
 void UMassNetworkIDFragmentInitializer::Execute(UMassEntitySubsystem& EntitySubsystem, FMassExecutionContext& Context)
@@ -38,16 +31,16 @@ void UMassNetworkIDFragmentInitializer::Execute(UMassEntitySubsystem& EntitySubs
 	if (NetMode != NM_Client)
 	{
 #if UE_REPLICATION_COMPILE_SERVER_CODE
-		check(ReplicationSubsystem);
-
-		EntityQuery.ForEachEntityChunk(EntitySubsystem, Context, [this](FMassExecutionContext& Context)
+		EntityQuery.ForEachEntityChunk(EntitySubsystem, Context, [World = EntitySubsystem.GetWorld()](FMassExecutionContext& Context)
 			{
+				UMassReplicationSubsystem& ReplicationSubsystem = Context.GetMutableSubsystemChecked<UMassReplicationSubsystem>(World);
+
 				const TArrayView<FMassNetworkIDFragment> NetworkIDList = Context.GetMutableFragmentView<FMassNetworkIDFragment>();
 				const int32 NumEntities = Context.GetNumEntities();
 
 				for (int32 Idx = 0; Idx < NumEntities; ++Idx)
 				{
-					NetworkIDList[Idx].NetID = ReplicationSubsystem->GetNextAvailableMassNetID();
+					NetworkIDList[Idx].NetID = ReplicationSubsystem.GetNextAvailableMassNetID();
 				}
 			});
 #endif //UE_REPLICATION_COMPILE_SERVER_CODE
diff --git a/Engine/Plugins/Runtime/MassGameplay/Source/MassReplication/Private/MassReplicationGridProcessor.cpp b/Engine/Plugins/Runtime/MassGameplay/Source/MassReplication/Private/MassReplicationGridProcessor.cpp
--- a/Engine/Plugins/Runtime/MassGameplay/Source/MassReplication/Private/MassReplicationGridProcessor.cpp
+++ b/Engine/Plugins/Runtime/MassGameplay/Source/MassReplication/Private/MassReplicationGridProcessor.cpp
@@ -26,6 +26,8 @@ void UMassReplicationGridProcessor::ConfigureQueries()
 	AddToGridEntityQuery.AddRequirement<FTransformFragment>(EMassFragmentAccess::ReadOnly);
 	AddToGridEntityQuery.AddRequirement<FAgentRadiusFragment>(EMassFragmentAccess::ReadOnly);
 	AddToGridEntityQuery.AddRequirement<FMassReplicationGridCellLocationFragment>(EMassFragmentAccess::ReadWrite);
+	AddToGridEntityQuery.AddSystemRequirement<UMassReplicationSubsystem>(EMassFragmentAccess::ReadWrite);
+
 	UpdateGridEntityQuery = AddToGridEntityQuery;
 	RemoveFromGridEntityQuery = AddToGridEntityQuery;
 
@@ -39,21 +41,14 @@ void UMassReplicationGridProcessor::ConfigureQueries()
 	RemoveFromGridEntityQuery.AddTagRequirement<FMassInReplicationGridTag>(EMassFragmentPresence::All);
 }
 
-void UMassReplicationGridProcessor::Initialize(UObject& Owner)
-{
-	Super::Initialize(Owner);
-
-	ReplicationSubsystem = UWorld::GetSubsystem<UMassReplicationSubsystem>(Owner.GetWorld());
-}
-
 void UMassReplicationGridProcessor::Execute(UMassEntitySubsystem& EntitySubsystem, FMassExecutionContext& Context)
 {
-	check(ReplicationSubsystem);
+	UWorld* World = EntitySubsystem.GetWorld();
 
-	// can't be ParallelFor due to ReplicationSubsystem->GetGridMutable().Move not being thread-safe
-	FReplicationHashGrid2D& ReplicationGrid = ReplicationSubsystem->GetGridMutable();
-	AddToGridEntityQuery.ForEachEntityChunk(EntitySubsystem, Context, [&ReplicationGrid, &EntitySubsystem](FMassExecutionContext& Context)
+	AddToGridEntityQuery.ForEachEntityChunk(EntitySubsystem, Context, [World](FMassExecutionContext& Context)
 	{
+		UMassReplicationSubsystem& ReplicationSubsystem = Context.GetMutableSubsystemChecked<UMassReplicationSubsystem>(World);
+		FReplicationHashGrid2D& ReplicationGrid = ReplicationSubsystem.GetGridMutable();
 		const int32 NumEntities = Context.GetNumEntities();
 
 		TConstArrayView<FTransformFragment> LocationList = Context.GetFragmentView<FTransformFragment>();
@@ -74,8 +69,10 @@ void UMassReplicationGridProcessor::Execute(UMassEntitySubsystem& EntitySubsyste
 		}
 	});
 
-	UpdateGridEntityQuery.ForEachEntityChunk(EntitySubsystem, Context, [&ReplicationGrid, &EntitySubsystem](FMassExecutionContext& Context)
+	UpdateGridEntityQuery.ForEachEntityChunk(EntitySubsystem, Context, [World](FMassExecutionContext& Context)
 	{
+		UMassReplicationSubsystem& ReplicationSubsystem = Context.GetMutableSubsystemChecked<UMassReplicationSubsystem>(World);
+		FReplicationHashGrid2D& ReplicationGrid = ReplicationSubsystem.GetGridMutable();
 		const int32 NumEntities = Context.GetNumEntities();
 
 		TConstArrayView<FTransformFragment> LocationList = Context.GetFragmentView<FTransformFragment>();
@@ -103,8 +100,10 @@ void UMassReplicationGridProcessor::Execute(UMassEntitySubsystem& EntitySubsyste
 		}
 	});
 
-	RemoveFromGridEntityQuery.ForEachEntityChunk(EntitySubsystem, Context, [&ReplicationGrid](FMassExecutionContext& Context)
+	RemoveFromGridEntityQuery.ForEachEntityChunk(EntitySubsystem, Context, [World](FMassExecutionContext& Context)
 	{
+		UMassReplicationSubsystem& ReplicationSubsystem = Context.GetMutableSubsystemChecked<UMassReplicationSubsystem>(World);
+		FReplicationHashGrid2D& ReplicationGrid = ReplicationSubsystem.GetGridMutable();
 		const int32 NumEntities = Context.GetNumEntities();
 
 		TArrayView<FMassReplicationGridCellLocationFragment> ReplicationCellLocationList = Context.GetMutableFragmentView<FMassReplicationGridCellLocationFragment>();
@@ -132,24 +131,15 @@ UMassReplicationGridRemoverProcessor::UMassReplicationGridRemoverProcessor()
 void UMassReplicationGridRemoverProcessor::ConfigureQueries()
 {
 	EntityQuery.AddRequirement<FMassReplicationGridCellLocationFragment>(EMassFragmentAccess::ReadWrite);
-}
-
-void UMassReplicationGridRemoverProcessor::Initialize(UObject& Owner)
-{
-	Super::Initialize(Owner);
-	ReplicationSubsystem = UWorld::GetSubsystem<UMassReplicationSubsystem>(Owner.GetWorld());
+	EntityQuery.AddSystemRequirement<UMassReplicationSubsystem>(EMassFragmentAccess::ReadWrite);
 }
 
 void UMassReplicationGridRemoverProcessor::Execute(UMassEntitySubsystem& EntitySubsystem, FMassExecutionContext& Context)
 {
-	if (!ReplicationSubsystem)
-	{
-		return;
-	}
-
-	FReplicationHashGrid2D& ReplicationGrid = ReplicationSubsystem->GetGridMutable();
-	EntityQuery.ForEachEntityChunk(EntitySubsystem, Context, [&ReplicationGrid](FMassExecutionContext& Context)
+	EntityQuery.ForEachEntityChunk(EntitySubsystem, Context, [World = EntitySubsystem.GetWorld()](FMassExecutionContext& Context)
 	{
+		UMassReplicationSubsystem& ReplicationSubsystem = Context.GetMutableSubsystemChecked<UMassReplicationSubsystem>(World);
+		FReplicationHashGrid2D& ReplicationGrid = ReplicationSubsystem.GetGridMutable();
 		const int32 NumEntities = Context.GetNumEntities();
 		const TArrayView<FMassReplicationGridCellLocationFragment> ReplicationCellLocationList = Context.GetMutableFragmentView<FMassReplicationGridCellLocationFragment>();
 
diff --git a/Engine/Plugins/Runtime/MassGameplay/Source/MassReplication/Private/MassReplicationProcessor.cpp b/Engine/Plugins/Runtime/MassGameplay/Source/MassReplication/Private/MassReplicationProcessor.cpp
--- a/Engine/Plugins/Runtime/MassGameplay/Source/MassReplication/Private/MassReplicationProcessor.cpp
+++ b/Engine/Plugins/Runtime/MassGameplay/Source/MassReplication/Private/MassReplicationProcessor.cpp
@@ -76,7 +76,6 @@ void UMassReplicationProcessor::PrepareExecution(UMassEntitySubsystem& EntitySub
 	//first synchronize clients and viewers
 	ReplicationSubsystem->SynchronizeClientsAndViewers();
 
-	check(LODSubsystem);
 	EntitySubsystem.ForEachSharedFragment<FMassReplicationSharedFragment>([this](FMassReplicationSharedFragment& RepSharedFragment)
 	{
 		if (!RepSharedFragment.bEntityQueryInitialized)
@@ -151,7 +150,6 @@ void UMassReplicationProcessor::Execute(UMassEntitySubsystem& EntitySubsystem, F
 {
 #if UE_REPLICATION_COMPILE_SERVER_CODE
 	check(World);
-	check(LODSubsystem);
 	check(ReplicationSubsystem);
 
 	{
@@ -159,8 +157,8 @@ void UMassReplicationProcessor::Execute(UMassEntitySubsystem& EntitySubsystem, F
 		PrepareExecution(EntitySubsystem);
 	}
 
-
-	const TArray<FViewerInfo>& AllViewersInfo = LODSubsystem->GetViewers();
+	UMassLODSubsystem& LODSubsystem = Context.GetMutableSubsystemChecked<UMassLODSubsystem>(EntitySubsystem.GetWorld());
+	const TArray<FViewerInfo>& AllViewersInfo = LODSubsystem.GetViewers();
 	const TArray<FMassClientHandle>& ClientHandles = ReplicationSubsystem->GetClientReplicationHandles();
 	for( const FMassClientHandle ClientHandle : ClientHandles )
 	{
@@ -328,7 +326,7 @@ void UMassReplicationProcessor::Execute(UMassEntitySubsystem& EntitySubsystem, F
 
 				{
 					QUICK_SCOPE_CYCLE_COUNTER(UMassReplicationProcessor_ProcessClientReplication);
-					FMassReplicationContext ReplicationContext(*World, *LODSubsystem, *ReplicationSubsystem);
+					FMassReplicationContext ReplicationContext(*World, LODSubsystem, *ReplicationSubsystem);
 					EntitySubsystem.ForEachSharedFragment<FMassReplicationSharedFragment>([&EntitySubsystem, &Context, &ReplicationContext, &ClientHandle](FMassReplicationSharedFragment& RepSharedFragment)
 					{
 						RepSharedFragment.CurrentClientHandle = ClientHandle;
diff --git a/Engine/Plugins/Runtime/MassGameplay/Source/MassReplication/Public/MassReplicationFragments.h b/Engine/Plugins/Runtime/MassGameplay/Source/MassReplication/Public/MassReplicationFragments.h
--- a/Engine/Plugins/Runtime/MassGameplay/Source/MassReplication/Public/MassReplicationFragments.h
+++ b/Engine/Plugins/Runtime/MassGameplay/Source/MassReplication/Public/MassReplicationFragments.h
@@ -70,12 +70,8 @@ public:
 
 protected:
 	virtual void ConfigureQueries() override;
-	virtual void Initialize(UObject& Owner) override;
 	virtual void Execute(UMassEntitySubsystem& EntitySubsystem, FMassExecutionContext& Context) override;
 
-	UPROPERTY()
-	UMassReplicationSubsystem* ReplicationSubsystem = nullptr;
-
 	FMassEntityQuery EntityQuery;
 };
 
diff --git a/Engine/Plugins/Runtime/MassGameplay/Source/MassReplication/Public/MassReplicationGridProcessor.h b/Engine/Plugins/Runtime/MassGameplay/Source/MassReplication/Public/MassReplicationGridProcessor.h
--- a/Engine/Plugins/Runtime/MassGameplay/Source/MassReplication/Public/MassReplicationGridProcessor.h
+++ b/Engine/Plugins/Runtime/MassGameplay/Source/MassReplication/Public/MassReplicationGridProcessor.h
@@ -21,13 +21,9 @@ public:
 
 protected:
 	virtual void ConfigureQueries() override;
-	virtual void Initialize(UObject& Owner) override;
 	virtual void Execute(UMassEntitySubsystem& EntitySubsystem, FMassExecutionContext& Context) override;
 
 private:
-	UPROPERTY(Transient)
-	UMassReplicationSubsystem* ReplicationSubsystem;
-
 	FMassEntityQuery AddToGridEntityQuery;
 	FMassEntityQuery UpdateGridEntityQuery;
 	FMassEntityQuery RemoveFromGridEntityQuery;
@@ -43,11 +39,7 @@ class MASSREPLICATION_API UMassReplicationGridRemoverProcessor : public UMassObs
 
 protected:
 	virtual void ConfigureQueries() override;
-	virtual void Initialize(UObject& Owner) override;
 	virtual void Execute(UMassEntitySubsystem& EntitySubsystem, FMassExecutionContext& Context) override;
 
 	FMassEntityQuery EntityQuery;
-
-	UPROPERTY(Transient)
-	UMassReplicationSubsystem* ReplicationSubsystem;
 };
\ No newline at end of file
diff --git a/Engine/Plugins/Runtime/MassGameplay/Source/MassRepresentation/Private/MassRepresentationProcessor.cpp b/Engine/Plugins/Runtime/MassGameplay/Source/MassRepresentation/Private/MassRepresentationProcessor.cpp
--- a/Engine/Plugins/Runtime/MassGameplay/Source/MassRepresentation/Private/MassRepresentationProcessor.cpp
+++ b/Engine/Plugins/Runtime/MassGameplay/Source/MassRepresentation/Private/MassRepresentationProcessor.cpp
@@ -35,6 +35,7 @@ void UMassRepresentationProcessor::ConfigureQueries()
 	EntityQuery.AddRequirement<FMassActorFragment>(EMassFragmentAccess::ReadWrite);
 	EntityQuery.AddConstSharedRequirement<FMassRepresentationParameters>();
 	EntityQuery.AddSharedRequirement<FMassRepresentationSubsystemSharedFragment>(EMassFragmentAccess::ReadWrite);
+	EntityQuery.AddSystemRequirement<UMassActorSubsystem>(EMassFragmentAccess::ReadWrite);
 }
 
 void UMassRepresentationProcessor::Initialize(UObject& Owner)
@@ -53,6 +54,8 @@ void UMassRepresentationProcessor::UpdateRepresentation(FMassExecutionContext& C
 	UMassRepresentationActorManagement* RepresentationActorManagement = RepresentationParams.CachedRepresentationActorManagement;
 	check(RepresentationActorManagement);
 
+	UMassActorSubsystem* MassActorSubsystem = Context.GetMutableSubsystem<UMassActorSubsystem>(RepresentationSubsystem->GetWorld());
+
 	const TConstArrayView<FTransformFragment> TransformList = Context.GetFragmentView<FTransformFragment>();
 	const TArrayView<FMassRepresentationFragment> RepresentationList = Context.GetMutableFragmentView<FMassRepresentationFragment>();
 	const TConstArrayView<FMassRepresentationLODFragment> RepresentationLODList = Context.GetFragmentView<FMassRepresentationLODFragment>();
@@ -91,12 +94,12 @@ void UMassRepresentationProcessor::UpdateRepresentation(FMassExecutionContext& C
 				if (Representation.HighResTemplateActorIndex != Representation.LowResTemplateActorIndex || !RepresentationParams.bKeepLowResActors)
 				{
 					// Try releasing the high actor or any high res spawning request
-					if (ReleaseActorOrCancelSpawning(*RepresentationSubsystem, MassAgent, ActorInfo, Representation.HighResTemplateActorIndex, Representation.ActorSpawnRequestHandle, Context.Defer()))
+					if (ReleaseActorOrCancelSpawning(*RepresentationSubsystem, MassActorSubsystem, MassAgent, ActorInfo, Representation.HighResTemplateActorIndex, Representation.ActorSpawnRequestHandle, Context.Defer()))
 					{
 						Actor = ActorInfo.GetOwnedByMassMutable();
 					}
 					// Do not do the same with low res if indicated so
-					if (!RepresentationParams.bKeepLowResActors && ReleaseActorOrCancelSpawning(*RepresentationSubsystem, MassAgent, ActorInfo, Representation.LowResTemplateActorIndex, Representation.ActorSpawnRequestHandle, Context.Defer()))
+					if (!RepresentationParams.bKeepLowResActors && ReleaseActorOrCancelSpawning(*RepresentationSubsystem, MassActorSubsystem, MassAgent, ActorInfo, Representation.LowResTemplateActorIndex, Representation.ActorSpawnRequestHandle, Context.Defer()))
 					{
 						Actor = ActorInfo.GetOwnedByMassMutable();
 					}
@@ -139,7 +142,7 @@ void UMassRepresentationProcessor::UpdateRepresentation(FMassExecutionContext& C
 						// If the low res is different than the high res, cancel any pending spawn request that is the opposite of what is needed.
 						if (Representation.LowResTemplateActorIndex != Representation.HighResTemplateActorIndex)
 						{
-							ReleaseActorOrCancelSpawning(*RepresentationSubsystem, MassAgent, ActorInfo, bHighResActor ? Representation.LowResTemplateActorIndex : Representation.HighResTemplateActorIndex, Representation.ActorSpawnRequestHandle, Context.Defer(), /*bCancelSpawningOnly*/true);
+							ReleaseActorOrCancelSpawning(*RepresentationSubsystem, MassActorSubsystem, MassAgent, ActorInfo, bHighResActor ? Representation.LowResTemplateActorIndex : Representation.HighResTemplateActorIndex, Representation.ActorSpawnRequestHandle, Context.Defer(), /*bCancelSpawningOnly*/true);
 							Actor = ActorInfo.GetOwnedByMassMutable();
 						}
 
@@ -194,8 +197,8 @@ void UMassRepresentationProcessor::UpdateRepresentation(FMassExecutionContext& C
 					if (!Actor || ActorInfo.IsOwnedByMass())
 					{
 						// Try releasing both, could have an high res spawned actor and a spawning request for a low res one
-						ReleaseActorOrCancelSpawning(*RepresentationSubsystem, MassAgent, ActorInfo, Representation.LowResTemplateActorIndex, Representation.ActorSpawnRequestHandle, Context.Defer());
-						ReleaseActorOrCancelSpawning(*RepresentationSubsystem, MassAgent, ActorInfo, Representation.HighResTemplateActorIndex, Representation.ActorSpawnRequestHandle, Context.Defer());
+						ReleaseActorOrCancelSpawning(*RepresentationSubsystem, MassActorSubsystem, MassAgent, ActorInfo, Representation.LowResTemplateActorIndex, Representation.ActorSpawnRequestHandle, Context.Defer());
+						ReleaseActorOrCancelSpawning(*RepresentationSubsystem, MassActorSubsystem, MassAgent, ActorInfo, Representation.HighResTemplateActorIndex, Representation.ActorSpawnRequestHandle, Context.Defer());
 					}
 					else
 					{
@@ -227,7 +230,7 @@ void UMassRepresentationProcessor::Execute(UMassEntitySubsystem& InEntitySubsyst
 	});
 }
 
-bool UMassRepresentationProcessor::ReleaseActorOrCancelSpawning(UMassRepresentationSubsystem& RepresentationSubsystem, const FMassEntityHandle MassAgent, FMassActorFragment& ActorInfo, const int16 TemplateActorIndex, FMassActorSpawnRequestHandle& SpawnRequestHandle, FMassCommandBuffer& CommandBuffer, bool bCancelSpawningOnly /*= false*/)
+bool UMassRepresentationProcessor::ReleaseActorOrCancelSpawning(UMassRepresentationSubsystem& RepresentationSubsystem, UMassActorSubsystem* MassActorSubsystem, const FMassEntityHandle MassAgent, FMassActorFragment& ActorInfo, const int16 TemplateActorIndex, FMassActorSpawnRequestHandle& SpawnRequestHandle, FMassCommandBuffer& CommandBuffer, bool bCancelSpawningOnly /*= false*/)
 {
 	if (TemplateActorIndex == INDEX_NONE)
 	{
@@ -249,7 +252,7 @@ bool UMassRepresentationProcessor::ReleaseActorOrCancelSpawning(UMassRepresentat
 			ActorInfo.ResetNoHandleMapUpdate();
 			
 			TObjectKey<const AActor> ActorKey(Actor); 
-			if (UMassActorSubsystem* MassActorSubsystem = UWorld::GetSubsystem<UMassActorSubsystem>(World))
+			if (MassActorSubsystem)
 			{
 				CommandBuffer.PushCommand<FMassDeferredSetCommand>([MassActorSubsystem, ActorKey](UMassEntitySubsystem&)
 				{
diff --git a/Engine/Plugins/Runtime/MassGameplay/Source/MassRepresentation/Private/MassVisualizationLODProcessor.cpp b/Engine/Plugins/Runtime/MassGameplay/Source/MassRepresentation/Private/MassVisualizationLODProcessor.cpp
--- a/Engine/Plugins/Runtime/MassGameplay/Source/MassRepresentation/Private/MassVisualizationLODProcessor.cpp
+++ b/Engine/Plugins/Runtime/MassGameplay/Source/MassRepresentation/Private/MassVisualizationLODProcessor.cpp
@@ -59,8 +59,9 @@ void UMassVisualizationLODProcessor::Execute(UMassEntitySubsystem& EntitySubsyst
 
 	{
 		TRACE_CPUPROFILER_EVENT_SCOPE(PrepareExecution)
-		check(LODSubsystem);
-		const TArray<FViewerInfo>& Viewers = LODSubsystem->GetViewers();
+		// @todo this access patter will be transformed into "computation tasks" in near future. Do not duplicate.
+		const UMassLODSubsystem& LODSubsystem = Context.GetSubsystemChecked<UMassLODSubsystem>(EntitySubsystem.GetWorld());
+		const TArray<FViewerInfo>& Viewers = LODSubsystem.GetViewers();
 		EntitySubsystem.ForEachSharedFragment<FMassVisualizationLODSharedFragment>([this, &Viewers](FMassVisualizationLODSharedFragment& LODSharedFragment)
 		{
 			if (FilterTag == LODSharedFragment.FilterTag)
diff --git a/Engine/Plugins/Runtime/MassGameplay/Source/MassRepresentation/Public/MassRepresentationProcessor.h b/Engine/Plugins/Runtime/MassGameplay/Source/MassRepresentation/Public/MassRepresentationProcessor.h
--- a/Engine/Plugins/Runtime/MassGameplay/Source/MassRepresentation/Public/MassRepresentationProcessor.h
+++ b/Engine/Plugins/Runtime/MassGameplay/Source/MassRepresentation/Public/MassRepresentationProcessor.h
@@ -10,6 +10,7 @@
 #include "MassRepresentationProcessor.generated.h"
 
 class UMassRepresentationSubsystem;
+class UMassActorSubsystem;
 struct FMassActorFragment;
 
 UCLASS()
@@ -49,7 +50,7 @@ protected:
 	 * @param bCancelSpawningOnly tell to only cancel the existing spawning request and to not release the associated actor it any.
 	 * @return if the actor was release or the spawning was canceled.
 	 */
-	bool ReleaseActorOrCancelSpawning(UMassRepresentationSubsystem& RepresentationSubsystem, const FMassEntityHandle MassAgent, FMassActorFragment& ActorInfo, const int16 TemplateActorIndex, FMassActorSpawnRequestHandle& SpawnRequestHandle, FMassCommandBuffer& CommandBuffer, const bool bCancelSpawningOnly = false);
+	bool ReleaseActorOrCancelSpawning(UMassRepresentationSubsystem& RepresentationSubsystem, UMassActorSubsystem* MassActorSubsystem, const FMassEntityHandle MassAgent, FMassActorFragment& ActorInfo, const int16 TemplateActorIndex, FMassActorSpawnRequestHandle& SpawnRequestHandle, FMassCommandBuffer& CommandBuffer, const bool bCancelSpawningOnly = false);
 
 	/*
 	 * Update representation type for each entity, must be called within a ForEachEntityChunk
diff --git a/Engine/Plugins/Runtime/MassGameplay/Source/MassSmartObjects/Private/MassSmartObjectProcessor.cpp b/Engine/Plugins/Runtime/MassGameplay/Source/MassSmartObjects/Private/MassSmartObjectProcessor.cpp
--- a/Engine/Plugins/Runtime/MassGameplay/Source/MassSmartObjects/Private/MassSmartObjectProcessor.cpp
+++ b/Engine/Plugins/Runtime/MassGameplay/Source/MassSmartObjects/Private/MassSmartObjectProcessor.cpp
@@ -22,10 +22,8 @@
 //----------------------------------------------------------------------//
 void UMassSmartObjectCandidatesFinderProcessor::Initialize(UObject& Owner)
 {
-	SmartObjectSubsystem = UWorld::GetSubsystem<USmartObjectSubsystem>(GetWorld());
 	SignalSubsystem = UWorld::GetSubsystem<UMassSignalSubsystem>(Owner.GetWorld());
 	AnnotationSubsystem = UWorld::GetSubsystem<UZoneGraphAnnotationSubsystem>(GetWorld());
-	ZoneGraphSubsystem = UWorld::GetSubsystem<UZoneGraphSubsystem>(GetWorld());
 }
 
 void UMassSmartObjectCandidatesFinderProcessor::ConfigureQueries()
@@ -33,10 +31,13 @@ void UMassSmartObjectCandidatesFinderProcessor::ConfigureQueries()
 	WorldRequestQuery.AddRequirement<FMassSmartObjectWorldLocationRequestFragment>(EMassFragmentAccess::ReadOnly);
 	WorldRequestQuery.AddRequirement<FMassSmartObjectRequestResultFragment>(EMassFragmentAccess::ReadWrite);
 	WorldRequestQuery.AddTagRequirement<FMassSmartObjectCompletedRequestTag>(EMassFragmentPresence::None);
+	WorldRequestQuery.AddSystemRequirement<USmartObjectSubsystem>(EMassFragmentAccess::ReadOnly);
 
 	LaneRequestQuery.AddRequirement<FMassSmartObjectLaneLocationRequestFragment>(EMassFragmentAccess::ReadOnly);
 	LaneRequestQuery.AddRequirement<FMassSmartObjectRequestResultFragment>(EMassFragmentAccess::ReadWrite);
 	LaneRequestQuery.AddTagRequirement<FMassSmartObjectCompletedRequestTag>(EMassFragmentPresence::None);
+	LaneRequestQuery.AddSystemRequirement<UZoneGraphSubsystem>(EMassFragmentAccess::ReadOnly);
+	LaneRequestQuery.AddSystemRequirement<USmartObjectSubsystem>(EMassFragmentAccess::ReadOnly);
 }
 
 UMassSmartObjectCandidatesFinderProcessor::UMassSmartObjectCandidatesFinderProcessor()
@@ -50,9 +51,7 @@ UMassSmartObjectCandidatesFinderProcessor::UMassSmartObjectCandidatesFinderProce
 
 void UMassSmartObjectCandidatesFinderProcessor::Execute(UMassEntitySubsystem& EntitySubsystem, FMassExecutionContext& Context)
 {
-	checkf(SmartObjectSubsystem != nullptr, TEXT("SmartObjectSubsystem should exist when executing processors."));
 	checkf(SignalSubsystem != nullptr, TEXT("MassSignalSubsystem should exist when executing processors."));
-	checkf(ZoneGraphSubsystem != nullptr, TEXT("ZoneGraphSubsystem should exist when executing processors."));
 	checkf(AnnotationSubsystem != nullptr, TEXT("ZoneGraphAnnotationSubsystem should exist when executing processors."));
 
 	// Create filter
@@ -83,8 +82,10 @@ void UMassSmartObjectCandidatesFinderProcessor::Execute(UMassEntitySubsystem& En
 	};
 
 	// Process world location based requests
-	WorldRequestQuery.ForEachEntityChunk(EntitySubsystem, Context, [this, &Filter, &EntitiesToSignal, &BeginRequestProcessing, &EndRequestProcessing](FMassExecutionContext& Context)
+	WorldRequestQuery.ForEachEntityChunk(EntitySubsystem, Context, [this, &Filter, &EntitiesToSignal, &BeginRequestProcessing, &EndRequestProcessing, World = EntitySubsystem.GetWorld()](FMassExecutionContext& Context)
 	{
+		const USmartObjectSubsystem& SmartObjectSubsystem = Context.GetSubsystemChecked<USmartObjectSubsystem>(World);
+
 		const int32 NumEntities = Context.GetNumEntities();
 		EntitiesToSignal.Reserve(EntitiesToSignal.Num() + NumEntities);
 
@@ -112,23 +113,23 @@ void UMassSmartObjectCandidatesFinderProcessor::Execute(UMassEntitySubsystem& En
 #endif // WITH_MASSGAMEPLAY_DEBUG
 
 			BeginRequestProcessing(Entity, Context, Result);
-			ON_SCOPE_EXIT{ EndRequestProcessing(SmartObjectSubsystem, Entity, Result);	};
+			ON_SCOPE_EXIT{ EndRequestProcessing(&SmartObjectSubsystem, Entity, Result);	};
 
 			QueryResults.Reset();
-			SmartObjectSubsystem->FindSmartObjects(FSmartObjectRequest(SearchBounds, Filter), QueryResults);
+			SmartObjectSubsystem.FindSmartObjects(FSmartObjectRequest(SearchBounds, Filter), QueryResults);
 
 			SortedResults.Reset(QueryResults.Num());
 			for (const FSmartObjectRequestResult& QueryResult : QueryResults)
 			{
-				const FVector SlotLocation = SmartObjectSubsystem->GetSlotLocation(QueryResult.SlotHandle).GetValue();
+				const FVector SlotLocation = SmartObjectSubsystem.GetSlotLocation(QueryResult.SlotHandle).GetValue();
 				SortedResults.Emplace(QueryResult, FVector::DistSquared(SearchOrigin, SlotLocation));
 
 #if WITH_MASSGAMEPLAY_DEBUG
 				if (bDisplayDebug)
 				{
 					constexpr float DebugRadius = 10.f;
-					UE_VLOG_LOCATION(SmartObjectSubsystem, LogSmartObject, Display, SlotLocation, DebugRadius, DebugColor, TEXT("%s"), *LexToString(QueryResult.SmartObjectHandle));
-					UE_VLOG_SEGMENT(SmartObjectSubsystem, LogSmartObject, Display, SearchOrigin, SlotLocation, DebugColor, TEXT(""));
+					UE_VLOG_LOCATION(&SmartObjectSubsystem, LogSmartObject, Display, SlotLocation, DebugRadius, DebugColor, TEXT("%s"), *LexToString(QueryResult.SmartObjectHandle));
+					UE_VLOG_SEGMENT(&SmartObjectSubsystem, LogSmartObject, Display, SearchOrigin, SlotLocation, DebugColor, TEXT(""));
 				}
 #endif // WITH_MASSGAMEPLAY_DEBUG
 			}
@@ -147,8 +148,13 @@ void UMassSmartObjectCandidatesFinderProcessor::Execute(UMassEntitySubsystem& En
 	USmartObjectZoneAnnotations* Annotations = Cast<USmartObjectZoneAnnotations>(AnnotationSubsystem->GetFirstAnnotationForTag(SmartObjectTag));
 
 	LaneRequestQuery.ForEachEntityChunk(EntitySubsystem, Context,
-		[this, Annotations, &Filter, SmartObjectTag, &EntitiesToSignal, &BeginRequestProcessing, &EndRequestProcessing](FMassExecutionContext& Context)
+		[this, Annotations, &Filter, SmartObjectTag, &EntitiesToSignal, &BeginRequestProcessing, &EndRequestProcessing, World = EntitySubsystem.GetWorld()](FMassExecutionContext& Context)
 		{
+#if WITH_MASSGAMEPLAY_DEBUG
+			const UZoneGraphSubsystem& ZoneGraphSubsystem = Context.GetSubsystemChecked<UZoneGraphSubsystem>(World);
+#endif // WITH_MASSGAMEPLAY_DEBUG
+			const USmartObjectSubsystem& SmartObjectSubsystem = Context.GetSubsystemChecked<USmartObjectSubsystem>(World);
+
 			const int32 NumEntities = Context.GetNumEntities();
 			EntitiesToSignal.Reserve(EntitiesToSignal.Num() + NumEntities);
 
@@ -176,7 +182,7 @@ void UMassSmartObjectCandidatesFinderProcessor::Execute(UMassEntitySubsystem& En
 #endif // WITH_MASSGAMEPLAY_DEBUG
 
 				BeginRequestProcessing(Entity, Context, Result);
-				ON_SCOPE_EXIT{ EndRequestProcessing(SmartObjectSubsystem, Entity, Result); };
+				ON_SCOPE_EXIT{ EndRequestProcessing(&SmartObjectSubsystem, Entity, Result); };
 
 				if (!ensureMsgf(RequestLaneHandle.IsValid(), TEXT("Requesting smart objects using an invalid handle")))
 				{
@@ -185,7 +191,7 @@ void UMassSmartObjectCandidatesFinderProcessor::Execute(UMassEntitySubsystem& En
 
 				if (Annotations == nullptr)
 				{
-					UE_VLOG(SmartObjectSubsystem, LogSmartObject, Warning, TEXT("%d lane location based requests failed since SmartObject annotations are not available"), NumEntities);
+					UE_VLOG(&SmartObjectSubsystem, LogSmartObject, Warning, TEXT("%d lane location based requests failed since SmartObject annotations are not available"), NumEntities);
 					return;
 				}
 
@@ -236,7 +242,7 @@ void UMassSmartObjectCandidatesFinderProcessor::Execute(UMassEntitySubsystem& En
 
 					// Make sure that we can use a slot in that object (availability with supported definitions, etc.)
 					TArray<FSmartObjectSlotHandle> SlotHandles;
-					SmartObjectSubsystem->FindSlots(Handle, Filter, SlotHandles);
+					SmartObjectSubsystem.FindSlots(Handle, Filter, SlotHandles);
 
 					if (SlotHandles.IsEmpty())
 					{
@@ -251,14 +257,14 @@ void UMassSmartObjectCandidatesFinderProcessor::Execute(UMassEntitySubsystem& En
 						if (bDisplayDebug)
 						{
 							FZoneGraphLaneLocation RequestLaneLocation, EntryPointLaneLocation;
-							ZoneGraphSubsystem->CalculateLocationAlongLane(RequestLaneHandle, RequestLocation.DistanceAlongLane, RequestLaneLocation);
-							ZoneGraphSubsystem->CalculateLocationAlongLane(RequestLaneHandle, EntryPoint.DistanceAlongLane, EntryPointLaneLocation);
+							ZoneGraphSubsystem.CalculateLocationAlongLane(RequestLaneHandle, RequestLocation.DistanceAlongLane, RequestLaneLocation);
+							ZoneGraphSubsystem.CalculateLocationAlongLane(RequestLaneHandle, EntryPoint.DistanceAlongLane, EntryPointLaneLocation);
 
 							constexpr float DebugRadius = 10.f;
-							FVector SlotLocation = SmartObjectSubsystem->GetSlotLocation(SlotHandle).Get(EntryPointLaneLocation.Position);
-							UE_VLOG_LOCATION(SmartObjectSubsystem, LogSmartObject, Display, SlotLocation, DebugRadius, DebugColor, TEXT("%s"), *LexToString(SlotHandle));
-							UE_VLOG_SEGMENT(SmartObjectSubsystem, LogSmartObject, Display, SlotLocation, EntryPointLaneLocation.Position, DebugColor, TEXT(""));
-							UE_VLOG_SEGMENT(SmartObjectSubsystem, LogSmartObject, Display, RequestLaneLocation.Position, EntryPointLaneLocation.Position, DebugColor, TEXT(""));
+							FVector SlotLocation = SmartObjectSubsystem.GetSlotLocation(SlotHandle).Get(EntryPointLaneLocation.Position);
+							UE_VLOG_LOCATION(&SmartObjectSubsystem, LogSmartObject, Display, SlotLocation, DebugRadius, DebugColor, TEXT("%s"), *LexToString(SlotHandle));
+							UE_VLOG_SEGMENT(&SmartObjectSubsystem, LogSmartObject, Display, SlotLocation, EntryPointLaneLocation.Position, DebugColor, TEXT(""));
+							UE_VLOG_SEGMENT(&SmartObjectSubsystem, LogSmartObject, Display, RequestLaneLocation.Position, EntryPointLaneLocation.Position, DebugColor, TEXT(""));
 						}
 #endif // WITH_MASSGAMEPLAY_DEBUG
 						if (Result.NumCandidates == FMassSmartObjectRequestResult::MaxNumCandidates)
@@ -287,7 +293,6 @@ void UMassSmartObjectCandidatesFinderProcessor::Execute(UMassEntitySubsystem& En
 //----------------------------------------------------------------------//
 void UMassSmartObjectTimedBehaviorProcessor::Initialize(UObject& Owner)
 {
-	SmartObjectSubsystem = UWorld::GetSubsystem<USmartObjectSubsystem>(GetWorld());
 	SignalSubsystem = UWorld::GetSubsystem<UMassSignalSubsystem>(Owner.GetWorld());
 }
 
@@ -295,6 +300,7 @@ void UMassSmartObjectTimedBehaviorProcessor::ConfigureQueries()
 {
 	EntityQuery.AddRequirement<FMassSmartObjectUserFragment>(EMassFragmentAccess::ReadWrite);
 	EntityQuery.AddRequirement<FMassSmartObjectTimedBehaviorFragment>(EMassFragmentAccess::ReadWrite);
+	EntityQuery.AddSystemRequirement<USmartObjectSubsystem>(EMassFragmentAccess::ReadOnly);
 }
 
 UMassSmartObjectTimedBehaviorProcessor::UMassSmartObjectTimedBehaviorProcessor()
@@ -304,15 +310,16 @@ UMassSmartObjectTimedBehaviorProcessor::UMassSmartObjectTimedBehaviorProcessor()
 
 void UMassSmartObjectTimedBehaviorProcessor::Execute(UMassEntitySubsystem& EntitySubsystem, FMassExecutionContext& Context)
 {
-	checkf(SmartObjectSubsystem != nullptr, TEXT("SmartObjectSubsystem should exist when executing processors."));
 	checkf(SignalSubsystem != nullptr, TEXT("MassSignalSubsystem should exist when executing processors."));
 
 	TArray<FMassEntityHandle> ToRelease;
 
 	QUICK_SCOPE_CYCLE_COUNTER(UMassProcessor_SmartObjectTestBehavior_Run);
 
-	EntityQuery.ForEachEntityChunk(EntitySubsystem, Context, [this, &ToRelease](FMassExecutionContext& Context)
+	EntityQuery.ForEachEntityChunk(EntitySubsystem, Context, [this, &ToRelease, World = EntitySubsystem.GetWorld()](FMassExecutionContext& Context)
 	{
+		const USmartObjectSubsystem& SmartObjectSubsystem = Context.GetSubsystemChecked<USmartObjectSubsystem>(World);
+
 		const int32 NumEntities = Context.GetNumEntities();
 		const TArrayView<FMassSmartObjectUserFragment> UserList = Context.GetMutableFragmentView<FMassSmartObjectUserFragment>();
 		const TArrayView<FMassSmartObjectTimedBehaviorFragment> TimedBehaviorFragments = Context.GetMutableFragmentView<FMassSmartObjectTimedBehaviorFragment>();
@@ -334,18 +341,18 @@ void UMassSmartObjectTimedBehaviorProcessor::Execute(UMassEntitySubsystem& Entit
 			const bool bIsDebuggingEntity = UE::Mass::Debug::IsDebuggingEntity(Entity, &DebugColor);
 			if (bIsDebuggingEntity)
 			{
-				UE_CVLOG(bMustRelease, SmartObjectSubsystem, LogSmartObject, Log, TEXT("[%s] stops using [%s]"), *Entity.DebugGetDescription(), *LexToString(SOUser.ClaimHandle));
-				UE_CVLOG(!bMustRelease, SmartObjectSubsystem, LogSmartObject, Verbose, TEXT("[%s] using [%s]. Time left: %.1f"), *Entity.DebugGetDescription(), *LexToString(SOUser.ClaimHandle), UseTime);
+				UE_CVLOG(bMustRelease, &SmartObjectSubsystem, LogSmartObject, Log, TEXT("[%s] stops using [%s]"), *Entity.DebugGetDescription(), *LexToString(SOUser.ClaimHandle));
+				UE_CVLOG(!bMustRelease, &SmartObjectSubsystem, LogSmartObject, Verbose, TEXT("[%s] using [%s]. Time left: %.1f"), *Entity.DebugGetDescription(), *LexToString(SOUser.ClaimHandle), UseTime);
 
-				const TOptional<FTransform> Transform = SmartObjectSubsystem->GetSlotTransform(SOUser.ClaimHandle);
+				const TOptional<FTransform> Transform = SmartObjectSubsystem.GetSlotTransform(SOUser.ClaimHandle);
 				if (Transform.IsSet())
 				{
 					constexpr float Radius = 40.f;
 					const FVector HalfHeightOffset(0.f, 0.f, 100.f);
 					const FVector Pos = Transform.GetValue().GetLocation();
 					const FVector Dir = Transform.GetValue().GetRotation().GetForwardVector();
-					UE_VLOG_CYLINDER(SmartObjectSubsystem, LogSmartObject, Display, Pos - HalfHeightOffset, Pos + HalfHeightOffset, Radius, DebugColor, TEXT(""));
-					UE_VLOG_ARROW(SmartObjectSubsystem, LogSmartObject, Display, Pos, Pos + Dir * 2.0f * Radius, DebugColor, TEXT(""));
+					UE_VLOG_CYLINDER(&SmartObjectSubsystem, LogSmartObject, Display, Pos - HalfHeightOffset, Pos + HalfHeightOffset, Radius, DebugColor, TEXT(""));
+					UE_VLOG_ARROW(&SmartObjectSubsystem, LogSmartObject, Display, Pos, Pos + Dir * 2.0f * Radius, DebugColor, TEXT(""));
 				}
 			}
 #endif // WITH_MASSGAMEPLAY_DEBUG
@@ -377,29 +384,20 @@ UMassSmartObjectUserFragmentDeinitializer::UMassSmartObjectUserFragmentDeinitial
 void UMassSmartObjectUserFragmentDeinitializer::ConfigureQueries()
 {
 	EntityQuery.AddRequirement<FMassSmartObjectUserFragment>(EMassFragmentAccess::ReadWrite);
-}
-
-void UMassSmartObjectUserFragmentDeinitializer::Initialize(UObject& Owner)
-{
-	Super::Initialize(Owner);
-	SmartObjectSubsystem = UWorld::GetSubsystem<USmartObjectSubsystem>(Owner.GetWorld());
+	EntityQuery.AddSystemRequirement<USmartObjectSubsystem>(EMassFragmentAccess::ReadWrite);
 }
 
 void UMassSmartObjectUserFragmentDeinitializer::Execute(UMassEntitySubsystem& EntitySubsystem, FMassExecutionContext& Context)
 {
-	if (SmartObjectSubsystem == nullptr)
-	{
-		return;
-	}
-
-	EntityQuery.ForEachEntityChunk(EntitySubsystem, Context, [this](FMassExecutionContext& Context)
+	EntityQuery.ForEachEntityChunk(EntitySubsystem, Context, [this, World = EntitySubsystem.GetWorld()](FMassExecutionContext& Context)
 		{
+			USmartObjectSubsystem& SmartObjectSubsystem = Context.GetMutableSubsystemChecked<USmartObjectSubsystem>(World);
 			const int32 NumEntities = Context.GetNumEntities();
 			const TArrayView<FMassSmartObjectUserFragment> SmartObjectUserFragments = Context.GetMutableFragmentView<FMassSmartObjectUserFragment>();
 
 			for (int32 i = 0; i < NumEntities; ++i)
 			{
-				SmartObjectSubsystem->UnregisterSlotInvalidationCallback(SmartObjectUserFragments[i].ClaimHandle);
+				SmartObjectSubsystem.UnregisterSlotInvalidationCallback(SmartObjectUserFragments[i].ClaimHandle);
 			}
 		});
 }
\ No newline at end of file
diff --git a/Engine/Plugins/Runtime/MassGameplay/Source/MassSmartObjects/Public/MassSmartObjectProcessor.h b/Engine/Plugins/Runtime/MassGameplay/Source/MassSmartObjects/Public/MassSmartObjectProcessor.h
--- a/Engine/Plugins/Runtime/MassGameplay/Source/MassSmartObjects/Public/MassSmartObjectProcessor.h
+++ b/Engine/Plugins/Runtime/MassGameplay/Source/MassSmartObjects/Public/MassSmartObjectProcessor.h
@@ -6,9 +6,7 @@
 #include "MassProcessor.h"
 #include "MassSmartObjectProcessor.generated.h"
 
-class USmartObjectSubsystem;
 class UMassSignalSubsystem;
-class UZoneGraphSubsystem;
 class UZoneGraphAnnotationSubsystem;
 
 /** Processor that builds a list of candidates objects for each users. */
@@ -29,15 +27,9 @@ protected:
 	UPROPERTY(EditDefaultsOnly, Category = SmartObject, config)
 	float SearchExtents = 5000.f;
 
-	UPROPERTY(Transient)
-	USmartObjectSubsystem* SmartObjectSubsystem;
-
 	UPROPERTY(Transient)
 	UMassSignalSubsystem* SignalSubsystem;
 
-	UPROPERTY(Transient)
-	UZoneGraphSubsystem* ZoneGraphSubsystem;
-
 	UPROPERTY(Transient)
 	UZoneGraphAnnotationSubsystem* AnnotationSubsystem;
 
@@ -61,9 +53,6 @@ protected:
 	virtual void Execute(UMassEntitySubsystem& EntitySubsystem, FMassExecutionContext& Context) override;
 	virtual void ConfigureQueries() override;
 
-	UPROPERTY(Transient)
-	USmartObjectSubsystem* SmartObjectSubsystem;
-
 	UPROPERTY(Transient)
 	UMassSignalSubsystem* SignalSubsystem;
 
@@ -80,11 +69,7 @@ class MASSSMARTOBJECTS_API UMassSmartObjectUserFragmentDeinitializer : public UM
 
 protected:
 	virtual void ConfigureQueries() override;
-	virtual void Initialize(UObject& Owner) override;
 	virtual void Execute(UMassEntitySubsystem& EntitySubsystem, FMassExecutionContext& Context) override;
 
-	UPROPERTY(Transient)
-	USmartObjectSubsystem* SmartObjectSubsystem;
-
 	FMassEntityQuery EntityQuery;
 };
\ No newline at end of file
