Removed UMassLODProcessorBase due to it not serving an actual purpose

diff --git a/Engine/Plugins/AI/MassCrowd/Source/MassCrowd/Private/MassCrowdServerRepresentationLODProcessor.cpp b/Engine/Plugins/AI/MassCrowd/Source/MassCrowd/Private/MassCrowdServerRepresentationLODProcessor.cpp
--- a/Engine/Plugins/AI/MassCrowd/Source/MassCrowd/Private/MassCrowdServerRepresentationLODProcessor.cpp
+++ b/Engine/Plugins/AI/MassCrowd/Source/MassCrowd/Private/MassCrowdServerRepresentationLODProcessor.cpp
@@ -81,8 +81,8 @@ void UMassCrowdServerRepresentationLODProcessor::Execute(UMassEntitySubsystem& E
 	if (UE::MassCrowd::bDebugCrowdServerRepresentationLOD)
 	{
 		TRACE_CPUPROFILER_EVENT_SCOPE(TEXT("DebugDisplayLOD"))
-		
-		EntityQuery.ForEachEntityChunk(EntitySubsystem, Context, [this](FMassExecutionContext& Context)
+		UWorld* World = EntitySubsystem.GetWorld();
+		EntityQuery.ForEachEntityChunk(EntitySubsystem, Context, [this, World](FMassExecutionContext& Context)
 		{
 			const TConstArrayView<FTransformFragment> LocationList = Context.GetFragmentView<FTransformFragment>();
 			const TConstArrayView<FMassRepresentationLODFragment> RepresentationLODFragments = Context.GetFragmentView<FMassRepresentationLODFragment>();
diff --git a/Engine/Plugins/AI/MassCrowd/Source/MassCrowd/Private/MassCrowdVisualizationLODProcessor.cpp b/Engine/Plugins/AI/MassCrowd/Source/MassCrowd/Private/MassCrowdVisualizationLODProcessor.cpp
--- a/Engine/Plugins/AI/MassCrowd/Source/MassCrowd/Private/MassCrowdVisualizationLODProcessor.cpp
+++ b/Engine/Plugins/AI/MassCrowd/Source/MassCrowd/Private/MassCrowdVisualizationLODProcessor.cpp
@@ -48,6 +48,8 @@ void UMassCrowdVisualizationLODProcessor::ConfigureQueries()
 
 void UMassCrowdVisualizationLODProcessor::Execute(UMassEntitySubsystem& EntitySubsystem, FMassExecutionContext& Context)
 {
+	UWorld* World = EntitySubsystem.GetWorld();
+
 	ForceOffLOD((bool)UE::MassCrowd::GCrowdTurnOffVisualization);
 
 	TRACE_CPUPROFILER_EVENT_SCOPE(TEXT("CrowdVisualizationLOD"))
@@ -58,7 +60,7 @@ void UMassCrowdVisualizationLODProcessor::Execute(UMassEntitySubsystem& EntitySu
 	{
 		TRACE_CPUPROFILER_EVENT_SCOPE(TEXT("DebugDisplayLOD"))
 
-		DebugEntityQuery.ForEachEntityChunk(EntitySubsystem, Context, [this](FMassExecutionContext& Context)
+		DebugEntityQuery.ForEachEntityChunk(EntitySubsystem, Context, [World](FMassExecutionContext& Context)
 		{
 			FMassVisualizationLODSharedFragment& LODSharedFragment = Context.GetMutableSharedFragment<FMassVisualizationLODSharedFragment>();
 			const TConstArrayView<FTransformFragment> LocationList = Context.GetFragmentView<FTransformFragment>();
@@ -71,7 +73,7 @@ void UMassCrowdVisualizationLODProcessor::Execute(UMassEntitySubsystem& EntitySu
 	{
 		TRACE_CPUPROFILER_EVENT_SCOPE(TEXT("ShowISMUnderSpecifiedRange"))
 
-		DebugEntityQuery.ForEachEntityChunk(EntitySubsystem, Context, [this](const FMassExecutionContext& Context)
+		DebugEntityQuery.ForEachEntityChunk(EntitySubsystem, Context, [World](const FMassExecutionContext& Context)
 		{
 			const TConstArrayView<FTransformFragment> LocationList = Context.GetFragmentView<FTransformFragment>();
 			const TConstArrayView<FMassRepresentationFragment> RepresentationFragmentList = Context.GetFragmentView<FMassRepresentationFragment>();
diff --git a/Engine/Plugins/AI/MassCrowd/Source/MassCrowd/Public/MassCrowdServerRepresentationLODProcessor.h b/Engine/Plugins/AI/MassCrowd/Source/MassCrowd/Public/MassCrowdServerRepresentationLODProcessor.h
--- a/Engine/Plugins/AI/MassCrowd/Source/MassCrowd/Public/MassCrowdServerRepresentationLODProcessor.h
+++ b/Engine/Plugins/AI/MassCrowd/Source/MassCrowd/Public/MassCrowdServerRepresentationLODProcessor.h
@@ -4,14 +4,14 @@
 
 #include "MassCrowdTypes.h"
 #include "MassCrowdFragments.h"
-#include "MassLODSubsystem.h"
+#include "MassProcessor.h"
 #include "MassLODCalculator.h"
 #include "MassLODTickRateController.h"
 
 #include "MassCrowdServerRepresentationLODProcessor.generated.h"
 
 UCLASS(meta=(DisplayName="Crowd Simulation LOD"))
-class MASSCROWD_API UMassCrowdServerRepresentationLODProcessor : public UMassLODProcessorBase
+class MASSCROWD_API UMassCrowdServerRepresentationLODProcessor : public UMassProcessor
 {
 	GENERATED_BODY()
 
diff --git a/Engine/Plugins/Runtime/MassGameplay/Source/MassLOD/Private/MassLODCollectorProcessor.cpp b/Engine/Plugins/Runtime/MassGameplay/Source/MassLOD/Private/MassLODCollectorProcessor.cpp
--- a/Engine/Plugins/Runtime/MassGameplay/Source/MassLOD/Private/MassLODCollectorProcessor.cpp
+++ b/Engine/Plugins/Runtime/MassGameplay/Source/MassLOD/Private/MassLODCollectorProcessor.cpp
@@ -79,6 +79,7 @@ void UMassLODCollectorProcessor::Execute(UMassEntitySubsystem& EntitySubsystem,
 	const TArray<FViewerInfo>& Viewers = LODSubsystem.GetViewers();
 	Collector.PrepareExecution(Viewers);
 
+	UWorld* World = EntitySubsystem.GetWorld();
 	check(World);
 	if (World->IsNetMode(NM_DedicatedServer))
 	{
diff --git a/Engine/Plugins/Runtime/MassGameplay/Source/MassLOD/Private/MassLODSubsystem.cpp b/Engine/Plugins/Runtime/MassGameplay/Source/MassLOD/Private/MassLODSubsystem.cpp
--- a/Engine/Plugins/Runtime/MassGameplay/Source/MassLOD/Private/MassLODSubsystem.cpp
+++ b/Engine/Plugins/Runtime/MassGameplay/Source/MassLOD/Private/MassLODSubsystem.cpp
@@ -21,13 +21,6 @@ namespace MassLOD
 }
 }
 
-void UMassLODProcessorBase::Initialize(UObject& Owner)
-{
-	Super::Initialize(Owner);
-
-	World = Owner.GetWorld();
-}
-
 void UMassLODSubsystem::Initialize(FSubsystemCollectionBase& Collection)
 {
 	Collection.InitializeDependency(UMassSimulationSubsystem::StaticClass());
diff --git a/Engine/Plugins/Runtime/MassGameplay/Source/MassLOD/Private/MassSimulationLOD.cpp b/Engine/Plugins/Runtime/MassGameplay/Source/MassLOD/Private/MassSimulationLOD.cpp
--- a/Engine/Plugins/Runtime/MassGameplay/Source/MassLOD/Private/MassSimulationLOD.cpp
+++ b/Engine/Plugins/Runtime/MassGameplay/Source/MassLOD/Private/MassSimulationLOD.cpp
@@ -149,6 +149,7 @@ void UMassSimulationLODProcessor::Execute(UMassEntitySubsystem& EntitySubsystem,
 		});
 	}
 
+	UWorld* World = EntitySubsystem.GetWorld();
 	{
 		TRACE_CPUPROFILER_EVENT_SCOPE(TEXT("VariableTickRates"))
 		check(World);
@@ -186,7 +187,7 @@ void UMassSimulationLODProcessor::Execute(UMassEntitySubsystem& EntitySubsystem,
 	if (UE::MassLOD::bDebugSimulationLOD)
 	{
 		TRACE_CPUPROFILER_EVENT_SCOPE(TEXT("DebugDisplayLOD"));
-		EntityQuery.ForEachEntityChunk(EntitySubsystem, Context, [this](FMassExecutionContext& Context)
+		EntityQuery.ForEachEntityChunk(EntitySubsystem, Context, [World](FMassExecutionContext& Context)
 		{
 			FMassSimulationLODSharedFragment& LODSharedFragment = Context.GetMutableSharedFragment<FMassSimulationLODSharedFragment>();
 			TConstArrayView<FTransformFragment> LocationList = Context.GetFragmentView<FTransformFragment>();
diff --git a/Engine/Plugins/Runtime/MassGameplay/Source/MassLOD/Public/MassLODCollectorProcessor.h b/Engine/Plugins/Runtime/MassGameplay/Source/MassLOD/Public/MassLODCollectorProcessor.h
--- a/Engine/Plugins/Runtime/MassGameplay/Source/MassLOD/Public/MassLODCollectorProcessor.h
+++ b/Engine/Plugins/Runtime/MassGameplay/Source/MassLOD/Public/MassLODCollectorProcessor.h
@@ -3,7 +3,7 @@
 #pragma once
 
 #include "MassEntityTypes.h"
-#include "MassLODSubsystem.h"
+#include "MassProcessor.h"
 #include "MassLODCollector.h"
 
 #include "MassLODCollectorProcessor.generated.h"
@@ -20,7 +20,7 @@ struct FMassGenericCollectorLogic : public FLODDefaultLogic
  * LOD collector which combines collection of LOD information for both Viewer and Simulation LODing when possible.
  */
 UCLASS(meta = (DisplayName = "LOD Collector"))
-class MASSLOD_API UMassLODCollectorProcessor : public UMassLODProcessorBase
+class MASSLOD_API UMassLODCollectorProcessor : public UMassProcessor
 {
 	GENERATED_BODY()
 
diff --git a/Engine/Plugins/Runtime/MassGameplay/Source/MassLOD/Public/MassLODSubsystem.h b/Engine/Plugins/Runtime/MassGameplay/Source/MassLOD/Public/MassLODSubsystem.h
--- a/Engine/Plugins/Runtime/MassGameplay/Source/MassLOD/Public/MassLODSubsystem.h
+++ b/Engine/Plugins/Runtime/MassGameplay/Source/MassLOD/Public/MassLODSubsystem.h
@@ -12,22 +12,6 @@
 
 class UMassLODSubsystem;
 
-/*
- * Base mass LOD processor to store common information for all LOD processors
- */
-UCLASS(abstract)
-class MASSLOD_API UMassLODProcessorBase : public UMassProcessor
-{
-	GENERATED_BODY()
-
-public:
-	virtual void Initialize(UObject& Owner) override;
-
-protected:
-	UPROPERTY()
-	UWorld* World = nullptr;
-};
-
 /*
  * Handle that lets you reference the concept of a viewer
  */
diff --git a/Engine/Plugins/Runtime/MassGameplay/Source/MassLOD/Public/MassSimulationLOD.h b/Engine/Plugins/Runtime/MassGameplay/Source/MassLOD/Public/MassSimulationLOD.h
--- a/Engine/Plugins/Runtime/MassGameplay/Source/MassLOD/Public/MassSimulationLOD.h
+++ b/Engine/Plugins/Runtime/MassGameplay/Source/MassLOD/Public/MassSimulationLOD.h
@@ -2,7 +2,7 @@
 
 #pragma once
 
-#include "MassLODSubsystem.h"
+#include "MassProcessor.h"
 #include "MassLODCollector.h"
 #include "MassLODCalculator.h"
 #include "MassLODTickRateController.h"
@@ -162,7 +162,7 @@ struct MASSLOD_API FMassSimulationVariableTickSharedFragment : public FMassShare
 };
 
 UCLASS(meta = (DisplayName = "Simulation LOD"))
-class MASSLOD_API UMassSimulationLODProcessor : public UMassLODProcessorBase
+class MASSLOD_API UMassSimulationLODProcessor : public UMassProcessor
 {
 	GENERATED_BODY()
 
diff --git a/Engine/Plugins/Runtime/MassGameplay/Source/MassReplication/Private/MassReplicationProcessor.cpp b/Engine/Plugins/Runtime/MassGameplay/Source/MassReplication/Private/MassReplicationProcessor.cpp
--- a/Engine/Plugins/Runtime/MassGameplay/Source/MassReplication/Private/MassReplicationProcessor.cpp
+++ b/Engine/Plugins/Runtime/MassGameplay/Source/MassReplication/Private/MassReplicationProcessor.cpp
@@ -61,6 +61,7 @@ void UMassReplicationProcessor::Initialize(UObject& Owner)
 	Super::Initialize(Owner);
 
 #if UE_REPLICATION_COMPILE_SERVER_CODE
+	UWorld* World = Owner.GetWorld();
 	ReplicationSubsystem = UWorld::GetSubsystem<UMassReplicationSubsystem>(World);
 
 	check(ReplicationSubsystem);
@@ -149,6 +150,7 @@ void UMassReplicationProcessor::PrepareExecution(UMassEntitySubsystem& EntitySub
 void UMassReplicationProcessor::Execute(UMassEntitySubsystem& EntitySubsystem, FMassExecutionContext& Context)
 {
 #if UE_REPLICATION_COMPILE_SERVER_CODE
+	UWorld* World = EntitySubsystem.GetWorld();
 	check(World);
 	check(ReplicationSubsystem);
 
@@ -357,9 +359,9 @@ void UMassReplicationProcessor::Execute(UMassEntitySubsystem& EntitySubsystem, F
 				// Optional debug display
 				if (UE::Mass::Replication::DebugClientReplicationLOD == ClientHandle.GetIndex())
 				{
-					EntitySubsystem.ForEachSharedFragment<FMassReplicationSharedFragment>([this, &EntitySubsystem, &Context](FMassReplicationSharedFragment& RepSharedFragment)
+					EntitySubsystem.ForEachSharedFragment<FMassReplicationSharedFragment>([World, &EntitySubsystem, &Context](FMassReplicationSharedFragment& RepSharedFragment)
 					{
-						RepSharedFragment.EntityQuery.ForEachEntityChunk(EntitySubsystem, Context, [this, &RepSharedFragment](FMassExecutionContext& Context)
+						RepSharedFragment.EntityQuery.ForEachEntityChunk(EntitySubsystem, Context, [World, &RepSharedFragment](FMassExecutionContext& Context)
 						{
 							const TConstArrayView<FTransformFragment> TransformList = Context.GetFragmentView<FTransformFragment>();
 							const TConstArrayView<FMassReplicationLODFragment> ViewerLODList = Context.GetFragmentView<FMassReplicationLODFragment>();
diff --git a/Engine/Plugins/Runtime/MassGameplay/Source/MassReplication/Public/MassReplicationProcessor.h b/Engine/Plugins/Runtime/MassGameplay/Source/MassReplication/Public/MassReplicationProcessor.h
--- a/Engine/Plugins/Runtime/MassGameplay/Source/MassReplication/Public/MassReplicationProcessor.h
+++ b/Engine/Plugins/Runtime/MassGameplay/Source/MassReplication/Public/MassReplicationProcessor.h
@@ -9,7 +9,7 @@
 #include "MassReplicationTypes.h"
 #include "MassReplicationFragments.h"
 #include "MassSpawnerTypes.h"
-#include "MassLODSubsystem.h"
+#include "MassProcessor.h"
 #include "MassReplicationSubsystem.h"
 #include "MassReplicationProcessor.generated.h"
 
@@ -22,7 +22,7 @@ class UWorld;
  *  query Mass entity fragments and set those values for replication when appropriate, using the MassClientBubbleHandler.
  */
 UCLASS()
-class MASSREPLICATION_API UMassReplicationProcessor : public UMassLODProcessorBase
+class MASSREPLICATION_API UMassReplicationProcessor : public UMassProcessor
 {
 	GENERATED_BODY()
 
diff --git a/Engine/Plugins/Runtime/MassGameplay/Source/MassRepresentation/Private/MassVisualizationLODProcessor.cpp b/Engine/Plugins/Runtime/MassGameplay/Source/MassRepresentation/Private/MassVisualizationLODProcessor.cpp
--- a/Engine/Plugins/Runtime/MassGameplay/Source/MassRepresentation/Private/MassVisualizationLODProcessor.cpp
+++ b/Engine/Plugins/Runtime/MassGameplay/Source/MassRepresentation/Private/MassVisualizationLODProcessor.cpp
@@ -109,7 +109,8 @@ void UMassVisualizationLODProcessor::Execute(UMassEntitySubsystem& EntitySubsyst
 	if (UE::MassRepresentation::bDebugRepresentationLOD)
 	{
 		TRACE_CPUPROFILER_EVENT_SCOPE(DebugDisplayLOD)
-		DebugEntityQuery.ForEachEntityChunk(EntitySubsystem, Context, [this](FMassExecutionContext& Context)
+		UWorld* World = EntitySubsystem.GetWorld();
+		DebugEntityQuery.ForEachEntityChunk(EntitySubsystem, Context, [World](FMassExecutionContext& Context)
 		{
 			FMassVisualizationLODSharedFragment& LODSharedFragment = Context.GetMutableSharedFragment<FMassVisualizationLODSharedFragment>();
 			TConstArrayView<FMassRepresentationLODFragment> RepresentationLODList = Context.GetFragmentView<FMassRepresentationLODFragment>();
diff --git a/Engine/Plugins/Runtime/MassGameplay/Source/MassRepresentation/Public/MassVisualizationLODProcessor.h b/Engine/Plugins/Runtime/MassGameplay/Source/MassRepresentation/Public/MassVisualizationLODProcessor.h
--- a/Engine/Plugins/Runtime/MassGameplay/Source/MassRepresentation/Public/MassVisualizationLODProcessor.h
+++ b/Engine/Plugins/Runtime/MassGameplay/Source/MassRepresentation/Public/MassVisualizationLODProcessor.h
@@ -4,12 +4,13 @@
 
 #include "MassCommonFragments.h"
 #include "MassRepresentationFragments.h"
+#include "MassProcessor.h"
 #include "MassLODCalculator.h"
 
 #include "MassVisualizationLODProcessor.generated.h"
 
 UCLASS()
-class MASSREPRESENTATION_API UMassVisualizationLODProcessor : public UMassLODProcessorBase
+class MASSREPRESENTATION_API UMassVisualizationLODProcessor : public UMassProcessor
 {
 	GENERATED_BODY()
 
