Renamed missed "system" occurences in the "subsystem requirements" code

diff --git a/Engine/Plugins/AI/MassAI/Source/MassAIBehavior/Private/MassLookAtProcessors.cpp b/Engine/Plugins/AI/MassAI/Source/MassAIBehavior/Private/MassLookAtProcessors.cpp
--- a/Engine/Plugins/AI/MassAI/Source/MassAIBehavior/Private/MassLookAtProcessors.cpp
+++ b/Engine/Plugins/AI/MassAI/Source/MassAIBehavior/Private/MassLookAtProcessors.cpp
@@ -108,8 +108,8 @@ void UMassLookAtProcessor::ConfigureQueries()
 	EntityQuery_Conditional.AddTagRequirement<FMassOffLODTag>(EMassFragmentPresence::None);
 	EntityQuery_Conditional.AddChunkRequirement<FMassVisualizationChunkFragment>(EMassFragmentAccess::ReadOnly);
 	EntityQuery_Conditional.SetChunkFilter(&FMassVisualizationChunkFragment::AreAnyEntitiesVisibleInChunk);
-	EntityQuery_Conditional.AddSystemRequirement<UMassNavigationSubsystem>(EMassFragmentAccess::ReadOnly);
-	EntityQuery_Conditional.AddSystemRequirement<UZoneGraphSubsystem>(EMassFragmentAccess::ReadOnly);
+	EntityQuery_Conditional.AddSubsystemRequirement<UMassNavigationSubsystem>(EMassFragmentAccess::ReadOnly);
+	EntityQuery_Conditional.AddSubsystemRequirement<UZoneGraphSubsystem>(EMassFragmentAccess::ReadOnly);
 }
 
 void UMassLookAtProcessor::Execute(UMassEntitySubsystem& EntitySubsystem, FMassExecutionContext& Context)
diff --git a/Engine/Plugins/AI/MassAI/Source/MassAIBehavior/Private/MassStateTreeProcessors.cpp b/Engine/Plugins/AI/MassAI/Source/MassAIBehavior/Private/MassStateTreeProcessors.cpp
--- a/Engine/Plugins/AI/MassAI/Source/MassAIBehavior/Private/MassStateTreeProcessors.cpp
+++ b/Engine/Plugins/AI/MassAI/Source/MassAIBehavior/Private/MassStateTreeProcessors.cpp
@@ -187,7 +187,7 @@ void UMassStateTreeFragmentDestructor::ConfigureQueries()
 {
 	EntityQuery.AddRequirement<FMassStateTreeInstanceFragment>(EMassFragmentAccess::ReadWrite);
 	EntityQuery.AddConstSharedRequirement<FMassStateTreeSharedFragment>();
-	EntityQuery.AddSystemRequirement<UMassStateTreeSubsystem>(EMassFragmentAccess::ReadWrite);
+	EntityQuery.AddSubsystemRequirement<UMassStateTreeSubsystem>(EMassFragmentAccess::ReadWrite);
 }
 
 void UMassStateTreeFragmentDestructor::Execute(UMassEntitySubsystem& EntitySubsystem, FMassExecutionContext& Context)
@@ -248,7 +248,7 @@ void UMassStateTreeActivationProcessor::ConfigureQueries()
 	EntityQuery.AddConstSharedRequirement<FMassStateTreeSharedFragment>();
 	EntityQuery.AddTagRequirement<FMassStateTreeActivatedTag>(EMassFragmentPresence::None);
 	EntityQuery.AddChunkRequirement<FMassSimulationVariableTickChunkFragment>(EMassFragmentAccess::ReadOnly, EMassFragmentPresence::Optional);
-	EntityQuery.AddSystemRequirement<UMassStateTreeSubsystem>(EMassFragmentAccess::ReadWrite);
+	EntityQuery.AddSubsystemRequirement<UMassStateTreeSubsystem>(EMassFragmentAccess::ReadWrite);
 }
 
 void UMassStateTreeActivationProcessor::Execute(UMassEntitySubsystem& EntitySubsystem, FMassExecutionContext& Context)
@@ -375,7 +375,7 @@ void UMassStateTreeProcessor::ConfigureQueries()
 {
 	EntityQuery.AddRequirement<FMassStateTreeInstanceFragment>(EMassFragmentAccess::ReadWrite);
 	EntityQuery.AddConstSharedRequirement<FMassStateTreeSharedFragment>();
-	EntityQuery.AddSystemRequirement<UMassStateTreeSubsystem>(EMassFragmentAccess::ReadWrite);
+	EntityQuery.AddSubsystemRequirement<UMassStateTreeSubsystem>(EMassFragmentAccess::ReadWrite);
 }
 
 void UMassStateTreeProcessor::SignalEntities(UMassEntitySubsystem& EntitySubsystem, FMassExecutionContext& Context, FMassSignalNameLookup& EntitySignals)
diff --git a/Engine/Plugins/AI/MassAI/Source/MassNavigation/Private/MassNavigationProcessors.cpp b/Engine/Plugins/AI/MassAI/Source/MassNavigation/Private/MassNavigationProcessors.cpp
--- a/Engine/Plugins/AI/MassAI/Source/MassNavigation/Private/MassNavigationProcessors.cpp
+++ b/Engine/Plugins/AI/MassAI/Source/MassNavigation/Private/MassNavigationProcessors.cpp
@@ -169,7 +169,7 @@ void UMassNavigationObstacleGridProcessor::ConfigureQueries()
 	BaseEntityQuery.AddRequirement<FTransformFragment>(EMassFragmentAccess::ReadOnly);
 	BaseEntityQuery.AddRequirement<FAgentRadiusFragment>(EMassFragmentAccess::ReadOnly);
 	BaseEntityQuery.AddRequirement<FMassNavigationObstacleGridCellLocationFragment>(EMassFragmentAccess::ReadWrite);
-	BaseEntityQuery.AddSystemRequirement<UMassNavigationSubsystem>(EMassFragmentAccess::ReadWrite);
+	BaseEntityQuery.AddSubsystemRequirement<UMassNavigationSubsystem>(EMassFragmentAccess::ReadWrite);
 
 	AddToGridEntityQuery = BaseEntityQuery;
 	AddToGridEntityQuery.AddRequirement<FMassAvoidanceColliderFragment>(EMassFragmentAccess::ReadOnly, EMassFragmentPresence::Optional);
@@ -281,7 +281,7 @@ UMassNavigationObstacleRemoverProcessor::UMassNavigationObstacleRemoverProcessor
 void UMassNavigationObstacleRemoverProcessor::ConfigureQueries()
 {
 	EntityQuery.AddRequirement<FMassNavigationObstacleGridCellLocationFragment>(EMassFragmentAccess::ReadWrite);
-	EntityQuery.AddSystemRequirement<UMassNavigationSubsystem>(EMassFragmentAccess::ReadWrite);
+	EntityQuery.AddSubsystemRequirement<UMassNavigationSubsystem>(EMassFragmentAccess::ReadWrite);
 }
 
 void UMassNavigationObstacleRemoverProcessor::Execute(UMassEntitySubsystem& EntitySubsystem, FMassExecutionContext& Context)
diff --git a/Engine/Plugins/AI/MassAI/Source/MassZoneGraphNavigation/Private/MassZoneGraphNavigationProcessors.cpp b/Engine/Plugins/AI/MassAI/Source/MassZoneGraphNavigation/Private/MassZoneGraphNavigationProcessors.cpp
--- a/Engine/Plugins/AI/MassAI/Source/MassZoneGraphNavigation/Private/MassZoneGraphNavigationProcessors.cpp
+++ b/Engine/Plugins/AI/MassAI/Source/MassZoneGraphNavigation/Private/MassZoneGraphNavigationProcessors.cpp
@@ -48,7 +48,7 @@ void UMassZoneGraphLocationInitializer::ConfigureQueries()
 	EntityQuery.AddRequirement<FTransformFragment>(EMassFragmentAccess::ReadOnly);
 	EntityQuery.AddRequirement<FMassMoveTargetFragment>(EMassFragmentAccess::ReadWrite); // Make optional?
 	EntityQuery.AddConstSharedRequirement<FMassZoneGraphNavigationParameters>(EMassFragmentPresence::All);
-	EntityQuery.AddSystemRequirement<UZoneGraphSubsystem>(EMassFragmentAccess::ReadOnly);
+	EntityQuery.AddSubsystemRequirement<UZoneGraphSubsystem>(EMassFragmentAccess::ReadOnly);
 }
 
 void UMassZoneGraphLocationInitializer::Execute(UMassEntitySubsystem& EntitySubsystem, FMassExecutionContext& Context)
@@ -133,7 +133,7 @@ void UMassZoneGraphPathFollowProcessor::ConfigureQueries()
 	EntityQuery_Conditional.AddChunkRequirement<FMassSimulationVariableTickChunkFragment>(EMassFragmentAccess::ReadOnly, EMassFragmentPresence::Optional);
 	EntityQuery_Conditional.SetChunkFilter(&FMassSimulationVariableTickChunkFragment::ShouldTickChunkThisFrame);
 
-	EntityQuery_Conditional.AddSystemRequirement<UZoneGraphSubsystem>(EMassFragmentAccess::ReadOnly);
+	EntityQuery_Conditional.AddSubsystemRequirement<UZoneGraphSubsystem>(EMassFragmentAccess::ReadOnly);
 }
 
 void UMassZoneGraphPathFollowProcessor::Execute(UMassEntitySubsystem& EntitySubsystem, FMassExecutionContext& Context)
diff --git a/Engine/Plugins/AI/MassCrowd/Source/MassCrowd/Private/MassCrowdNavigationProcessor.cpp b/Engine/Plugins/AI/MassCrowd/Source/MassCrowd/Private/MassCrowdNavigationProcessor.cpp
--- a/Engine/Plugins/AI/MassCrowd/Source/MassCrowd/Private/MassCrowdNavigationProcessor.cpp
+++ b/Engine/Plugins/AI/MassCrowd/Source/MassCrowd/Private/MassCrowdNavigationProcessor.cpp
@@ -26,7 +26,7 @@ void UMassCrowdLaneTrackingSignalProcessor::ConfigureQueries()
 	EntityQuery.AddTagRequirement<FMassCrowdTag>(EMassFragmentPresence::All);
 	EntityQuery.AddRequirement<FMassCrowdLaneTrackingFragment>(EMassFragmentAccess::ReadWrite);
 	EntityQuery.AddRequirement<FMassZoneGraphLaneLocationFragment>(EMassFragmentAccess::ReadOnly);
-	EntityQuery.AddSystemRequirement<UMassCrowdSubsystem>(EMassFragmentAccess::ReadWrite);
+	EntityQuery.AddSubsystemRequirement<UMassCrowdSubsystem>(EMassFragmentAccess::ReadWrite);
 }
 
 void UMassCrowdLaneTrackingSignalProcessor::Initialize(UObject& Owner)
@@ -72,7 +72,7 @@ void UMassCrowdLaneTrackingDestructor::ConfigureQueries()
 {
 	EntityQuery.AddTagRequirement<FMassCrowdTag>(EMassFragmentPresence::All);
 	EntityQuery.AddRequirement<FMassCrowdLaneTrackingFragment>(EMassFragmentAccess::ReadOnly);
-	EntityQuery.AddSystemRequirement<UMassCrowdSubsystem>(EMassFragmentAccess::ReadWrite);
+	EntityQuery.AddSubsystemRequirement<UMassCrowdSubsystem>(EMassFragmentAccess::ReadWrite);
 }
 
 void UMassCrowdLaneTrackingDestructor::Execute(UMassEntitySubsystem& EntitySubsystem, FMassExecutionContext& Context)
diff --git a/Engine/Plugins/Runtime/MassEntity/Source/MassEntity/Private/MassEntityQuery.cpp b/Engine/Plugins/Runtime/MassEntity/Source/MassEntity/Private/MassEntityQuery.cpp
--- a/Engine/Plugins/Runtime/MassEntity/Source/MassEntity/Private/MassEntityQuery.cpp
+++ b/Engine/Plugins/Runtime/MassEntity/Source/MassEntity/Private/MassEntityQuery.cpp
@@ -138,7 +138,7 @@ void FMassEntityQuery::ForEachEntityChunk(UMassEntitySubsystem& EntitySubsystem,
 	int32 NumEntitiesToProcess = 0;
 #endif
 
-	ExecutionContext.SetSystemRequirements(RequiredConstSystems, RequiredMutableSystems);
+	ExecutionContext.SetSubsystemRequirements(RequiredConstSubsystems, RequiredMutableSubsystems);
 
 	// if there's a chunk collection set by the external code - use that
 	if (ExecutionContext.GetEntityCollection().IsSet())
diff --git a/Engine/Plugins/Runtime/MassEntity/Source/MassEntity/Private/MassExecutionContext.cpp b/Engine/Plugins/Runtime/MassEntity/Source/MassEntity/Private/MassExecutionContext.cpp
--- a/Engine/Plugins/Runtime/MassEntity/Source/MassEntity/Private/MassExecutionContext.cpp
+++ b/Engine/Plugins/Runtime/MassEntity/Source/MassEntity/Private/MassExecutionContext.cpp
@@ -22,10 +22,10 @@ void FMassExecutionContext::ClearExecutionData()
 	bSubsystemRequirementsSet = false;
 }
 
-void FMassExecutionContext::SetSystemRequirements(const FMassExternalSubystemBitSet& RequiredConstSystems, const FMassExternalSubystemBitSet& RequiredMutableSystems)
+void FMassExecutionContext::SetSubsystemRequirements(const FMassExternalSubystemBitSet& RequiredConstSubsystems, const FMassExternalSubystemBitSet& RequiredMutableSubsystems)
 {
-	ConstSubsystemsBitSet = RequiredConstSystems;
-	MutableSubsystemsBitSet = RequiredMutableSystems;
+	ConstSubsystemsBitSet = RequiredConstSubsystems;
+	MutableSubsystemsBitSet = RequiredMutableSubsystems;
 	bSubsystemRequirementsSet = true;
 }
 
diff --git a/Engine/Plugins/Runtime/MassEntity/Source/MassEntity/Public/MassEntityQuery.h b/Engine/Plugins/Runtime/MassEntity/Source/MassEntity/Public/MassEntityQuery.h
--- a/Engine/Plugins/Runtime/MassEntity/Source/MassEntity/Public/MassEntityQuery.h
+++ b/Engine/Plugins/Runtime/MassEntity/Source/MassEntity/Public/MassEntityQuery.h
@@ -309,17 +309,17 @@ public:
 	}
 
 	template<typename T>
-	FMassEntityQuery& AddSystemRequirement(const EMassFragmentAccess AccessMode)
+	FMassEntityQuery& AddSubsystemRequirement(const EMassFragmentAccess AccessMode)
 	{
 		check(AccessMode != EMassFragmentAccess::None && AccessMode != EMassFragmentAccess::MAX);
 
 		switch (AccessMode)
 		{
 		case EMassFragmentAccess::ReadOnly:
-			RequiredConstSystems.Add<T>();
+			RequiredConstSubsystems.Add<T>();
 			break;
 		case EMassFragmentAccess::ReadWrite:
-			RequiredMutableSystems.Add<T>();
+			RequiredMutableSubsystems.Add<T>();
 			break;
 		default:
 			check(false);
@@ -428,8 +428,8 @@ protected:
 	FMassSharedFragmentBitSet RequiredOptionalSharedFragments;
 	FMassSharedFragmentBitSet RequiredNoneSharedFragments;
 
-	FMassExternalSubystemBitSet RequiredConstSystems;
-	FMassExternalSubystemBitSet RequiredMutableSystems;
+	FMassExternalSubystemBitSet RequiredConstSubsystems;
+	FMassExternalSubystemBitSet RequiredMutableSubsystems;
 
 private:
 	/** 
diff --git a/Engine/Plugins/Runtime/MassEntity/Source/MassEntity/Public/MassExecutionContext.h b/Engine/Plugins/Runtime/MassEntity/Source/MassEntity/Public/MassExecutionContext.h
--- a/Engine/Plugins/Runtime/MassEntity/Source/MassEntity/Public/MassExecutionContext.h
+++ b/Engine/Plugins/Runtime/MassEntity/Source/MassEntity/Public/MassExecutionContext.h
@@ -284,7 +284,7 @@ public:
 	}
 	
 
-	void SetSystemRequirements(const FMassExternalSubystemBitSet& RequiredConstSystems, const FMassExternalSubystemBitSet& RequiredMutableSystems);
+	void SetSubsystemRequirements(const FMassExternalSubystemBitSet& RequiredConstSubsystems, const FMassExternalSubystemBitSet& RequiredMutableSubsystems);
 
 	/** Sparse chunk related operation */
 	const FMassArchetypeEntityCollection& GetEntityCollection() const { return EntityCollection; }
diff --git a/Engine/Plugins/Runtime/MassEntity/Source/MassEntityTestSuite/Private/MassSystemRequirementTest.cpp b/Engine/Plugins/Runtime/MassEntity/Source/MassEntityTestSuite/Private/MassSystemRequirementTest.cpp
--- a/Engine/Plugins/Runtime/MassEntity/Source/MassEntityTestSuite/Private/MassSystemRequirementTest.cpp
+++ b/Engine/Plugins/Runtime/MassEntity/Source/MassEntityTestSuite/Private/MassSystemRequirementTest.cpp
@@ -40,7 +40,7 @@ struct FMutableRequirement : FSystemRequirementTestBase
 
 		FMassEntityQuery EntityQuery;
 		EntityQuery.AddRequirement<FTestFragment_Float>(EMassFragmentAccess::ReadOnly);
-		EntityQuery.AddSystemRequirement<UMassTestWorldSubsystem>(EMassFragmentAccess::ReadWrite);
+		EntityQuery.AddSubsystemRequirement<UMassTestWorldSubsystem>(EMassFragmentAccess::ReadWrite);
 
 		check(EntitySubsystem->GetWorld());
 		UMassTestWorldSubsystem* TestSystemActual = EntitySubsystem->GetWorld()->GetSubsystem<UMassTestWorldSubsystem>();
@@ -71,7 +71,7 @@ struct FConstRequirement : FSystemRequirementTestBase
 
 		FMassEntityQuery EntityQuery;
 		EntityQuery.AddRequirement<FTestFragment_Float>(EMassFragmentAccess::ReadOnly);
-		EntityQuery.AddSystemRequirement<UMassTestWorldSubsystem>(EMassFragmentAccess::ReadOnly);
+		EntityQuery.AddSubsystemRequirement<UMassTestWorldSubsystem>(EMassFragmentAccess::ReadOnly);
 
 		check(EntitySubsystem->GetWorld());
 		UMassTestWorldSubsystem* TestMutableSystem = nullptr;
diff --git a/Engine/Plugins/Runtime/MassGameplay/Source/MassGameplayDebug/Private/AssignDebugVisProcessor.cpp b/Engine/Plugins/Runtime/MassGameplay/Source/MassGameplayDebug/Private/AssignDebugVisProcessor.cpp
--- a/Engine/Plugins/Runtime/MassGameplay/Source/MassGameplayDebug/Private/AssignDebugVisProcessor.cpp
+++ b/Engine/Plugins/Runtime/MassGameplay/Source/MassGameplayDebug/Private/AssignDebugVisProcessor.cpp
@@ -19,7 +19,7 @@ UAssignDebugVisProcessor::UAssignDebugVisProcessor()
 void UAssignDebugVisProcessor::ConfigureQueries()
 {
 	EntityQuery.AddRequirement<FSimDebugVisFragment>(EMassFragmentAccess::ReadWrite);
-	EntityQuery.AddSystemRequirement<UMassDebuggerSubsystem>(EMassFragmentAccess::ReadWrite);
+	EntityQuery.AddSubsystemRequirement<UMassDebuggerSubsystem>(EMassFragmentAccess::ReadWrite);
 }
 
 void UAssignDebugVisProcessor::Execute(UMassEntitySubsystem& EntitySubsystem, FMassExecutionContext& Context)
diff --git a/Engine/Plugins/Runtime/MassGameplay/Source/MassGameplayDebug/Private/DebugVisLocationProcessor.cpp b/Engine/Plugins/Runtime/MassGameplay/Source/MassGameplayDebug/Private/DebugVisLocationProcessor.cpp
--- a/Engine/Plugins/Runtime/MassGameplay/Source/MassGameplayDebug/Private/DebugVisLocationProcessor.cpp
+++ b/Engine/Plugins/Runtime/MassGameplay/Source/MassGameplayDebug/Private/DebugVisLocationProcessor.cpp
@@ -21,7 +21,7 @@ void UDebugVisLocationProcessor::ConfigureQueries()
 	EntityQuery.AddRequirement<FTransformFragment>(EMassFragmentAccess::ReadOnly);
 	EntityQuery.AddRequirement<FSimDebugVisFragment>(EMassFragmentAccess::ReadOnly);
 	EntityQuery.AddTagRequirement<FMassDebuggableTag>(EMassFragmentPresence::All);
-	EntityQuery.AddSystemRequirement<UMassDebuggerSubsystem>(EMassFragmentAccess::ReadWrite);
+	EntityQuery.AddSubsystemRequirement<UMassDebuggerSubsystem>(EMassFragmentAccess::ReadWrite);
 }
 
 void UDebugVisLocationProcessor::Execute(UMassEntitySubsystem& EntitySubsystem, FMassExecutionContext& Context)
@@ -82,7 +82,7 @@ void UMassProcessor_UpdateDebugVis::ConfigureQueries()
 	EntityQuery.AddRequirement<FDataFragment_DebugVis>(EMassFragmentAccess::ReadWrite);
 	EntityQuery.AddRequirement<FAgentRadiusFragment>(EMassFragmentAccess::ReadWrite);
 	EntityQuery.AddTagRequirement<FMassDebuggableTag>(EMassFragmentPresence::All);
-	EntityQuery.AddSystemRequirement<UMassDebuggerSubsystem>(EMassFragmentAccess::ReadWrite);
+	EntityQuery.AddSubsystemRequirement<UMassDebuggerSubsystem>(EMassFragmentAccess::ReadWrite);
 }
 
 void UMassProcessor_UpdateDebugVis::Execute(UMassEntitySubsystem& EntitySubsystem, FMassExecutionContext& Context)
diff --git a/Engine/Plugins/Runtime/MassGameplay/Source/MassReplication/Private/MassReplicationFragments.cpp b/Engine/Plugins/Runtime/MassGameplay/Source/MassReplication/Private/MassReplicationFragments.cpp
--- a/Engine/Plugins/Runtime/MassGameplay/Source/MassReplication/Private/MassReplicationFragments.cpp
+++ b/Engine/Plugins/Runtime/MassGameplay/Source/MassReplication/Private/MassReplicationFragments.cpp
@@ -18,7 +18,7 @@ UMassNetworkIDFragmentInitializer::UMassNetworkIDFragmentInitializer()
 void UMassNetworkIDFragmentInitializer::ConfigureQueries()
 {
 	EntityQuery.AddRequirement<FMassNetworkIDFragment>(EMassFragmentAccess::ReadWrite);
-	EntityQuery.AddSystemRequirement<UMassReplicationSubsystem>(EMassFragmentAccess::ReadWrite);
+	EntityQuery.AddSubsystemRequirement<UMassReplicationSubsystem>(EMassFragmentAccess::ReadWrite);
 }
 
 void UMassNetworkIDFragmentInitializer::Execute(UMassEntitySubsystem& EntitySubsystem, FMassExecutionContext& Context)
diff --git a/Engine/Plugins/Runtime/MassGameplay/Source/MassReplication/Private/MassReplicationGridProcessor.cpp b/Engine/Plugins/Runtime/MassGameplay/Source/MassReplication/Private/MassReplicationGridProcessor.cpp
--- a/Engine/Plugins/Runtime/MassGameplay/Source/MassReplication/Private/MassReplicationGridProcessor.cpp
+++ b/Engine/Plugins/Runtime/MassGameplay/Source/MassReplication/Private/MassReplicationGridProcessor.cpp
@@ -26,7 +26,7 @@ void UMassReplicationGridProcessor::ConfigureQueries()
 	AddToGridEntityQuery.AddRequirement<FTransformFragment>(EMassFragmentAccess::ReadOnly);
 	AddToGridEntityQuery.AddRequirement<FAgentRadiusFragment>(EMassFragmentAccess::ReadOnly);
 	AddToGridEntityQuery.AddRequirement<FMassReplicationGridCellLocationFragment>(EMassFragmentAccess::ReadWrite);
-	AddToGridEntityQuery.AddSystemRequirement<UMassReplicationSubsystem>(EMassFragmentAccess::ReadWrite);
+	AddToGridEntityQuery.AddSubsystemRequirement<UMassReplicationSubsystem>(EMassFragmentAccess::ReadWrite);
 
 	UpdateGridEntityQuery = AddToGridEntityQuery;
 	RemoveFromGridEntityQuery = AddToGridEntityQuery;
@@ -131,7 +131,7 @@ UMassReplicationGridRemoverProcessor::UMassReplicationGridRemoverProcessor()
 void UMassReplicationGridRemoverProcessor::ConfigureQueries()
 {
 	EntityQuery.AddRequirement<FMassReplicationGridCellLocationFragment>(EMassFragmentAccess::ReadWrite);
-	EntityQuery.AddSystemRequirement<UMassReplicationSubsystem>(EMassFragmentAccess::ReadWrite);
+	EntityQuery.AddSubsystemRequirement<UMassReplicationSubsystem>(EMassFragmentAccess::ReadWrite);
 }
 
 void UMassReplicationGridRemoverProcessor::Execute(UMassEntitySubsystem& EntitySubsystem, FMassExecutionContext& Context)
diff --git a/Engine/Plugins/Runtime/MassGameplay/Source/MassRepresentation/Private/MassRepresentationProcessor.cpp b/Engine/Plugins/Runtime/MassGameplay/Source/MassRepresentation/Private/MassRepresentationProcessor.cpp
--- a/Engine/Plugins/Runtime/MassGameplay/Source/MassRepresentation/Private/MassRepresentationProcessor.cpp
+++ b/Engine/Plugins/Runtime/MassGameplay/Source/MassRepresentation/Private/MassRepresentationProcessor.cpp
@@ -35,7 +35,7 @@ void UMassRepresentationProcessor::ConfigureQueries()
 	EntityQuery.AddRequirement<FMassActorFragment>(EMassFragmentAccess::ReadWrite);
 	EntityQuery.AddConstSharedRequirement<FMassRepresentationParameters>();
 	EntityQuery.AddSharedRequirement<FMassRepresentationSubsystemSharedFragment>(EMassFragmentAccess::ReadWrite);
-	EntityQuery.AddSystemRequirement<UMassActorSubsystem>(EMassFragmentAccess::ReadWrite);
+	EntityQuery.AddSubsystemRequirement<UMassActorSubsystem>(EMassFragmentAccess::ReadWrite);
 }
 
 void UMassRepresentationProcessor::Initialize(UObject& Owner)
diff --git a/Engine/Plugins/Runtime/MassGameplay/Source/MassSmartObjects/Private/MassSmartObjectProcessor.cpp b/Engine/Plugins/Runtime/MassGameplay/Source/MassSmartObjects/Private/MassSmartObjectProcessor.cpp
--- a/Engine/Plugins/Runtime/MassGameplay/Source/MassSmartObjects/Private/MassSmartObjectProcessor.cpp
+++ b/Engine/Plugins/Runtime/MassGameplay/Source/MassSmartObjects/Private/MassSmartObjectProcessor.cpp
@@ -31,13 +31,13 @@ void UMassSmartObjectCandidatesFinderProcessor::ConfigureQueries()
 	WorldRequestQuery.AddRequirement<FMassSmartObjectWorldLocationRequestFragment>(EMassFragmentAccess::ReadOnly);
 	WorldRequestQuery.AddRequirement<FMassSmartObjectRequestResultFragment>(EMassFragmentAccess::ReadWrite);
 	WorldRequestQuery.AddTagRequirement<FMassSmartObjectCompletedRequestTag>(EMassFragmentPresence::None);
-	WorldRequestQuery.AddSystemRequirement<USmartObjectSubsystem>(EMassFragmentAccess::ReadOnly);
+	WorldRequestQuery.AddSubsystemRequirement<USmartObjectSubsystem>(EMassFragmentAccess::ReadOnly);
 
 	LaneRequestQuery.AddRequirement<FMassSmartObjectLaneLocationRequestFragment>(EMassFragmentAccess::ReadOnly);
 	LaneRequestQuery.AddRequirement<FMassSmartObjectRequestResultFragment>(EMassFragmentAccess::ReadWrite);
 	LaneRequestQuery.AddTagRequirement<FMassSmartObjectCompletedRequestTag>(EMassFragmentPresence::None);
-	LaneRequestQuery.AddSystemRequirement<UZoneGraphSubsystem>(EMassFragmentAccess::ReadOnly);
-	LaneRequestQuery.AddSystemRequirement<USmartObjectSubsystem>(EMassFragmentAccess::ReadOnly);
+	LaneRequestQuery.AddSubsystemRequirement<UZoneGraphSubsystem>(EMassFragmentAccess::ReadOnly);
+	LaneRequestQuery.AddSubsystemRequirement<USmartObjectSubsystem>(EMassFragmentAccess::ReadOnly);
 }
 
 UMassSmartObjectCandidatesFinderProcessor::UMassSmartObjectCandidatesFinderProcessor()
@@ -300,7 +300,7 @@ void UMassSmartObjectTimedBehaviorProcessor::ConfigureQueries()
 {
 	EntityQuery.AddRequirement<FMassSmartObjectUserFragment>(EMassFragmentAccess::ReadWrite);
 	EntityQuery.AddRequirement<FMassSmartObjectTimedBehaviorFragment>(EMassFragmentAccess::ReadWrite);
-	EntityQuery.AddSystemRequirement<USmartObjectSubsystem>(EMassFragmentAccess::ReadOnly);
+	EntityQuery.AddSubsystemRequirement<USmartObjectSubsystem>(EMassFragmentAccess::ReadOnly);
 }
 
 UMassSmartObjectTimedBehaviorProcessor::UMassSmartObjectTimedBehaviorProcessor()
@@ -384,7 +384,7 @@ UMassSmartObjectUserFragmentDeinitializer::UMassSmartObjectUserFragmentDeinitial
 void UMassSmartObjectUserFragmentDeinitializer::ConfigureQueries()
 {
 	EntityQuery.AddRequirement<FMassSmartObjectUserFragment>(EMassFragmentAccess::ReadWrite);
-	EntityQuery.AddSystemRequirement<USmartObjectSubsystem>(EMassFragmentAccess::ReadWrite);
+	EntityQuery.AddSubsystemRequirement<USmartObjectSubsystem>(EMassFragmentAccess::ReadWrite);
 }
 
 void UMassSmartObjectUserFragmentDeinitializer::Execute(UMassEntitySubsystem& EntitySubsystem, FMassExecutionContext& Context)
