Fixes for MassReplication issues when adding and removing clients.

diff --git a/Engine/Plugins/Runtime/MassGameplay/Source/MassReplication/Private/MassReplicationProcessor.cpp b/Engine/Plugins/Runtime/MassGameplay/Source/MassReplication/Private/MassReplicationProcessor.cpp
--- a/Engine/Plugins/Runtime/MassGameplay/Source/MassReplication/Private/MassReplicationProcessor.cpp
+++ b/Engine/Plugins/Runtime/MassGameplay/Source/MassReplication/Private/MassReplicationProcessor.cpp
@@ -136,8 +136,6 @@ void UMassReplicationProcessor::PrepareExecution(UMassEntitySubsystem& EntitySub
 					ReplicationSubsystem->GetClientBubbleChecked(RepSharedFragment.BubbleInfoClassHandle, CurrentClientHandle) :
 					nullptr;
 
-				check(Info);
-
 				RepSharedFragment.BubbleInfos[Idx] = Info;
 				CachedClientHandle = CurrentClientHandle;
 			}
@@ -162,13 +160,18 @@ void UMassReplicationProcessor::Execute(UMassEntitySubsystem& EntitySubsystem, F
 	UMassLODSubsystem& LODSubsystem = Context.GetMutableSubsystemChecked<UMassLODSubsystem>(EntitySubsystem.GetWorld());
 	const TArray<FViewerInfo>& AllViewersInfo = LODSubsystem.GetViewers();
 	const TArray<FMassClientHandle>& ClientHandles = ReplicationSubsystem->GetClientReplicationHandles();
-	for( const FMassClientHandle ClientHandle : ClientHandles )
+	for (const FMassClientHandle ClientHandle : ClientHandles)
 	{
+		if (ReplicationSubsystem->IsValidClientHandle(ClientHandle) == false)
+		{
+			continue;
+		}
+
 		FMassClientReplicationInfo& ClientReplicationInfo = ReplicationSubsystem->GetMutableClientReplicationInfoChecked(ClientHandle);
 
 		// Figure out all viewer of this client
 		TArray<FViewerInfo> Viewers;
-		for(const FMassViewerHandle ClientViewerHandle : ClientReplicationInfo.Handles)
+		for (const FMassViewerHandle ClientViewerHandle : ClientReplicationInfo.Handles)
 		{
 			const FViewerInfo* ViewerInfo = AllViewersInfo.FindByPredicate([ClientViewerHandle](const FViewerInfo& ViewerInfo) { return ClientViewerHandle == ViewerInfo.Handle; });
 			if (ensureMsgf(ViewerInfo, TEXT("Expecting to find the client viewer handle in the all viewers info list")))
diff --git a/Engine/Plugins/Runtime/MassGameplay/Source/MassReplication/Private/MassReplicationSubsystem.cpp b/Engine/Plugins/Runtime/MassGameplay/Source/MassReplication/Private/MassReplicationSubsystem.cpp
--- a/Engine/Plugins/Runtime/MassGameplay/Source/MassReplication/Private/MassReplicationSubsystem.cpp
+++ b/Engine/Plugins/Runtime/MassGameplay/Source/MassReplication/Private/MassReplicationSubsystem.cpp
@@ -624,7 +624,7 @@ void UMassReplicationSubsystem::AddClient(FMassViewerHandle ViewerHandle, APlaye
 	}
 	else
 	{
-		checkf(ClientsReplicationInfo[ClientHandle.GetIndex()].Handles.Num() == 0, TEXT("ClientsReplicationInfo being replaced must have zero entries, they should have been removed first!"));
+		checkf(ClientsReplicationInfo[ClientHandle.GetIndex()].IsEmpty(), TEXT("ClientsReplicationInfo being replaced must have been reset prior to being reused!"));
 
 		FViewerClientPair& ClientToViewerHandleItem = ClientToViewerHandleArray[ClientHandle.GetIndex()];
 
@@ -671,10 +671,13 @@ void UMassReplicationSubsystem::RemoveClient(FMassClientHandle ClientHandle)
 	checkf(ClientToViewerHandleItem.ViewerHandle.IsValid(), TEXT("Invalid ViewerHandle! ClientHandle is out of sync with ClientToViewerHandleArray!"));
 	checkf(ClientToViewerHandleItem.ClientHandle == ClientHandle, TEXT("ClientHandle is out of sync with ClientToViewerHandleArray!"));
 
-	FMassClientReplicationInfo& ClientViewerHandles = ClientsReplicationInfo[ClientHandle.GetIndex()];
+	{
+		FMassClientReplicationInfo& ClientReplicationInfo = ClientsReplicationInfo[ClientHandle.GetIndex()];
+
+		checkf(ClientReplicationInfo.Handles.Num() > 0, TEXT("There should always be atleast one client viewer handle (the parent NetConnection)"));
 
-	checkf(ClientViewerHandles.Handles.Num() > 0, TEXT("There should always be atleast one client viewer handle (the parent NetConnection)"));
-	ClientViewerHandles.Handles.Reset();
+		ClientReplicationInfo.Reset();
+	}
 
 	checkf(ViewerToClientHandleArray.IsValidIndex(ClientToViewerHandleItem.ViewerHandle.GetIndex()), TEXT("ViewerHandle is out of sync with ViewerToClientHandleArray!"));
 	FViewerClientPair& ViewerToClientHandleItem = ViewerToClientHandleArray[ClientToViewerHandleItem.ViewerHandle.GetIndex()];
diff --git a/Engine/Plugins/Runtime/MassGameplay/Source/MassReplication/Public/MassReplicationSubsystem.h b/Engine/Plugins/Runtime/MassGameplay/Source/MassReplication/Public/MassReplicationSubsystem.h
--- a/Engine/Plugins/Runtime/MassGameplay/Source/MassReplication/Public/MassReplicationSubsystem.h
+++ b/Engine/Plugins/Runtime/MassGameplay/Source/MassReplication/Public/MassReplicationSubsystem.h
@@ -18,6 +18,19 @@ typedef TMap<FMassEntityHandle, FMassReplicatedAgentData> FMassReplicationAgentD
 
 struct FMassClientReplicationInfo
 {
+	void Reset()
+	{
+		Handles.Reset();
+		HandledEntities.Reset();
+		AgentsData.Reset();
+	}
+
+	/** Note this struct is constructed IsEmpty() == true */
+	bool IsEmpty() const
+	{
+		return (Handles.Num() == 0) && (HandledEntities.Num() == 0) && (AgentsData.Num() == 0);
+	}
+
 	/** Array of all the viewer of this client */
 	TArray<FMassViewerHandle> Handles;
 
@@ -106,6 +119,11 @@ public:
 
 	const TArray<FMassClientHandle>& GetClientReplicationHandles() const { return ClientHandleManager.GetHandles(); }
 
+	bool IsValidClientHandle(FMassClientHandle ClientHandle) const
+	{
+		return ClientHandleManager.IsValidHandle(ClientHandle);
+	}
+
 	/** Gets the client bubble safely */
 	AMassClientBubbleInfoBase* GetClientBubble(FMassBubbleInfoClassHandle BubbleClassHandle, FMassClientHandle ClientHandle) const
 	{
@@ -162,11 +180,15 @@ public:
 	}
 
 	/** Gets the client replication info. Faster version using check()s */
-	FMassClientReplicationInfo& GetMutableClientReplicationInfoChecked(FMassClientHandle Handle)
+	FMassClientReplicationInfo* GetMutableClientReplicationInfo(FMassClientHandle Handle)
 	{
-		check(ClientHandleManager.IsValidHandle(Handle));
+		return const_cast<FMassClientReplicationInfo*>(GetClientReplicationInfo(Handle));
+	}
 
-		return ClientsReplicationInfo[Handle.GetIndex()];
+	/** Gets the client replication info. Faster version using check()s */
+	FMassClientReplicationInfo& GetMutableClientReplicationInfoChecked(FMassClientHandle Handle)
+	{
+		return *const_cast<FMassClientReplicationInfo*>(&GetClientReplicationInfoChecked(Handle));
 	}
 
 	void SynchronizeClientsAndViewers();
