Marked MassCharacterMovementTranslators as requiring to be run on game thread since it's accessing physics in a way requiring game-thread accces.

diff --git a/Engine/Plugins/Runtime/MassGameplay/Source/MassActors/Private/Translators/MassCharacterMovementTranslators.cpp b/Engine/Plugins/Runtime/MassGameplay/Source/MassActors/Private/Translators/MassCharacterMovementTranslators.cpp
--- a/Engine/Plugins/Runtime/MassGameplay/Source/MassActors/Private/Translators/MassCharacterMovementTranslators.cpp
+++ b/Engine/Plugins/Runtime/MassGameplay/Source/MassActors/Private/Translators/MassCharacterMovementTranslators.cpp
@@ -134,6 +134,7 @@ UMassCharacterOrientationToActorTranslator::UMassCharacterOrientationToActorTran
 	ExecutionOrder.ExecuteInGroup = UE::Mass::ProcessorGroupNames::UpdateWorldFromMass;
 	ExecutionOrder.ExecuteAfter.Add(UE::Mass::ProcessorGroupNames::Movement);
 	RequiredTags.Add<FMassCharacterOrientationCopyToActorTag>();
+	bRequiresGameThreadExecution = true;
 }
 
 void UMassCharacterOrientationToActorTranslator::ConfigureQueries()
