MassBehavior: Changed smart object eval into a task
- Removed MassStateTreeSmartObjectEvaluator and implemented it as a task
- changed slot claim to not set state on SmartObject user
- changed smart object use to have start/stop, and it's cached on the SmartObject user
- added filter options to find smart object
- made unnecessary properties hidden in editor

diff --git a/Engine/Plugins/AI/MassAI/Source/MassAIBehavior/MassAIBehavior.Build.cs b/Engine/Plugins/AI/MassAI/Source/MassAIBehavior/MassAIBehavior.Build.cs
--- a/Engine/Plugins/AI/MassAI/Source/MassAIBehavior/MassAIBehavior.Build.cs
+++ b/Engine/Plugins/AI/MassAI/Source/MassAIBehavior/MassAIBehavior.Build.cs
@@ -18,7 +18,7 @@ namespace UnrealBuildTool.Rules
 
 			PrivateIncludePaths.AddRange(
 				new string[] {
-				}
+				} 
 			);
 
 			PublicDependencyModuleNames.AddRange(
@@ -27,6 +27,7 @@ namespace UnrealBuildTool.Rules
 					"Core",
 					"CoreUObject",
 					"Engine",
+					"GameplayTags",
 					"MassActors",
 					"MassCommon",
 					"MassLOD",
diff --git a/Engine/Plugins/AI/MassAI/Source/MassAIBehavior/Private/Evaluators/MassStateTreeSmartObjectEvaluator.cpp b/Engine/Plugins/AI/MassAI/Source/MassAIBehavior/Private/Evaluators/MassStateTreeSmartObjectEvaluator.cpp
deleted file mode 100644
diff --git a/Engine/Plugins/AI/MassAI/Source/MassAIBehavior/Private/MassStateTreeProcessors.cpp b/Engine/Plugins/AI/MassAI/Source/MassAIBehavior/Private/MassStateTreeProcessors.cpp
--- a/Engine/Plugins/AI/MassAI/Source/MassAIBehavior/Private/MassStateTreeProcessors.cpp
+++ b/Engine/Plugins/AI/MassAI/Source/MassAIBehavior/Private/MassStateTreeProcessors.cpp
@@ -358,6 +358,7 @@ void UMassStateTreeProcessor::Initialize(UObject& Owner)
 	SubscribeToSignal(UE::Mass::Signals::SmartObjectRequestCandidates);
 	SubscribeToSignal(UE::Mass::Signals::SmartObjectCandidatesReady);
 	SubscribeToSignal(UE::Mass::Signals::SmartObjectInteractionDone);
+	SubscribeToSignal(UE::Mass::Signals::SmartObjectInteractionAborted);
 
 	SubscribeToSignal(UE::Mass::Signals::FollowPointPathStart);
 	SubscribeToSignal(UE::Mass::Signals::FollowPointPathDone);
@@ -393,7 +394,7 @@ void UMassStateTreeProcessor::SignalEntities(UMassEntitySubsystem& EntitySubsyst
 	TArray<FMassEntityHandle> EntitiesToSignal;
 
 	EntityQuery.ForEachEntityChunk(EntitySubsystem, Context,
-		[this, &StateTreeContext, TimeInSeconds, &EntitiesToSignal, World = EntitySubsystem.GetWorld()](FMassExecutionContext& Context)
+		[this, &StateTreeContext, TimeInSeconds, &EntitiesToSignal, World = EntitySubsystem.GetWorld(), &EntitySignals](FMassExecutionContext& Context)
 		{
 			// Keep stats regarding the amount of tree instances ticked per frame
 			CSV_CUSTOM_STAT(StateTreeProcessor, NumTickedStateTree, Context.GetNumEntities(), ECsvCustomStatOp::Accumulate);
@@ -401,12 +402,30 @@ void UMassStateTreeProcessor::SignalEntities(UMassEntitySubsystem& EntitySubsyst
 			UMassStateTreeSubsystem& MassStateTreeSubsystem = Context.GetMutableSubsystemChecked<UMassStateTreeSubsystem>(World);
 
 			UE::MassBehavior::ForEachEntityInChunk(Context, StateTreeContext, MassStateTreeSubsystem,
-				[TimeInSeconds, &EntitiesToSignal](FMassStateTreeExecutionContext& StateTreeExecutionContext, FMassStateTreeInstanceFragment& StateTreeInstance, FStateTreeInstanceData& InstanceData)
+				[TimeInSeconds, &EntitiesToSignal, &EntitySignals, &MassStateTreeSubsystem]
+				(FMassStateTreeExecutionContext& StateTreeExecutionContext, FMassStateTreeInstanceFragment& StateTreeInstance, FStateTreeInstanceData& InstanceData)
 				{
 					// Compute adjusted delta time
 					const float AdjustedDeltaTime = TimeInSeconds - StateTreeInstance.LastUpdateTimeInSeconds;
 					StateTreeInstance.LastUpdateTimeInSeconds = TimeInSeconds;
 
+					const FMassEntityHandle Entity = StateTreeExecutionContext.GetEntity();
+					if (UE::Mass::Debug::IsDebuggingEntity(Entity))
+					{
+						TArray<FName> Signals;
+						EntitySignals.GetSignalsForEntity(Entity, Signals);
+						FString SignalsString;
+						for (const FName& SignalName : Signals)
+						{
+							if (!SignalsString.IsEmpty())
+							{
+								SignalsString += TEXT(", ");
+							}
+							SignalsString += SignalName.ToString();
+						}
+						UE_VLOG_UELOG(&MassStateTreeSubsystem, LogStateTree, Log, TEXT("%s: Ticking StateTree because of signals: %s"), *Entity.DebugGetDescription(), *SignalsString);
+					}
+
 					// Tick the tree instance
 					StateTreeExecutionContext.Tick(AdjustedDeltaTime, &InstanceData);
 
diff --git a/Engine/Plugins/AI/MassAI/Source/MassAIBehavior/Private/Tasks/MassClaimSmartObjectTask.cpp b/Engine/Plugins/AI/MassAI/Source/MassAIBehavior/Private/Tasks/MassClaimSmartObjectTask.cpp
new file mode 100644
--- /dev/null
+++ b/Engine/Plugins/AI/MassAI/Source/MassAIBehavior/Private/Tasks/MassClaimSmartObjectTask.cpp
@@ -0,0 +1,118 @@
+﻿// Copyright Epic Games, Inc. All Rights Reserved.
+
+#include "Tasks/MassClaimSmartObjectTask.h"
+#include "MassCommonFragments.h"
+#include "MassAIBehaviorTypes.h"
+#include "MassSignalSubsystem.h"
+#include "MassSmartObjectHandler.h"
+#include "MassSmartObjectFragments.h"
+#include "MassStateTreeExecutionContext.h"
+#include "MassNavigationFragments.h"
+#include "MassZoneGraphNavigationFragments.h"
+#include "MassZoneGraphNavigationUtils.h"
+#include "Engine/World.h"
+#include "StateTreeLinker.h"
+
+//----------------------------------------------------------------------//
+// FMassClaimSmartObjectTask
+//----------------------------------------------------------------------//
+
+bool FMassClaimSmartObjectTask::Link(FStateTreeLinker& Linker)
+{
+	Linker.LinkExternalData(SmartObjectUserHandle);
+	Linker.LinkExternalData(SmartObjectSubsystemHandle);
+	Linker.LinkExternalData(MassSignalSubsystemHandle);
+
+	Linker.LinkInstanceDataProperty(CandidateSlotsHandle, STATETREE_INSTANCEDATA_PROPERTY(FInstanceDataType, CandidateSlots));
+	Linker.LinkInstanceDataProperty(ClaimedSlotHandle, STATETREE_INSTANCEDATA_PROPERTY(FInstanceDataType, ClaimedSlot));
+
+	return true;
+}
+
+EStateTreeRunStatus FMassClaimSmartObjectTask::EnterState(FStateTreeExecutionContext& Context, const EStateTreeStateChangeType ChangeType, const FStateTreeTransitionResult& Transition) const
+{
+	if (ChangeType != EStateTreeStateChangeType::Changed)
+	{
+		return EStateTreeRunStatus::Running;
+	}
+
+	// Retrieve fragments and subsystems
+	USmartObjectSubsystem& SmartObjectSubsystem = Context.GetExternalData(SmartObjectSubsystemHandle);
+	UMassSignalSubsystem& SignalSubsystem = Context.GetExternalData(MassSignalSubsystemHandle);
+	FMassSmartObjectUserFragment& SOUser = Context.GetExternalData(SmartObjectUserHandle);
+
+	const FMassSmartObjectCandidateSlots& CandidateSlots = Context.GetInstanceData(CandidateSlotsHandle);
+	FSmartObjectClaimHandle& ClaimedSlot = Context.GetInstanceData(ClaimedSlotHandle);
+	ClaimedSlot.Invalidate();
+	
+	// Setup MassSmartObject handler and claim
+	const FMassStateTreeExecutionContext& MassContext = static_cast<FMassStateTreeExecutionContext&>(Context);
+	const FMassSmartObjectHandler MassSmartObjectHandler(MassContext.GetEntitySubsystem(), MassContext.GetEntitySubsystemExecutionContext(), SmartObjectSubsystem, SignalSubsystem);
+	
+	ClaimedSlot = MassSmartObjectHandler.ClaimCandidate(MassContext.GetEntity(), SOUser, CandidateSlots);
+
+	// Treat claiming a slot as consuming all the candidate slots.
+	// This is done here because of the limited ways we can communicate between FindSmartObject() and ClaimSmartObject().
+	// InteractionCooldownEndTime is used by the FindSmartObject() to invalidate the candidates.
+	SOUser.InteractionCooldownEndTime = Context.GetWorld()->GetTimeSeconds() + InteractionCooldown;
+
+	if (!ClaimedSlot.IsValid())
+	{
+		MASSBEHAVIOR_LOG(Log, TEXT("Failed to claim smart object slot from %d candidates"), CandidateSlots.NumSlots);
+		return EStateTreeRunStatus::Failed;
+	}
+
+	return EStateTreeRunStatus::Running;
+}
+
+void FMassClaimSmartObjectTask::ExitState(FStateTreeExecutionContext& Context, const EStateTreeStateChangeType ChangeType, const FStateTreeTransitionResult& Transition) const
+{
+	if (ChangeType != EStateTreeStateChangeType::Changed)
+	{
+		return;
+	}
+
+	const FSmartObjectClaimHandle& ClaimedSlot = Context.GetInstanceData(ClaimedSlotHandle);
+	FMassSmartObjectUserFragment& SOUser = Context.GetExternalData(SmartObjectUserHandle);
+
+	// Succeeded or not, prevent interactions for a specified duration.
+	SOUser.InteractionCooldownEndTime = Context.GetWorld()->GetTimeSeconds() + InteractionCooldown;
+
+	if (ClaimedSlot.IsValid())
+	{
+		const FMassStateTreeExecutionContext& MassContext = static_cast<FMassStateTreeExecutionContext&>(Context);
+		USmartObjectSubsystem& SmartObjectSubsystem = Context.GetExternalData(SmartObjectSubsystemHandle);
+		UMassSignalSubsystem& SignalSubsystem = Context.GetExternalData(MassSignalSubsystemHandle);
+		const FMassSmartObjectHandler MassSmartObjectHandler(MassContext.GetEntitySubsystem(), MassContext.GetEntitySubsystemExecutionContext(), SmartObjectSubsystem, SignalSubsystem);
+
+		MassSmartObjectHandler.ReleaseSmartObject(MassContext.GetEntity(), SOUser, ClaimedSlot);
+	}
+	else
+	{
+		MASSBEHAVIOR_LOG(VeryVerbose, TEXT("Exiting state with an invalid ClaimHandle: nothing to do."));
+	}
+}
+
+EStateTreeRunStatus FMassClaimSmartObjectTask::Tick(FStateTreeExecutionContext& Context, const float DeltaTime) const
+{
+	FMassSmartObjectUserFragment& SOUser = Context.GetExternalData(SmartObjectUserHandle);
+
+	// Prevent FindSmartObject() to query new objects while claimed.
+	// This is done here because of the limited ways we can communicate between FindSmartObject() and ClaimSmartObject().
+	// InteractionCooldownEndTime is used by the FindSmartObject() to invalidate the candidates.
+	SOUser.InteractionCooldownEndTime = Context.GetWorld()->GetTimeSeconds() + InteractionCooldown;
+
+	// Check that the claimed slot is still valid, and if not, fail the task.
+	// The slot can become invalid if the whole SO or slot becomes invalidated.
+	FSmartObjectClaimHandle& ClaimedSlot = Context.GetInstanceData(ClaimedSlotHandle);
+	if (ClaimedSlot.IsValid())
+	{
+		const USmartObjectSubsystem& SmartObjectSubsystem = Context.GetExternalData(SmartObjectSubsystemHandle);
+		if (!SmartObjectSubsystem.IsClaimedSmartObjectValid(ClaimedSlot))
+		{
+			ClaimedSlot.Invalidate();
+		}
+	}
+
+	return ClaimedSlot.IsValid() ? EStateTreeRunStatus::Running : EStateTreeRunStatus::Failed;
+}
diff --git a/Engine/Plugins/AI/MassAI/Source/MassAIBehavior/Private/Tasks/MassFindSmartObjectTask.cpp b/Engine/Plugins/AI/MassAI/Source/MassAIBehavior/Private/Tasks/MassFindSmartObjectTask.cpp
new file mode 100644
--- /dev/null
+++ b/Engine/Plugins/AI/MassAI/Source/MassAIBehavior/Private/Tasks/MassFindSmartObjectTask.cpp
@@ -0,0 +1,180 @@
+﻿// Copyright Epic Games, Inc. All Rights Reserved.
+
+#include "Tasks/MassFindSmartObjectTask.h"
+#include "ZoneGraphAnnotationSubsystem.h"
+#include "MassAIBehaviorTypes.h"
+#include "MassCommonFragments.h"
+#include "MassSmartObjectFragments.h"
+#include "MassSmartObjectHandler.h"
+#include "MassStateTreeExecutionContext.h"
+#include "MassZoneGraphNavigationFragments.h"
+#include "SmartObjectZoneAnnotations.h"
+#include "StateTreeLinker.h"
+
+bool FMassFindSmartObjectTask::Link(FStateTreeLinker& Linker)
+{
+	Linker.LinkExternalData(SmartObjectSubsystemHandle);
+	Linker.LinkExternalData(MassSignalSubsystemHandle);
+	Linker.LinkExternalData(EntityTransformHandle);
+	Linker.LinkExternalData(SmartObjectUserHandle);
+	Linker.LinkExternalData(LocationHandle);
+	
+	Linker.LinkInstanceDataProperty(FoundCandidateSlotsHandle, STATETREE_INSTANCEDATA_PROPERTY(FInstanceDataType, FoundCandidateSlots));
+	Linker.LinkInstanceDataProperty(HasCandidateSlotsHandle, STATETREE_INSTANCEDATA_PROPERTY(FInstanceDataType, bHasCandidateSlots));
+	Linker.LinkInstanceDataProperty(SearchRequestIDHandle, STATETREE_INSTANCEDATA_PROPERTY(FInstanceDataType, SearchRequestID));
+	Linker.LinkInstanceDataProperty(NextUpdateHandle, STATETREE_INSTANCEDATA_PROPERTY(FInstanceDataType, NextUpdate));
+	Linker.LinkInstanceDataProperty(LastLaneHandle, STATETREE_INSTANCEDATA_PROPERTY(FInstanceDataType, LastLane));
+
+	return true;
+}
+
+void FMassFindSmartObjectTask::ExitState(FStateTreeExecutionContext& Context, const EStateTreeStateChangeType ChangeType, const FStateTreeTransitionResult& Transition) const
+{
+	// Do not clear the request on sustained transitions.
+	// A child state (move) task can succeed on the same tick as the request is made (very likely in event based ticking).
+	// That will cause transitions which would kill out request immediately.
+	if (ChangeType == EStateTreeStateChangeType::Sustained)
+	{
+		return;
+	}
+	
+	// Stop any request that are still in flight.
+	FMassSmartObjectRequestID& SearchRequestID = Context.GetInstanceData(SearchRequestIDHandle);
+	
+	if (SearchRequestID.IsSet())
+	{
+		const FMassStateTreeExecutionContext& MassContext = static_cast<FMassStateTreeExecutionContext&>(Context);
+		USmartObjectSubsystem& SmartObjectSubsystem = Context.GetExternalData(SmartObjectSubsystemHandle);
+		UMassSignalSubsystem& SignalSubsystem = Context.GetExternalData(MassSignalSubsystemHandle);
+		const FMassSmartObjectHandler MassSmartObjectHandler(
+			MassContext.GetEntitySubsystem(),
+			MassContext.GetEntitySubsystemExecutionContext(),
+			SmartObjectSubsystem,
+			SignalSubsystem);
+		MassSmartObjectHandler.RemoveRequest(SearchRequestID);
+		SearchRequestID.Reset();
+
+		MASSBEHAVIOR_LOG(Verbose, TEXT("Cancelling pending SmartObject search on ExitState."));
+	}
+}
+
+void FMassFindSmartObjectTask::StateCompleted(FStateTreeExecutionContext& Context, const EStateTreeRunStatus CompletionStatus, const FStateTreeActiveStates& CompletedActiveStates) const
+{
+	const UWorld* World = Context.GetWorld();
+	FMassSmartObjectUserFragment& SOUser = Context.GetExternalData(SmartObjectUserHandle);
+
+	// This is done here because of the limited ways we can communicate between FindSmartObject() and ClaimSmartObject().
+	// ClaimSmartObject() sets the InteractionCooldownEndTime when it tries to claim the candidates.
+	// Use that to signal that the candidates have been consumed (either in success or failure).
+	// Doing the reset here, allows the conditions relying on bHasCandidateSlots to function properly
+	// in failure cases (i.e. don't try to use failed slot).
+	// This code assumes that ClaimSmartObject() accesses the candidates only in EnterState() and Tick().
+	if (SOUser.InteractionHandle.IsValid() || SOUser.InteractionCooldownEndTime > World->GetTimeSeconds())
+	{
+		MASSBEHAVIOR_LOG(Verbose, TEXT("StateCompleted: Reset candidates because of interaction cooldown."));
+
+		FMassSmartObjectCandidateSlots& FoundSlots = Context.GetInstanceData(FoundCandidateSlotsHandle);
+		bool& bHasCandidateSlots = Context.GetInstanceData(HasCandidateSlotsHandle);
+
+		FoundSlots.Reset();
+		bHasCandidateSlots = false;
+	}
+}
+
+EStateTreeRunStatus FMassFindSmartObjectTask::Tick(FStateTreeExecutionContext& Context, const float DeltaTime) const
+{
+	const UWorld* World = Context.GetWorld();
+	USmartObjectSubsystem& SmartObjectSubsystem = Context.GetExternalData(SmartObjectSubsystemHandle);
+	UMassSignalSubsystem& SignalSubsystem = Context.GetExternalData(MassSignalSubsystemHandle);
+	const FMassStateTreeExecutionContext& MassContext = static_cast<FMassStateTreeExecutionContext&>(Context);
+	const FMassSmartObjectHandler MassSmartObjectHandler(
+		MassContext.GetEntitySubsystem(),
+		MassContext.GetEntitySubsystemExecutionContext(),
+		SmartObjectSubsystem,
+		SignalSubsystem);
+	
+	FMassSmartObjectUserFragment& SOUser = Context.GetExternalData(SmartObjectUserHandle);
+	FMassSmartObjectRequestID& SearchRequestID = Context.GetInstanceData(SearchRequestIDHandle);
+	
+	// Try to search for new slots if not already in progress.
+	if (!SearchRequestID.IsSet())
+	{
+		// If the user is already using a SmartObject, or has used interaction recently, skip search and empty results. 
+		if (SOUser.InteractionHandle.IsValid() || SOUser.InteractionCooldownEndTime > World->GetTimeSeconds())
+		{
+			MASSBEHAVIOR_LOG(Verbose, TEXT("Skipped: Recently interacted (%s %.1f)"), SOUser.InteractionHandle.IsValid() ? TEXT("Interacting") : TEXT("Cooldown"), FMath::Max(0.0f, SOUser.InteractionCooldownEndTime - World->GetTimeSeconds()));
+
+			// Do not offer any new candidates during cool down.
+			FMassSmartObjectCandidateSlots& FoundSlots = Context.GetInstanceData(FoundCandidateSlotsHandle);
+			bool& bHasCandidateSlots = Context.GetInstanceData(HasCandidateSlotsHandle);
+
+			FoundSlots.Reset();
+			bHasCandidateSlots = false;
+
+			return EStateTreeRunStatus::Running;
+		}
+
+		// Check to see if we should request. 
+		const FMassZoneGraphLaneLocationFragment* LaneLocation = Context.GetExternalDataPtr(LocationHandle);
+		const FZoneGraphLaneHandle LastLane = Context.GetInstanceData(LastLaneHandle);
+		const bool bLaneHasChanged = (LaneLocation && LastLane != LaneLocation->LaneHandle);
+		const float NextUpdateTime = Context.GetInstanceData(NextUpdateHandle);
+		const bool bTimeForNextUpdate = World->GetTimeSeconds() > NextUpdateTime;
+
+		if (bTimeForNextUpdate || bLaneHasChanged)
+		{
+			// Use lanes if possible for faster queries using zone graph annotations
+			const FMassEntityHandle RequestingEntity = MassContext.GetEntity();
+			if (LaneLocation != nullptr)
+			{
+				MASSBEHAVIOR_CLOG(!LaneLocation->LaneHandle.IsValid(), Error, TEXT("Always expecting a valid lane from the ZoneGraph movement"));
+				if (LaneLocation->LaneHandle.IsValid())
+				{
+					MASSBEHAVIOR_LOG(Log, TEXT("Requesting search candidates from lane %s (%s/%s)"),
+						*LaneLocation->LaneHandle.ToString(),
+						*LexToString(LaneLocation->DistanceAlongLane),
+						*LexToString(LaneLocation->LaneLength));
+
+					SearchRequestID = MassSmartObjectHandler.FindCandidatesAsync(RequestingEntity, SOUser.UserTags, ActivityRequirements, { LaneLocation->LaneHandle, LaneLocation->DistanceAlongLane });
+				}
+			}
+			else
+			{
+				const FTransformFragment& TransformFragment = Context.GetExternalData(EntityTransformHandle);
+				SearchRequestID = MassSmartObjectHandler.FindCandidatesAsync(RequestingEntity, SOUser.UserTags, ActivityRequirements, TransformFragment.GetTransform().GetLocation());
+			}
+		}
+	}
+	else
+	{
+		// Poll to see if the candidates are ready.
+		// A "candidates ready" signal will trigger the state tree evaluation when candidates are ready.
+		if (const FMassSmartObjectCandidateSlots* NewCandidates = MassSmartObjectHandler.GetRequestCandidates(SearchRequestID))
+		{
+			FMassSmartObjectCandidateSlots& FoundSlots = Context.GetInstanceData(FoundCandidateSlotsHandle);
+			bool& bHasCandidateSlots = Context.GetInstanceData(HasCandidateSlotsHandle);
+			float& NextUpdate = Context.GetInstanceData(NextUpdateHandle);
+
+			MASSBEHAVIOR_LOG(Log, TEXT("Found %d smart object candidates"), NewCandidates->NumSlots);
+
+			FoundSlots = *NewCandidates;
+			bHasCandidateSlots = FoundSlots.NumSlots > 0;
+			
+			// Remove requests
+			MassSmartObjectHandler.RemoveRequest(SearchRequestID);
+			SearchRequestID.Reset();
+
+			// Schedule next update.
+			const FMassEntityHandle Entity = MassContext.GetEntity(); 
+
+			constexpr float SearchIntervalDeviation = 0.1f;
+			const float DelayInSeconds = SearchInterval * FMath::FRandRange(1.0f - SearchIntervalDeviation, 1.0f + SearchIntervalDeviation);
+				
+			NextUpdate = World->GetTimeSeconds() + DelayInSeconds;
+			UMassSignalSubsystem& MassSignalSubsystem = Context.GetExternalData(MassSignalSubsystemHandle);
+			MassSignalSubsystem.DelaySignalEntity(UE::Mass::Signals::SmartObjectRequestCandidates, Entity, DelayInSeconds);
+		}
+	}
+	
+	return EStateTreeRunStatus::Running;
+}
diff --git a/Engine/Plugins/AI/MassAI/Source/MassAIBehavior/Private/Tasks/MassStateTreeSmartObjectTask.cpp b/Engine/Plugins/AI/MassAI/Source/MassAIBehavior/Private/Tasks/MassStateTreeSmartObjectTask.cpp
deleted file mode 100644
diff --git a/Engine/Plugins/AI/MassAI/Source/MassAIBehavior/Private/Tasks/MassUseSmartObjectTask.cpp b/Engine/Plugins/AI/MassAI/Source/MassAIBehavior/Private/Tasks/MassUseSmartObjectTask.cpp
new file mode 100644
--- /dev/null
+++ b/Engine/Plugins/AI/MassAI/Source/MassAIBehavior/Private/Tasks/MassUseSmartObjectTask.cpp
@@ -0,0 +1,153 @@
+// Copyright Epic Games, Inc. All Rights Reserved.
+
+#include "Tasks/MassUseSmartObjectTask.h"
+#include "MassCommonFragments.h"
+#include "MassAIBehaviorTypes.h"
+#include "MassSignalSubsystem.h"
+#include "MassSmartObjectHandler.h"
+#include "MassSmartObjectFragments.h"
+#include "MassStateTreeExecutionContext.h"
+#include "MassNavigationFragments.h"
+#include "MassZoneGraphNavigationFragments.h"
+#include "MassZoneGraphNavigationUtils.h"
+#include "Engine/World.h"
+#include "StateTreeLinker.h"
+
+//----------------------------------------------------------------------//
+// FMassUseSmartObjectTask
+//----------------------------------------------------------------------//
+
+bool FMassUseSmartObjectTask::Link(FStateTreeLinker& Linker)
+{
+	Linker.LinkExternalData(SmartObjectSubsystemHandle);
+	Linker.LinkExternalData(MassSignalSubsystemHandle);
+	Linker.LinkExternalData(EntityTransformHandle);
+	Linker.LinkExternalData(SmartObjectUserHandle);
+	Linker.LinkExternalData(MoveTargetHandle);
+
+	Linker.LinkInstanceDataProperty(ClaimedSlotHandle, STATETREE_INSTANCEDATA_PROPERTY(FInstanceDataType, ClaimedSlot));
+
+	return true;
+}
+
+EStateTreeRunStatus FMassUseSmartObjectTask::EnterState(FStateTreeExecutionContext& Context, const EStateTreeStateChangeType ChangeType, const FStateTreeTransitionResult& Transition) const
+{
+	if (ChangeType != EStateTreeStateChangeType::Changed)
+	{
+		return EStateTreeRunStatus::Running;
+	}
+
+	FMassSmartObjectUserFragment& SOUser = Context.GetExternalData(SmartObjectUserHandle);
+	
+	if (SOUser.InteractionHandle.IsValid())
+	{
+		MASSBEHAVIOR_LOG(Error, TEXT("Agent is already using smart object slot %s."), *LexToString(SOUser.InteractionHandle));
+		return EStateTreeRunStatus::Failed;
+	}
+
+	USmartObjectSubsystem& SmartObjectSubsystem = Context.GetExternalData(SmartObjectSubsystemHandle);
+	UMassSignalSubsystem& SignalSubsystem = Context.GetExternalData(MassSignalSubsystemHandle);
+	FMassMoveTargetFragment& MoveTarget = Context.GetExternalData(MoveTargetHandle);
+
+	const FSmartObjectClaimHandle& ClaimedSlot = Context.GetInstanceData(ClaimedSlotHandle);
+
+	// Setup MassSmartObject handler and start interaction
+	const FMassStateTreeExecutionContext& MassContext = static_cast<FMassStateTreeExecutionContext&>(Context);
+	const FMassSmartObjectHandler MassSmartObjectHandler(MassContext.GetEntitySubsystem(), MassContext.GetEntitySubsystemExecutionContext(), SmartObjectSubsystem, SignalSubsystem);
+
+	if (!MassSmartObjectHandler.StartUsingSmartObject(MassContext.GetEntity(), SOUser, ClaimedSlot))
+	{
+		return EStateTreeRunStatus::Failed;
+	}
+
+	// @todo: we should have common API to control this, currently handle via tasks.
+	const UWorld* World = Context.GetWorld();
+	checkf(World != nullptr, TEXT("A valid world is expected from the execution context"));
+
+	MoveTarget.CreateNewAction(EMassMovementAction::Animate, *World);
+	const bool bSuccess = UE::MassNavigation::ActivateActionAnimate(*World, Context.GetOwner(), MassContext.GetEntity(), MoveTarget);
+
+	return bSuccess ? EStateTreeRunStatus::Running : EStateTreeRunStatus::Failed;
+}
+
+void FMassUseSmartObjectTask::ExitState(FStateTreeExecutionContext& Context, const EStateTreeStateChangeType ChangeType, const FStateTreeTransitionResult& Transition) const
+{
+	if (ChangeType == EStateTreeStateChangeType::Changed)
+	{
+		FMassSmartObjectUserFragment& SOUser = Context.GetExternalData(SmartObjectUserHandle);
+
+		if (SOUser.InteractionHandle.IsValid())
+		{
+			MASSBEHAVIOR_LOG(VeryVerbose, TEXT("Exiting state with a valid InteractionHandle: stop using the smart object."));
+
+			const FMassStateTreeExecutionContext& MassContext = static_cast<FMassStateTreeExecutionContext&>(Context);
+			USmartObjectSubsystem& SmartObjectSubsystem = Context.GetExternalData(SmartObjectSubsystemHandle);
+			UMassSignalSubsystem& SignalSubsystem = Context.GetExternalData(MassSignalSubsystemHandle);
+			const FMassSmartObjectHandler MassSmartObjectHandler(MassContext.GetEntitySubsystem(), MassContext.GetEntitySubsystemExecutionContext(), SmartObjectSubsystem, SignalSubsystem);
+			MassSmartObjectHandler.StopUsingSmartObject(MassContext.GetEntity(), SOUser, EMassSmartObjectInteractionStatus::Aborted);
+		}
+		else
+		{
+			MASSBEHAVIOR_LOG(VeryVerbose, TEXT("Exiting state with an invalid ClaimHandle: nothing to do."));
+		}
+	}
+}
+
+void FMassUseSmartObjectTask::StateCompleted(FStateTreeExecutionContext& Context, const EStateTreeRunStatus CompletionStatus, const FStateTreeActiveStates& CompletedActiveStates) const
+{
+	FMassSmartObjectUserFragment& SOUser = Context.GetExternalData(SmartObjectUserHandle);
+
+	if (SOUser.InteractionHandle.IsValid())
+	{
+		MASSBEHAVIOR_LOG(VeryVerbose, TEXT("Completing state with a valid InteractionHandle: stop using the smart object."));
+
+		const FMassStateTreeExecutionContext& MassContext = static_cast<FMassStateTreeExecutionContext&>(Context);
+		USmartObjectSubsystem& SmartObjectSubsystem = Context.GetExternalData(SmartObjectSubsystemHandle);
+		UMassSignalSubsystem& SignalSubsystem = Context.GetExternalData(MassSignalSubsystemHandle);
+		const FMassSmartObjectHandler MassSmartObjectHandler(MassContext.GetEntitySubsystem(), MassContext.GetEntitySubsystemExecutionContext(), SmartObjectSubsystem, SignalSubsystem);
+		const EMassSmartObjectInteractionStatus NewStatus = (CompletionStatus == EStateTreeRunStatus::Succeeded)
+																? EMassSmartObjectInteractionStatus::TaskCompleted
+																: EMassSmartObjectInteractionStatus::Aborted;
+		MassSmartObjectHandler.StopUsingSmartObject(MassContext.GetEntity(), SOUser, NewStatus);
+	}
+}
+
+EStateTreeRunStatus FMassUseSmartObjectTask::Tick(FStateTreeExecutionContext& Context, const float DeltaTime) const
+{
+	EStateTreeRunStatus Status = EStateTreeRunStatus::Failed;
+
+	FMassSmartObjectUserFragment& SOUser = Context.GetExternalData(SmartObjectUserHandle);
+	switch (SOUser.InteractionStatus)
+	{
+	case EMassSmartObjectInteractionStatus::InProgress:
+		MASSBEHAVIOR_LOG(VeryVerbose, TEXT("Interacting ..."));
+		Status = EStateTreeRunStatus::Running;
+		break;
+
+	case EMassSmartObjectInteractionStatus::BehaviorCompleted:
+		MASSBEHAVIOR_LOG(Log, TEXT("Behavior completed"));
+		Status = EStateTreeRunStatus::Succeeded;
+		break;
+
+	case EMassSmartObjectInteractionStatus::TaskCompleted:
+		ensureMsgf(false, TEXT("Not expecting to tick an already completed task"));
+		Status = EStateTreeRunStatus::Failed;
+		break;
+
+	case EMassSmartObjectInteractionStatus::Aborted:
+		MASSBEHAVIOR_LOG(Log, TEXT("Interaction aborted"));
+		Status = EStateTreeRunStatus::Failed;
+		break;
+
+	case EMassSmartObjectInteractionStatus::Unset:
+		MASSBEHAVIOR_LOG(Error, TEXT("Error while using smart object: interaction state is not valid"));
+		Status = EStateTreeRunStatus::Failed;
+		break;
+
+	default:
+		ensureMsgf(false, TEXT("Unhandled interaction status % => Returning EStateTreeRunStatus::Failed"), *UEnum::GetValueAsString(SOUser.InteractionStatus));
+		Status = EStateTreeRunStatus::Failed;
+	}
+
+	return Status;
+}
diff --git a/Engine/Plugins/AI/MassAI/Source/MassAIBehavior/Private/Tasks/MassZoneGraphFindSmartObjectTarget.cpp b/Engine/Plugins/AI/MassAI/Source/MassAIBehavior/Private/Tasks/MassZoneGraphFindSmartObjectTarget.cpp
--- a/Engine/Plugins/AI/MassAI/Source/MassAIBehavior/Private/Tasks/MassZoneGraphFindSmartObjectTarget.cpp
+++ b/Engine/Plugins/AI/MassAI/Source/MassAIBehavior/Private/Tasks/MassZoneGraphFindSmartObjectTarget.cpp
@@ -3,6 +3,7 @@
 #include "Tasks/MassZoneGraphFindSmartObjectTarget.h"
 #include "ZoneGraphAnnotationSubsystem.h"
 #include "MassAIBehaviorTypes.h"
+#include "SmartObjectSubsystem.h"
 #include "MassSmartObjectFragments.h"
 #include "MassSmartObjectSettings.h"
 #include "MassStateTreeExecutionContext.h"
@@ -12,11 +13,12 @@
 
 bool FMassZoneGraphFindSmartObjectTarget::Link(FStateTreeLinker& Linker)
 {
-	Linker.LinkExternalData(SmartObjectUserHandle);
 	Linker.LinkExternalData(LocationHandle);
 	Linker.LinkExternalData(AnnotationSubsystemHandle);
+	Linker.LinkExternalData(SmartObjectSubsystemHandle);
 
-	Linker.LinkInstanceDataProperty(SmartObjectLocationHandle, STATETREE_INSTANCEDATA_PROPERTY(FMassZoneGraphFindSmartObjectTargetInstanceData, SmartObjectLocation));
+	Linker.LinkInstanceDataProperty(ClaimedSlotHandle, STATETREE_INSTANCEDATA_PROPERTY(FInstanceDataType, ClaimedSlot));
+	Linker.LinkInstanceDataProperty(SmartObjectLocationHandle, STATETREE_INSTANCEDATA_PROPERTY(FInstanceDataType, SmartObjectLocation));
 	
 	return true;
 }
@@ -24,14 +26,15 @@ bool FMassZoneGraphFindSmartObjectTarget::Link(FStateTreeLinker& Linker)
 EStateTreeRunStatus FMassZoneGraphFindSmartObjectTarget::EnterState(FStateTreeExecutionContext& Context, const EStateTreeStateChangeType ChangeType, const FStateTreeTransitionResult& Transition) const
 {
 	const UZoneGraphAnnotationSubsystem& AnnotationSubsystem = Context.GetExternalData(AnnotationSubsystemHandle);
-	const FMassSmartObjectUserFragment& SOUser = Context.GetExternalData(SmartObjectUserHandle);
+	const USmartObjectSubsystem& SmartObjectSubsystem = Context.GetExternalData(SmartObjectSubsystemHandle);
 	const FMassZoneGraphLaneLocationFragment& LaneLocation = Context.GetExternalData(LocationHandle);
 	const FZoneGraphLaneHandle LaneHandle(LaneLocation.LaneHandle);
 
+	const FSmartObjectClaimHandle& ClaimedSlot = Context.GetInstanceData(ClaimedSlotHandle);
 	FMassZoneGraphTargetLocation& SmartObjectLocation = Context.GetInstanceData(SmartObjectLocationHandle);
 	SmartObjectLocation.Reset();
 
-	if (!SOUser.ClaimHandle.SmartObjectHandle.IsValid())
+	if (!ClaimedSlot.SmartObjectHandle.IsValid())
 	{
 		MASSBEHAVIOR_LOG(Error, TEXT("Invalid claimed smart object ID."));
 		return EStateTreeRunStatus::Failed;
@@ -45,13 +48,14 @@ EStateTreeRunStatus FMassZoneGraphFindSmartObjectTarget::EnterState(FStateTreeEx
 
 	const FZoneGraphTag SmartObjectTag = GetDefault<UMassSmartObjectSettings>()->SmartObjectTag;
 	const USmartObjectZoneAnnotations* SOAnnotations = Cast<USmartObjectZoneAnnotations>(AnnotationSubsystem.GetFirstAnnotationForTag(SmartObjectTag));
+	const FTransform Transform = SmartObjectSubsystem.GetSlotTransform(ClaimedSlot).Get(FTransform::Identity);
 
 	SmartObjectLocation.LaneHandle = LaneHandle;
 	SmartObjectLocation.NextExitLinkType = EZoneLaneLinkType::None;
 	SmartObjectLocation.NextLaneHandle.Reset();
 	SmartObjectLocation.bMoveReverse = false;
 	SmartObjectLocation.EndOfPathIntent = EMassMovementAction::Stand;
-	SmartObjectLocation.EndOfPathPosition = SOUser.TargetLocation;
+	SmartObjectLocation.EndOfPathPosition = Transform.GetLocation();
 	// Can't set direction at the moment since it seems problematic if it's opposite to the steering direction
 	//SmartObjectLocation.EndOfPathDirection = SOUser.TargetDirection;
 
@@ -62,7 +66,7 @@ EStateTreeRunStatus FMassZoneGraphFindSmartObjectTarget::EnterState(FStateTreeEx
 	TOptional<FSmartObjectLaneLocation> SmartObjectLaneLocation;
 	if (SOAnnotations != nullptr)
 	{
-		SmartObjectLaneLocation = SOAnnotations->GetSmartObjectLaneLocation(LaneHandle.DataHandle, SOUser.ClaimHandle.SmartObjectHandle);
+		SmartObjectLaneLocation = SOAnnotations->GetSmartObjectLaneLocation(LaneHandle.DataHandle, ClaimedSlot.SmartObjectHandle);
 	}
 
 	if (SmartObjectLaneLocation.IsSet())
diff --git a/Engine/Plugins/AI/MassAI/Source/MassAIBehavior/Public/Evaluators/MassStateTreeSmartObjectEvaluator.h b/Engine/Plugins/AI/MassAI/Source/MassAIBehavior/Public/Evaluators/MassStateTreeSmartObjectEvaluator.h
deleted file mode 100644
diff --git a/Engine/Plugins/AI/MassAI/Source/MassAIBehavior/Public/Tasks/MassClaimSmartObjectTask.h b/Engine/Plugins/AI/MassAI/Source/MassAIBehavior/Public/Tasks/MassClaimSmartObjectTask.h
new file mode 100644
--- /dev/null
+++ b/Engine/Plugins/AI/MassAI/Source/MassAIBehavior/Public/Tasks/MassClaimSmartObjectTask.h
@@ -0,0 +1,56 @@
+﻿// Copyright Epic Games, Inc. All Rights Reserved.
+
+#pragma once
+
+#include "MassSmartObjectRequest.h"
+#include "MassStateTreeTypes.h"
+#include "MassClaimSmartObjectTask.generated.h"
+
+struct FStateTreeExecutionContext;
+struct FMassSmartObjectUserFragment;
+class USmartObjectSubsystem;
+class UMassSignalSubsystem;
+struct FTransformFragment;
+struct FMassMoveTargetFragment;
+
+/**
+ * Tasks to claim a smart object from search results and release it when done.
+ */
+USTRUCT()
+struct MASSAIBEHAVIOR_API FMassClaimSmartObjectTaskInstanceData
+{
+	GENERATED_BODY()
+
+	/** Result of the candidates search request (Input) */
+	UPROPERTY(VisibleAnywhere, Category = Input)
+	FMassSmartObjectCandidateSlots CandidateSlots;
+
+	UPROPERTY(VisibleAnywhere, Category = Output)
+	FSmartObjectClaimHandle ClaimedSlot;
+};
+
+USTRUCT(meta = (DisplayName = "Claim SmartObject"))
+struct MASSAIBEHAVIOR_API FMassClaimSmartObjectTask : public FMassStateTreeTaskBase
+{
+	GENERATED_BODY()
+
+	using FInstanceDataType = FMassClaimSmartObjectTaskInstanceData;
+
+protected:
+	virtual bool Link(FStateTreeLinker& Linker) override;
+	virtual const UStruct* GetInstanceDataType() const override { return FInstanceDataType::StaticStruct(); }
+	virtual EStateTreeRunStatus EnterState(FStateTreeExecutionContext& Context, const EStateTreeStateChangeType ChangeType, const FStateTreeTransitionResult& Transition) const override;
+	virtual void ExitState(FStateTreeExecutionContext& Context, const EStateTreeStateChangeType ChangeType, const FStateTreeTransitionResult& Transition) const override;
+	virtual EStateTreeRunStatus Tick(FStateTreeExecutionContext& Context, const float DeltaTime) const override;
+
+	TStateTreeExternalDataHandle<FMassSmartObjectUserFragment> SmartObjectUserHandle;
+	TStateTreeExternalDataHandle<USmartObjectSubsystem> SmartObjectSubsystemHandle;
+	TStateTreeExternalDataHandle<UMassSignalSubsystem> MassSignalSubsystemHandle;
+
+	TStateTreeInstanceDataPropertyHandle<FMassSmartObjectCandidateSlots> CandidateSlotsHandle;
+	TStateTreeInstanceDataPropertyHandle<FSmartObjectClaimHandle> ClaimedSlotHandle;
+
+	/** Delay in seconds before trying to use another smart object */
+	UPROPERTY(EditAnywhere, Category = Parameter)
+	float InteractionCooldown = 0.f;
+};
diff --git a/Engine/Plugins/AI/MassAI/Source/MassAIBehavior/Public/Tasks/MassFindSmartObjectTask.h b/Engine/Plugins/AI/MassAI/Source/MassAIBehavior/Public/Tasks/MassFindSmartObjectTask.h
new file mode 100644
--- /dev/null
+++ b/Engine/Plugins/AI/MassAI/Source/MassAIBehavior/Public/Tasks/MassFindSmartObjectTask.h
@@ -0,0 +1,72 @@
+﻿// Copyright Epic Games, Inc. All Rights Reserved.
+
+#pragma once
+#include "MassStateTreeTypes.h"
+#include "SmartObjectSubsystem.h"
+#include "MassSignalSubsystem.h"
+#include "MassSmartObjectRequest.h"
+#include "MassFindSmartObjectTask.generated.h"
+
+struct FMassSmartObjectUserFragment;
+struct FMassZoneGraphLaneLocationFragment;
+struct FTransformFragment;
+
+USTRUCT()
+struct MASSAIBEHAVIOR_API FMassFindSmartObjectTaskInstanceData
+{
+	GENERATED_BODY()
+
+	/** Result of the candidates search request */
+	UPROPERTY(VisibleAnywhere, Category = Output)
+	FMassSmartObjectCandidateSlots FoundCandidateSlots;		// @todo: Should turn this in a StateTree result/value.
+
+	UPROPERTY(VisibleAnywhere, Category = Output)
+	bool bHasCandidateSlots = false;
+
+	/** The identifier of the search request send by the task to find candidates */
+	UPROPERTY()
+	FMassSmartObjectRequestID SearchRequestID;
+
+	/** Next update time; task will not do anything when Tick gets called before that time */
+	UPROPERTY()
+	float NextUpdate = 0.f;
+
+	/** Last lane where the smart objects were searched. */
+	UPROPERTY()
+	FZoneGraphLaneHandle LastLane;
+};
+
+USTRUCT(meta = (DisplayName = "Find Smart Object"))
+struct MASSAIBEHAVIOR_API FMassFindSmartObjectTask : public FMassStateTreeTaskBase
+{
+	GENERATED_BODY()
+
+	using FInstanceDataType = FMassFindSmartObjectTaskInstanceData;
+
+protected:
+	virtual bool Link(FStateTreeLinker& Linker) override;
+	virtual const UStruct* GetInstanceDataType() const override { return FInstanceDataType::StaticStruct(); }
+	virtual void ExitState(FStateTreeExecutionContext& Context, const EStateTreeStateChangeType ChangeType, const FStateTreeTransitionResult& Transition) const override;
+	virtual void StateCompleted(FStateTreeExecutionContext& Context, const EStateTreeRunStatus CompletionStatus, const FStateTreeActiveStates& CompletedActiveStates) const override;
+	virtual EStateTreeRunStatus Tick(FStateTreeExecutionContext& Context, const float DeltaTime) const override;
+
+	TStateTreeExternalDataHandle<USmartObjectSubsystem> SmartObjectSubsystemHandle;
+	TStateTreeExternalDataHandle<UMassSignalSubsystem> MassSignalSubsystemHandle;
+	TStateTreeExternalDataHandle<FTransformFragment> EntityTransformHandle;
+	TStateTreeExternalDataHandle<FMassSmartObjectUserFragment> SmartObjectUserHandle;
+	TStateTreeExternalDataHandle<FMassZoneGraphLaneLocationFragment, EStateTreeExternalDataRequirement::Optional> LocationHandle;
+
+	TStateTreeInstanceDataPropertyHandle<FMassSmartObjectCandidateSlots> FoundCandidateSlotsHandle;
+	TStateTreeInstanceDataPropertyHandle<bool> HasCandidateSlotsHandle;
+	TStateTreeInstanceDataPropertyHandle<FMassSmartObjectRequestID> SearchRequestIDHandle;
+	TStateTreeInstanceDataPropertyHandle<float> NextUpdateHandle;
+	TStateTreeInstanceDataPropertyHandle<FZoneGraphLaneHandle> LastLaneHandle;
+
+	/** Gameplay tag query for finding matching smart objects. */
+	UPROPERTY(EditAnywhere, Category = Parameter)
+	FGameplayTagQuery ActivityRequirements;
+
+	/** How frequently to search for new candidates. */
+	UPROPERTY(EditAnywhere, Category = Parameter)
+	float SearchInterval = 10.0f;
+};
\ No newline at end of file
diff --git a/Engine/Plugins/AI/MassAI/Source/MassAIBehavior/Public/Tasks/MassStateTreeSmartObjectTask.h b/Engine/Plugins/AI/MassAI/Source/MassAIBehavior/Public/Tasks/MassStateTreeSmartObjectTask.h
deleted file mode 100644
diff --git a/Engine/Plugins/AI/MassAI/Source/MassAIBehavior/Public/Tasks/MassUseSmartObjectTask.h b/Engine/Plugins/AI/MassAI/Source/MassAIBehavior/Public/Tasks/MassUseSmartObjectTask.h
new file mode 100644
--- /dev/null
+++ b/Engine/Plugins/AI/MassAI/Source/MassAIBehavior/Public/Tasks/MassUseSmartObjectTask.h
@@ -0,0 +1,51 @@
+// Copyright Epic Games, Inc. All Rights Reserved.
+
+#pragma once
+
+#include "MassSmartObjectRequest.h"
+#include "MassStateTreeTypes.h"
+#include "MassUseSmartObjectTask.generated.h"
+
+struct FStateTreeExecutionContext;
+struct FMassSmartObjectUserFragment;
+class USmartObjectSubsystem;
+class UMassSignalSubsystem;
+struct FTransformFragment;
+struct FMassMoveTargetFragment;
+
+/**
+ * Task to tell an entity to start using a claimed smart object.
+ */
+USTRUCT()
+struct MASSAIBEHAVIOR_API FMassUseSmartObjectTaskInstanceData
+{
+	GENERATED_BODY()
+
+	UPROPERTY(VisibleAnywhere, Category = Input)
+	FSmartObjectClaimHandle ClaimedSlot;
+};
+
+USTRUCT(meta = (DisplayName = "Mass Use SmartObject Task"))
+struct MASSAIBEHAVIOR_API FMassUseSmartObjectTask : public FMassStateTreeTaskBase
+{
+	GENERATED_BODY()
+
+	using FInstanceDataType = FMassUseSmartObjectTaskInstanceData;
+
+protected:
+	virtual bool Link(FStateTreeLinker& Linker) override;
+	virtual const UStruct* GetInstanceDataType() const override { return FInstanceDataType::StaticStruct(); }
+	virtual EStateTreeRunStatus EnterState(FStateTreeExecutionContext& Context, const EStateTreeStateChangeType ChangeType, const FStateTreeTransitionResult& Transition) const override;
+	virtual void ExitState(FStateTreeExecutionContext& Context, const EStateTreeStateChangeType ChangeType, const FStateTreeTransitionResult& Transition) const override;
+	virtual void StateCompleted(FStateTreeExecutionContext& Context, const EStateTreeRunStatus CompletionStatus, const FStateTreeActiveStates& CompletedActiveStates) const override;
+	virtual EStateTreeRunStatus Tick(FStateTreeExecutionContext& Context, const float DeltaTime) const override;
+
+	TStateTreeExternalDataHandle<USmartObjectSubsystem> SmartObjectSubsystemHandle;
+	TStateTreeExternalDataHandle<UMassSignalSubsystem> MassSignalSubsystemHandle;
+	TStateTreeExternalDataHandle<FTransformFragment> EntityTransformHandle;
+	TStateTreeExternalDataHandle<FMassSmartObjectUserFragment> SmartObjectUserHandle;
+	TStateTreeExternalDataHandle<FMassMoveTargetFragment> MoveTargetHandle;
+
+	TStateTreeInstanceDataPropertyHandle<FSmartObjectClaimHandle> ClaimedSlotHandle;
+};
+
diff --git a/Engine/Plugins/AI/MassAI/Source/MassAIBehavior/Public/Tasks/MassZoneGraphFindSmartObjectTarget.h b/Engine/Plugins/AI/MassAI/Source/MassAIBehavior/Public/Tasks/MassZoneGraphFindSmartObjectTarget.h
--- a/Engine/Plugins/AI/MassAI/Source/MassAIBehavior/Public/Tasks/MassZoneGraphFindSmartObjectTarget.h
+++ b/Engine/Plugins/AI/MassAI/Source/MassAIBehavior/Public/Tasks/MassZoneGraphFindSmartObjectTarget.h
@@ -2,12 +2,14 @@
 
 #pragma once
 #include "MassStateTreeTypes.h"
+#include "SmartObjectRuntime.h"
 #include "MassZoneGraphPathFollowTask.h"
 #include "MassZoneGraphFindSmartObjectTarget.generated.h"
 
 struct FMassSmartObjectUserFragment;
 struct FMassZoneGraphLaneLocationFragment;
 class UZoneGraphAnnotationSubsystem;
+class USmartObjectSubsystem;
 
 /**
 * Computes move target to a smart object based on current location on ZoneGraph.
@@ -17,6 +19,9 @@ struct MASSAIBEHAVIOR_API FMassZoneGraphFindSmartObjectTargetInstanceData
 {
 	GENERATED_BODY()
 
+	UPROPERTY(VisibleAnywhere, Category = Input)
+	FSmartObjectClaimHandle ClaimedSlot;
+
 	UPROPERTY(EditAnywhere, Category = Output)
 	FMassZoneGraphTargetLocation SmartObjectLocation;
 };
@@ -26,14 +31,17 @@ struct MASSAIBEHAVIOR_API FMassZoneGraphFindSmartObjectTarget : public FMassStat
 {
 	GENERATED_BODY()
 
+	using FInstanceDataType = FMassZoneGraphFindSmartObjectTargetInstanceData;
+
 protected:
 	virtual bool Link(FStateTreeLinker& Linker) override;
-	virtual const UStruct* GetInstanceDataType() const override { return FMassZoneGraphFindSmartObjectTargetInstanceData::StaticStruct(); }
+	virtual const UStruct* GetInstanceDataType() const override { return FInstanceDataType::StaticStruct(); }
 	virtual EStateTreeRunStatus EnterState(FStateTreeExecutionContext& Context, const EStateTreeStateChangeType ChangeType, const FStateTreeTransitionResult& Transition) const override;
 
-	TStateTreeExternalDataHandle<FMassSmartObjectUserFragment> SmartObjectUserHandle;
 	TStateTreeExternalDataHandle<FMassZoneGraphLaneLocationFragment> LocationHandle;
 	TStateTreeExternalDataHandle<UZoneGraphAnnotationSubsystem> AnnotationSubsystemHandle;
+	TStateTreeExternalDataHandle<USmartObjectSubsystem> SmartObjectSubsystemHandle;
 
+	TStateTreeInstanceDataPropertyHandle<FSmartObjectClaimHandle> ClaimedSlotHandle;
 	TStateTreeInstanceDataPropertyHandle<FMassZoneGraphTargetLocation> SmartObjectLocationHandle;
 };
\ No newline at end of file
diff --git a/Engine/Plugins/AI/MassAI/Source/MassAIDebug/MassAIDebug.Build.cs b/Engine/Plugins/AI/MassAI/Source/MassAIDebug/MassAIDebug.Build.cs
--- a/Engine/Plugins/AI/MassAI/Source/MassAIDebug/MassAIDebug.Build.cs
+++ b/Engine/Plugins/AI/MassAI/Source/MassAIDebug/MassAIDebug.Build.cs
@@ -45,6 +45,9 @@ namespace UnrealBuildTool.Rules
 					"MassSignals",
 					"MassLOD",
 					"StructUtils",
+					"StructUtils",
+					"MassSmartObjects",
+					"SmartObjectsModule",
 				}
 			);
 
diff --git a/Engine/Plugins/AI/MassAI/Source/MassAIDebug/Private/GameplayDebuggerCategory_Mass.cpp b/Engine/Plugins/AI/MassAI/Source/MassAIDebug/Private/GameplayDebuggerCategory_Mass.cpp
--- a/Engine/Plugins/AI/MassAI/Source/MassAIDebug/Private/GameplayDebuggerCategory_Mass.cpp
+++ b/Engine/Plugins/AI/MassAI/Source/MassAIDebug/Private/GameplayDebuggerCategory_Mass.cpp
@@ -19,6 +19,8 @@
 #include "MassStateTreeFragments.h"
 #include "MassStateTreeExecutionContext.h"
 #include "MassZoneGraphNavigationFragments.h"
+#include "MassSmartObjectFragments.h"
+#include "SmartObjectSubsystem.h"
 #include "Util/ColorConstants.h"
 #include "MassSimulationLOD.h"
 #include "CanvasItem.h"
@@ -40,7 +42,7 @@ namespace UE::Mass::Debug
 		}
 		return EntityHandle;
 	};
-	
+
 	FMassEntityHandle GetBestEntity(const FVector ViewLocation, const FVector ViewDirection, const TConstArrayView<FMassEntityHandle> Entities, const TConstArrayView<FVector> Locations, const bool bLimitAngle)
 	{
 		// Reusing similar algorithm as UGameplayDebuggerLocalController for now 
@@ -362,17 +364,19 @@ void FGameplayDebuggerCategory_Mass::CollectData(APlayerController* OwnerPC, AAc
 		EntityQuery.AddRequirement<FMassLookAtFragment>(EMassFragmentAccess::ReadOnly, EMassFragmentPresence::Optional);
 		EntityQuery.AddRequirement<FMassSimulationLODFragment>(EMassFragmentAccess::ReadOnly, EMassFragmentPresence::Optional);
 		EntityQuery.AddRequirement<FMassZoneGraphShortPathFragment>(EMassFragmentAccess::ReadOnly);
+		EntityQuery.AddRequirement<FMassSmartObjectUserFragment>(EMassFragmentAccess::ReadOnly, EMassFragmentPresence::Optional);
 
 		const float CurrentTime = World->GetTimeSeconds();
 		
 		UMassStateTreeSubsystem* MassStateTreeSubsystem = World->GetSubsystem<UMassStateTreeSubsystem>();
 		UMassSignalSubsystem* SignalSubsystem = World->GetSubsystem<UMassSignalSubsystem>();
+		USmartObjectSubsystem* SmartObjectSubsystem = World->GetSubsystem<USmartObjectSubsystem>();
 		
-		if (MassStateTreeSubsystem && EntitySystem && SignalSubsystem)
+		if (MassStateTreeSubsystem && EntitySystem && SignalSubsystem && SmartObjectSubsystem)
 		{
 			FMassExecutionContext Context(0.0f);
 		
-			EntityQuery.ForEachEntityChunk(*EntitySystem, Context, [this, MassStateTreeSubsystem, SignalSubsystem, EntitySystem, OwnerPC, ViewLocation, ViewDirection, CurrentTime](FMassExecutionContext& Context)
+			EntityQuery.ForEachEntityChunk(*EntitySystem, Context, [this, MassStateTreeSubsystem, SignalSubsystem, EntitySystem, SmartObjectSubsystem, OwnerPC, ViewLocation, ViewDirection, CurrentTime](FMassExecutionContext& Context)
 			{
 				const int32 NumEntities = Context.GetNumEntities();
 				const TConstArrayView<FMassStateTreeInstanceFragment> StateTreeInstanceList = Context.GetFragmentView<FMassStateTreeInstanceFragment>();
@@ -387,10 +391,12 @@ void FGameplayDebuggerCategory_Mass::CollectData(APlayerController* OwnerPC, AAc
 				const TConstArrayView<FMassLookAtFragment> LookAtList = Context.GetFragmentView<FMassLookAtFragment>();
 				const TConstArrayView<FMassSimulationLODFragment> SimLODList = Context.GetFragmentView<FMassSimulationLODFragment>();
 				const TConstArrayView<FMassZoneGraphShortPathFragment> ShortPathList = Context.GetFragmentView<FMassZoneGraphShortPathFragment>();
+				const TConstArrayView<FMassSmartObjectUserFragment> SOUserList = Context.GetFragmentView<FMassSmartObjectUserFragment>();
 				const FMassStateTreeSharedFragment& SharedStateTree = Context.GetConstSharedFragment<FMassStateTreeSharedFragment>();
 
 				const bool bHasLOD = (SimLODList.Num() > 0);
 				const bool bHasLookAt = (LookAtList.Num() > 0);
+				const bool bHasSOUser = (SOUserList.Num() > 0);
 
 				const UGameplayDebuggerUserSettings* Settings = GetDefault<UGameplayDebuggerUserSettings>();
 				const float MaxViewDistance = Settings->MaxViewDistance;
@@ -486,6 +492,19 @@ void FGameplayDebuggerCategory_Mass::CollectData(APlayerController* OwnerPC, AAc
 						}
 					}
 
+					// SmartObject
+					if (bHasSOUser)
+					{
+						const FMassSmartObjectUserFragment& SOUser = SOUserList[EntityIndex];
+						if (SOUser.InteractionHandle.IsValid())
+						{
+							const FVector ZOffset = FVector(0.0f , 0.0f , 25.0f );
+							const FTransform SlotTransform = SmartObjectSubsystem->GetSlotTransform(SOUser.InteractionHandle).Get(FTransform::Identity);
+							const FVector SlotLocation = SlotTransform.GetLocation();
+							AddShape(FGameplayDebuggerShape::MakeSegment(EntityLocation + ZOffset, SlotLocation + ZOffset, 3.0f, FColorList::Orange));
+						}
+					}
+
 					// Path
 					if (bShowNearEntityPath)
 					{
diff --git a/Engine/Plugins/Runtime/MassGameplay/Source/MassSmartObjects/Private/MassSmartObjectHandler.cpp b/Engine/Plugins/Runtime/MassGameplay/Source/MassSmartObjects/Private/MassSmartObjectHandler.cpp
--- a/Engine/Plugins/Runtime/MassGameplay/Source/MassSmartObjects/Private/MassSmartObjectHandler.cpp
+++ b/Engine/Plugins/Runtime/MassGameplay/Source/MassSmartObjects/Private/MassSmartObjectHandler.cpp
@@ -13,22 +13,23 @@
 
 namespace UE::Mass::SmartObject
 {
-	void Release(FMassCommandBuffer& CommandBuffer, const FMassBehaviorEntityContext& Context, const EMassSmartObjectInteractionStatus NewStatus)
+	void StopInteraction(FMassCommandBuffer& CommandBuffer, const FMassBehaviorEntityContext& Context, const FSmartObjectClaimHandle& ClaimHandle, const EMassSmartObjectInteractionStatus NewStatus)
 	{
 		FMassSmartObjectUserFragment& User = Context.EntityView.GetFragmentData<FMassSmartObjectUserFragment>();
-
-		if (User.InteractionStatus == EMassSmartObjectInteractionStatus::InProgress ||
-			User.InteractionStatus == EMassSmartObjectInteractionStatus::BehaviorCompleted)
+		if ( User.InteractionHandle == ClaimHandle)
 		{
-			if (const USmartObjectMassBehaviorDefinition* BehaviorDefinition = Context.SmartObjectSubsystem.GetBehaviorDefinition<USmartObjectMassBehaviorDefinition>(User.ClaimHandle))
+			if (User.InteractionStatus == EMassSmartObjectInteractionStatus::InProgress ||
+				User.InteractionStatus == EMassSmartObjectInteractionStatus::BehaviorCompleted)
 			{
-				BehaviorDefinition->Deactivate(CommandBuffer, Context);
+				if (const USmartObjectMassBehaviorDefinition* BehaviorDefinition = Context.SmartObjectSubsystem.GetBehaviorDefinition<USmartObjectMassBehaviorDefinition>(User.InteractionHandle))
+				{
+					BehaviorDefinition->Deactivate(CommandBuffer, Context);
+				}
 			}
-		}
-		Context.SmartObjectSubsystem.Release(User.ClaimHandle);
 
-		User.InteractionStatus = NewStatus;
-		User.ClaimHandle.Invalidate();
+			User.InteractionStatus = NewStatus;
+			User.InteractionHandle.Invalidate();
+		}
 	}
 
 	struct FPayload
@@ -49,13 +50,13 @@ namespace UE::Mass::SmartObject
 			const FMassEntityView EntityView(*EntitySubsystem, Payload.Entity);
 			const FMassBehaviorEntityContext Context(EntityView, *SmartObjectSubsystem);
 
-			Release(EntitySubsystem->Defer(), Context, EMassSmartObjectInteractionStatus::Aborted);
+			StopInteraction(EntitySubsystem->Defer(), Context, ClaimHandle, EMassSmartObjectInteractionStatus::Aborted);
 			if (EntitySubsystem->IsProcessing() == false)
 			{
 				EntitySubsystem->FlushCommands();
 			}
 
-			SignalSubsystem->SignalEntity(UE::Mass::Signals::SmartObjectInteractionDone, Payload.Entity);
+			SignalSubsystem->SignalEntity(UE::Mass::Signals::SmartObjectInteractionAborted, Payload.Entity);
 		}
 	}
 } // UE::Mass::SmartObject;
@@ -63,66 +64,57 @@ namespace UE::Mass::SmartObject
 //----------------------------------------------------------------------//
 // FMassSmartObjectHandler
 //----------------------------------------------------------------------//
-FMassSmartObjectRequestID FMassSmartObjectHandler::FindCandidatesAsync(const FMassEntityHandle RequestingEntity, const FVector& Location) const
+FMassSmartObjectRequestID FMassSmartObjectHandler::FindCandidatesAsync(const FMassEntityHandle RequestingEntity, const FGameplayTagContainer& UserTags, const FGameplayTagQuery& ActivityRequirements, const FVector& Location) const
 {
 	const FMassEntityHandle ReservedEntity = EntitySubsystem.ReserveEntity();
 
 	FMassSmartObjectWorldLocationRequestFragment RequestFragment;
 	RequestFragment.SearchOrigin = Location;
 	RequestFragment.RequestingEntity = RequestingEntity;
+	RequestFragment.UserTags = UserTags;
+	RequestFragment.ActivityRequirements = ActivityRequirements;
 
 	FMassSmartObjectRequestResultFragment ResultFragment;
 
-	/*ExecutionContext.Defer().PushCommand(
-		FBuildEntityFromFragmentInstances(ReservedEntity,
-			{
-				FStructView::Make(RequestFragment),
-				FStructView::Make(ResultFragment)
-			}));*/
-
 	ExecutionContext.Defer().PushCommand<FMassCommandAddFragmentInstances>(ReservedEntity, RequestFragment, ResultFragment);
 
 	return FMassSmartObjectRequestID(ReservedEntity);
 }
 
-FMassSmartObjectRequestID FMassSmartObjectHandler::FindCandidatesAsync(const FMassEntityHandle RequestingEntity, const FZoneGraphCompactLaneLocation& LaneLocation) const
+FMassSmartObjectRequestID FMassSmartObjectHandler::FindCandidatesAsync(const FMassEntityHandle RequestingEntity, const FGameplayTagContainer& UserTags, const FGameplayTagQuery& ActivityRequirements, const FZoneGraphCompactLaneLocation& LaneLocation) const
 {
 	const FMassEntityHandle ReservedEntity = EntitySubsystem.ReserveEntity();
 
 	FMassSmartObjectLaneLocationRequestFragment RequestFragment;
 	RequestFragment.CompactLaneLocation = LaneLocation;
 	RequestFragment.RequestingEntity = RequestingEntity;
+	RequestFragment.UserTags = UserTags;
+	RequestFragment.ActivityRequirements = ActivityRequirements;
 
 	FMassSmartObjectRequestResultFragment ResultFragment;
 
-	/*ExecutionContext.Defer().PushCommand(
-		FBuildEntityFromFragmentInstances(ReservedEntity,
-			{
-				FStructView::Make(RequestFragment),
-				FStructView::Make(ResultFragment)
-			}));*/
-
 	ExecutionContext.Defer().PushCommand<FMassCommandAddFragmentInstances>(ReservedEntity, RequestFragment, ResultFragment);
 
 	return FMassSmartObjectRequestID(ReservedEntity);
 }
 
-FMassSmartObjectRequestResult FMassSmartObjectHandler::GetRequestResult(const FMassSmartObjectRequestID& RequestID) const
+const FMassSmartObjectCandidateSlots* FMassSmartObjectHandler::GetRequestCandidates(const FMassSmartObjectRequestID& RequestID) const
 {
 	const FMassEntityHandle RequestEntity = static_cast<FMassEntityHandle>(RequestID);
 	if (!ensureMsgf(EntitySubsystem.IsEntityValid(RequestEntity), TEXT("Invalid request.")))
 	{
-		return {};
+		return nullptr;
 	}
 
 	// Check if entity is built by now.
 	if (!EntitySubsystem.IsEntityBuilt(RequestEntity))
 	{
-		return {};
+		return nullptr;
 	}
 
 	const FMassSmartObjectRequestResultFragment& RequestFragment = EntitySubsystem.GetFragmentDataChecked<FMassSmartObjectRequestResultFragment>(RequestEntity);
-	return RequestFragment.Result;
+
+	return RequestFragment.bProcessed ? &RequestFragment.Candidates : nullptr;
 }
 
 void FMassSmartObjectHandler::RemoveRequest(const FMassSmartObjectRequestID& RequestID) const
@@ -131,36 +123,18 @@ void FMassSmartObjectHandler::RemoveRequest(const FMassSmartObjectRequestID& Req
 	ExecutionContext.Defer().DestroyEntity(RequestEntity);
 }
 
-EMassSmartObjectClaimResult FMassSmartObjectHandler::ClaimCandidate(const FMassEntityHandle Entity, FMassSmartObjectUserFragment& User, const FMassSmartObjectRequestID& RequestID) const
-{
-	if (!ensureMsgf(RequestID.IsSet(), TEXT("Trying to claim using an invalid request.")))
-	{
-		return EMassSmartObjectClaimResult::Failed_InvalidRequest;
-	}
-
-	const FMassEntityHandle RequestEntity = static_cast<FMassEntityHandle>(RequestID);
-	const FMassSmartObjectRequestResultFragment& RequestFragment = EntitySubsystem.GetFragmentDataChecked<FMassSmartObjectRequestResultFragment>(RequestEntity);
-	if (!ensureMsgf(RequestFragment.Result.bProcessed, TEXT("Trying to claim using a request that is not processed.")))
-	{
-		return EMassSmartObjectClaimResult::Failed_UnprocessedRequest;
-	}
-
-	return ClaimCandidate(Entity, User, RequestFragment.Result);
-}
-
-EMassSmartObjectClaimResult FMassSmartObjectHandler::ClaimCandidate(const FMassEntityHandle Entity, FMassSmartObjectUserFragment& User, const FMassSmartObjectRequestResult& SearchRequestResult) const
+FSmartObjectClaimHandle FMassSmartObjectHandler::ClaimCandidate(const FMassEntityHandle Entity, FMassSmartObjectUserFragment& User, const FMassSmartObjectCandidateSlots& Candidates) const
 {
-	if (!ensureMsgf(SearchRequestResult.bProcessed, TEXT("Trying to claim using a search request that is not processed.")))
-	{
-		return EMassSmartObjectClaimResult::Failed_UnprocessedRequest;
-	}
+	checkf(!User.InteractionHandle.IsValid(), TEXT("User should not already have an interaction."));
+	
+	const TConstArrayView<FSmartObjectCandidateSlot> View = MakeArrayView(Candidates.Slots.GetData(), Candidates.NumSlots);
 
-	checkf(!User.ClaimHandle.IsValid(), TEXT("User should not already have a valid handle."));
-	const TConstArrayView<FSmartObjectCandidate> View = MakeArrayView(SearchRequestResult.Candidates.GetData(), SearchRequestResult.NumCandidates);
-
-	for (const FSmartObjectCandidate& Candidate : View)
+	FSmartObjectClaimHandle ClaimedSlot;
+	
+	for (const FSmartObjectCandidateSlot& CandidateSlot : View)
 	{
-		if (ClaimSmartObject(Entity, User, Candidate.Result))
+		ClaimedSlot = ClaimSmartObject(Entity, User, CandidateSlot.Result);  
+		if (ClaimedSlot.IsValid())
 		{
 #if WITH_MASSGAMEPLAY_DEBUG
 			UE_CVLOG(UE::Mass::Debug::IsDebuggingEntity(Entity),
@@ -169,36 +143,18 @@ EMassSmartObjectClaimResult FMassSmartObjectHandler::ClaimCandidate(const FMassE
 				Log,
 				TEXT("[%s] claimed [%s]"),
 				*Entity.DebugGetDescription(),
-				*LexToString(Candidate.Result));
+				*LexToString(CandidateSlot.Result));
 #endif // WITH_MASSGAMEPLAY_DEBUG
 			break;
 		}
 	}
 
-	return User.ClaimHandle.IsValid() ? EMassSmartObjectClaimResult::Succeeded : EMassSmartObjectClaimResult::Failed_NoAvailableCandidate;
+	return ClaimedSlot;
 }
 
-bool FMassSmartObjectHandler::ClaimSmartObject(const FMassEntityHandle Entity, FMassSmartObjectUserFragment& User, const FSmartObjectRequestResult& RequestResult) const
+FSmartObjectClaimHandle FMassSmartObjectHandler::ClaimSmartObject(const FMassEntityHandle Entity, FMassSmartObjectUserFragment& User, const FSmartObjectRequestResult& RequestResult) const
 {
 	const FSmartObjectClaimHandle ClaimHandle = SmartObjectSubsystem.Claim(RequestResult);
-	const bool bSuccess = ClaimHandle.IsValid();
-	if (bSuccess)
-	{
-		User.ClaimHandle = ClaimHandle;
-		User.InteractionStatus = EMassSmartObjectInteractionStatus::Unset;
-		const FTransform Transform =SmartObjectSubsystem.GetSlotTransform(User.ClaimHandle).Get(FTransform::Identity);
-		User.TargetLocation = Transform.GetLocation();
-		User.TargetDirection = Transform.GetRotation().Vector();
-
-		// Register callback to abort interaction if slot gets invalidated.
-		// Callback will be unregistered by UMassSmartObjectUserFragmentDeinitializer
-		UE::Mass::SmartObject::FPayload Payload;
-		Payload.Entity = Entity;
-		Payload.EntitySubsystem = &EntitySubsystem;
-		Payload.SmartObjectSubsystem = &SmartObjectSubsystem;
-		Payload.SignalSubsystem = &SignalSubsystem;
-		SmartObjectSubsystem.RegisterSlotInvalidationCallback(ClaimHandle, FOnSlotInvalidated::CreateStatic(&UE::Mass::SmartObject::OnSlotInvalidated, Payload));
-	}
 
 #if WITH_MASSGAMEPLAY_DEBUG
 	UE_CVLOG(UE::Mass::Debug::IsDebuggingEntity(Entity),
@@ -208,16 +164,25 @@ bool FMassSmartObjectHandler::ClaimSmartObject(const FMassEntityHandle Entity, F
 		TEXT("[%s] claim for [%s] %s"),
 		*Entity.DebugGetDescription(),
 		*LexToString(RequestResult),
-		bSuccess ? TEXT("Succeeded") : TEXT("Failed"));
+		ClaimHandle.IsValid() ? TEXT("Succeeded") : TEXT("Failed"));
 #endif // WITH_MASSGAMEPLAY_DEBUG
 
-	return bSuccess;
+	// Register callback to abort interaction if slot gets invalidated.
+	// Callback will be unregistered by UMassSmartObjectUserFragmentDeinitializer
+	UE::Mass::SmartObject::FPayload Payload;
+	Payload.Entity = Entity;
+	Payload.EntitySubsystem = &EntitySubsystem;
+	Payload.SmartObjectSubsystem = &SmartObjectSubsystem;
+	Payload.SignalSubsystem = &SignalSubsystem;
+	SmartObjectSubsystem.RegisterSlotInvalidationCallback(ClaimHandle, FOnSlotInvalidated::CreateStatic(&UE::Mass::SmartObject::OnSlotInvalidated, Payload));
+	
+	return ClaimHandle;
 }
 
-bool FMassSmartObjectHandler::UseSmartObject(
+bool FMassSmartObjectHandler::StartUsingSmartObject(
 	const FMassEntityHandle Entity,
 	FMassSmartObjectUserFragment& User,
-	const FTransformFragment& Transform) const
+	const FSmartObjectClaimHandle ClaimHandle) const
 {
 #if WITH_MASSGAMEPLAY_DEBUG
 	UE_CVLOG(UE::Mass::Debug::IsDebuggingEntity(Entity),
@@ -226,25 +191,28 @@ bool FMassSmartObjectHandler::UseSmartObject(
 		Log,
 		TEXT("[%s] starts using [%s]"),
 		*Entity.DebugGetDescription(),
-		*LexToString(User.ClaimHandle));
+		*LexToString(User.InteractionHandle));
 #endif // WITH_MASSGAMEPLAY_DEBUG
 
-	const USmartObjectMassBehaviorDefinition* BehaviorDefinition = SmartObjectSubsystem.Use<USmartObjectMassBehaviorDefinition>(User.ClaimHandle);
+	const USmartObjectMassBehaviorDefinition* BehaviorDefinition = SmartObjectSubsystem.Use<USmartObjectMassBehaviorDefinition>(ClaimHandle);
 	if (BehaviorDefinition == nullptr)
 	{
 		return false;
 	}
 
+	// Activate behavior
 	const FMassEntityView EntityView(EntitySubsystem, Entity);
 	const FMassBehaviorEntityContext Context(EntityView, SmartObjectSubsystem);
 	BehaviorDefinition->Activate(ExecutionContext.Defer(), Context);
 
 	User.InteractionStatus = EMassSmartObjectInteractionStatus::InProgress;
+	User.InteractionHandle = ClaimHandle;
 
 	return true;
 }
 
-void FMassSmartObjectHandler::ReleaseSmartObject(const FMassEntityHandle Entity, FMassSmartObjectUserFragment& User, const EMassSmartObjectInteractionStatus NewStatus) const
+
+void FMassSmartObjectHandler::StopUsingSmartObject(const FMassEntityHandle Entity, FMassSmartObjectUserFragment& User, const EMassSmartObjectInteractionStatus NewStatus) const
 {
 #if WITH_MASSGAMEPLAY_DEBUG
 	UE_CVLOG(UE::Mass::Debug::IsDebuggingEntity(Entity),
@@ -253,7 +221,7 @@ void FMassSmartObjectHandler::ReleaseSmartObject(const FMassEntityHandle Entity,
 		Log,
 		TEXT("[%s] releases handle [%s]"),
 		*Entity.DebugGetDescription(),
-		*LexToString(User.ClaimHandle));
+		*LexToString(User.InteractionHandle));
 #endif // WITH_MASSGAMEPLAY_DEBUG
 
 #if DO_ENSURE
@@ -295,5 +263,33 @@ void FMassSmartObjectHandler::ReleaseSmartObject(const FMassEntityHandle Entity,
 
 	const FMassEntityView EntityView(EntitySubsystem, Entity);
 	const FMassBehaviorEntityContext Context(EntityView, SmartObjectSubsystem);
-	UE::Mass::SmartObject::Release(ExecutionContext.Defer(), Context, NewStatus);
+	UE::Mass::SmartObject::StopInteraction(ExecutionContext.Defer(), Context, User.InteractionHandle, NewStatus);
+}
+
+
+void FMassSmartObjectHandler::ReleaseSmartObject(const FMassEntityHandle Entity, FMassSmartObjectUserFragment& User, const FSmartObjectClaimHandle ClaimHandle) const
+{
+#if WITH_MASSGAMEPLAY_DEBUG
+	UE_CVLOG(UE::Mass::Debug::IsDebuggingEntity(Entity),
+		&SmartObjectSubsystem,
+		LogSmartObject,
+		Log,
+		TEXT("[%s] releases handle [%s]"),
+		*Entity.DebugGetDescription(),
+		*LexToString(User.InteractionHandle));
+#endif // WITH_MASSGAMEPLAY_DEBUG
+
+#if DO_ENSURE
+	if (User.InteractionStatus == EMassSmartObjectInteractionStatus::InProgress)
+	{
+		ensureMsgf(false, TEXT("Expecting interaction to be compled before releasing the SmartObject. Current state %s"), *UEnum::GetValueAsString(User.InteractionStatus));
+	}
+#endif
+
+	const FMassEntityView EntityView(EntitySubsystem, Entity);
+	const FMassBehaviorEntityContext Context(EntityView, SmartObjectSubsystem);
+
+	SmartObjectSubsystem.UnregisterSlotInvalidationCallback(ClaimHandle);
+
+	Context.SmartObjectSubsystem.Release(ClaimHandle);
 }
diff --git a/Engine/Plugins/Runtime/MassGameplay/Source/MassSmartObjects/Private/MassSmartObjectProcessor.cpp b/Engine/Plugins/Runtime/MassGameplay/Source/MassSmartObjects/Private/MassSmartObjectProcessor.cpp
--- a/Engine/Plugins/Runtime/MassGameplay/Source/MassSmartObjects/Private/MassSmartObjectProcessor.cpp
+++ b/Engine/Plugins/Runtime/MassGameplay/Source/MassSmartObjects/Private/MassSmartObjectProcessor.cpp
@@ -61,23 +61,23 @@ void UMassSmartObjectCandidatesFinderProcessor::Execute(UMassEntitySubsystem& En
 	// Build list of request owner entities to send a completion signal
 	TArray<FMassEntityHandle> EntitiesToSignal;
 
-	auto BeginRequestProcessing = [](const FMassEntityHandle Entity, FMassExecutionContext& Context, FMassSmartObjectRequestResult& Result)
+	auto BeginRequestProcessing = [](const FMassEntityHandle Entity, FMassExecutionContext& Context, FMassSmartObjectRequestResultFragment& Result)
 	{
 		Context.Defer().AddTag<FMassSmartObjectCompletedRequestTag>(Entity);
-		Result.NumCandidates = 0;
+		Result.Candidates.NumSlots = 0;
 	};
 
-	auto EndRequestProcessing = [](const UObject* LogOwner, const FMassEntityHandle Entity, FMassSmartObjectRequestResult& Result)
+	auto EndRequestProcessing = [](const UObject* LogOwner, const FMassEntityHandle Entity, FMassSmartObjectRequestResultFragment& Result)
 	{
-		if (Result.NumCandidates > 0)
+		if (Result.Candidates.NumSlots > 0)
 		{
-			TArrayView<FSmartObjectCandidate> View = MakeArrayView(Result.Candidates.GetData(), Result.NumCandidates);
-			Algo::Sort(View, [](const FSmartObjectCandidate& LHS, const FSmartObjectCandidate& RHS) { return LHS.Cost < RHS.Cost; });
+			TArrayView<FSmartObjectCandidateSlot> View = MakeArrayView(Result.Candidates.Slots.GetData(), Result.Candidates.NumSlots);
+			Algo::Sort(View, [](const FSmartObjectCandidateSlot& LHS, const FSmartObjectCandidateSlot& RHS) { return LHS.Cost < RHS.Cost; });
 		}
 		Result.bProcessed = true;
 
 #if WITH_MASSGAMEPLAY_DEBUG
-		UE_VLOG(LogOwner, LogSmartObject, Verbose, TEXT("[%s] search completed: found %d"), *Entity.DebugGetDescription(), Result.NumCandidates);
+		UE_VLOG(LogOwner, LogSmartObject, Verbose, TEXT("[%s] search completed: found %d"), *Entity.DebugGetDescription(), Result.Candidates.NumSlots);
 #endif // WITH_MASSGAMEPLAY_DEBUG
 	};
 
@@ -93,12 +93,13 @@ void UMassSmartObjectCandidatesFinderProcessor::Execute(UMassEntitySubsystem& En
 		const TArrayView<FMassSmartObjectRequestResultFragment> ResultList = Context.GetMutableFragmentView<FMassSmartObjectRequestResultFragment>();
 
 		TArray<FSmartObjectRequestResult> QueryResults;
-		TArray<FSmartObjectCandidate> SortedResults;
+		TArray<FSmartObjectCandidateSlot> SortedCandidateSlots;
 
 		for (int32 i = 0; i < NumEntities; ++i)
 		{
 			const FMassSmartObjectWorldLocationRequestFragment& RequestFragment = RequestList[i];
-			FMassSmartObjectRequestResult& Result = ResultList[i].Result;
+			FMassSmartObjectRequestResultFragment& Result = ResultList[i];
+			
 			EntitiesToSignal.Add(RequestFragment.RequestingEntity);
 
 			const FVector& SearchOrigin = RequestFragment.SearchOrigin;
@@ -115,14 +116,17 @@ void UMassSmartObjectCandidatesFinderProcessor::Execute(UMassEntitySubsystem& En
 			BeginRequestProcessing(Entity, Context, Result);
 			ON_SCOPE_EXIT{ EndRequestProcessing(&SmartObjectSubsystem, Entity, Result);	};
 
+			Filter.UserTags = RequestFragment.UserTags;
+			Filter.ActivityRequirements = RequestFragment.ActivityRequirements;
+			
 			QueryResults.Reset();
 			SmartObjectSubsystem.FindSmartObjects(FSmartObjectRequest(SearchBounds, Filter), QueryResults);
 
-			SortedResults.Reset(QueryResults.Num());
+			SortedCandidateSlots.Reset(QueryResults.Num());
 			for (const FSmartObjectRequestResult& QueryResult : QueryResults)
 			{
 				const FVector SlotLocation = SmartObjectSubsystem.GetSlotLocation(QueryResult.SlotHandle).GetValue();
-				SortedResults.Emplace(QueryResult, FVector::DistSquared(SearchOrigin, SlotLocation));
+				SortedCandidateSlots.Emplace(QueryResult, FVector::DistSquared(SearchOrigin, SlotLocation));
 
 #if WITH_MASSGAMEPLAY_DEBUG
 				if (bDisplayDebug)
@@ -133,12 +137,12 @@ void UMassSmartObjectCandidatesFinderProcessor::Execute(UMassEntitySubsystem& En
 				}
 #endif // WITH_MASSGAMEPLAY_DEBUG
 			}
-			SortedResults.Sort([](const FSmartObjectCandidate& First, const FSmartObjectCandidate& Second){ return First.Cost < Second.Cost; });
+			SortedCandidateSlots.Sort([](const FSmartObjectCandidateSlot& First, const FSmartObjectCandidateSlot& Second){ return First.Cost < Second.Cost; });
 
-			Result.NumCandidates = FMath::Min<uint8>(FMassSmartObjectRequestResult::MaxNumCandidates, SortedResults.Num());
-			for (int ResultIndex = 0; ResultIndex < Result.NumCandidates; ResultIndex++)
+			Result.Candidates.NumSlots = FMath::Min<uint8>(FMassSmartObjectCandidateSlots::MaxNumCandidates, SortedCandidateSlots.Num());
+			for (int ResultIndex = 0; ResultIndex < Result.Candidates.NumSlots; ResultIndex++)
 			{
-				Result.Candidates[ResultIndex] = SortedResults[ResultIndex];
+				Result.Candidates.Slots[ResultIndex] = SortedCandidateSlots[ResultIndex];
 			}
 		}
 	});
@@ -168,7 +172,8 @@ void UMassSmartObjectCandidatesFinderProcessor::Execute(UMassEntitySubsystem& En
 			for (int32 i = 0; i < NumEntities; ++i)
 			{
 				const FMassSmartObjectLaneLocationRequestFragment& RequestFragment = RequestList[i];
-				FMassSmartObjectRequestResult& Result = ResultList[i].Result;
+				FMassSmartObjectRequestResultFragment& Result = ResultList[i];
+				
 				EntitiesToSignal.Add(RequestFragment.RequestingEntity);
 
 				const FZoneGraphCompactLaneLocation RequestLocation = RequestFragment.CompactLaneLocation;
@@ -220,6 +225,9 @@ void UMassSmartObjectCandidatesFinderProcessor::Execute(UMassEntitySubsystem& En
 					continue;
 				}
 
+				Filter.UserTags = RequestFragment.UserTags;
+				Filter.ActivityRequirements = RequestFragment.ActivityRequirements;
+
 				for (const int32 Index : SmartObjectList->SmartObjectLaneLocationIndices)
 				{
 					// Find entry point using FindChecked since all smart objects added to LaneToSmartObjects lookup table
@@ -251,7 +259,7 @@ void UMassSmartObjectCandidatesFinderProcessor::Execute(UMassEntitySubsystem& En
 
 					for (FSmartObjectSlotHandle SlotHandle : SlotHandles)
 					{
-						Result.Candidates[Result.NumCandidates++] = FSmartObjectCandidate(FSmartObjectRequestResult(Handle, SlotHandle), Cost);
+						Result.Candidates.Slots[Result.Candidates.NumSlots++] = FSmartObjectCandidateSlot(FSmartObjectRequestResult(Handle, SlotHandle), Cost);
 
 #if WITH_MASSGAMEPLAY_DEBUG
 						if (bDisplayDebug)
@@ -267,13 +275,13 @@ void UMassSmartObjectCandidatesFinderProcessor::Execute(UMassEntitySubsystem& En
 							UE_VLOG_SEGMENT(&SmartObjectSubsystem, LogSmartObject, Display, RequestLaneLocation.Position, EntryPointLaneLocation.Position, DebugColor, TEXT(""));
 						}
 #endif // WITH_MASSGAMEPLAY_DEBUG
-						if (Result.NumCandidates == FMassSmartObjectRequestResult::MaxNumCandidates)
+						if (Result.Candidates.NumSlots == FMassSmartObjectCandidateSlots::MaxNumCandidates)
 						{
 							break;
 						}
 					}
 
-					if (Result.NumCandidates == FMassSmartObjectRequestResult::MaxNumCandidates)
+					if (Result.Candidates.NumSlots == FMassSmartObjectCandidateSlots::MaxNumCandidates)
 					{
 						break;
 					}
@@ -341,10 +349,10 @@ void UMassSmartObjectTimedBehaviorProcessor::Execute(UMassEntitySubsystem& Entit
 			const bool bIsDebuggingEntity = UE::Mass::Debug::IsDebuggingEntity(Entity, &DebugColor);
 			if (bIsDebuggingEntity)
 			{
-				UE_CVLOG(bMustRelease, &SmartObjectSubsystem, LogSmartObject, Log, TEXT("[%s] stops using [%s]"), *Entity.DebugGetDescription(), *LexToString(SOUser.ClaimHandle));
-				UE_CVLOG(!bMustRelease, &SmartObjectSubsystem, LogSmartObject, Verbose, TEXT("[%s] using [%s]. Time left: %.1f"), *Entity.DebugGetDescription(), *LexToString(SOUser.ClaimHandle), UseTime);
+				UE_CVLOG(bMustRelease, &SmartObjectSubsystem, LogSmartObject, Log, TEXT("[%s] stops using [%s]"), *Entity.DebugGetDescription(), *LexToString(SOUser.InteractionHandle));
+				UE_CVLOG(!bMustRelease, &SmartObjectSubsystem, LogSmartObject, Verbose, TEXT("[%s] using [%s]. Time left: %.1f"), *Entity.DebugGetDescription(), *LexToString(SOUser.InteractionHandle), UseTime);
 
-				const TOptional<FTransform> Transform = SmartObjectSubsystem.GetSlotTransform(SOUser.ClaimHandle);
+				const TOptional<FTransform> Transform = SmartObjectSubsystem.GetSlotTransform(SOUser.InteractionHandle);
 				if (Transform.IsSet())
 				{
 					constexpr float Radius = 40.f;
@@ -397,7 +405,7 @@ void UMassSmartObjectUserFragmentDeinitializer::Execute(UMassEntitySubsystem& En
 
 			for (int32 i = 0; i < NumEntities; ++i)
 			{
-				SmartObjectSubsystem.UnregisterSlotInvalidationCallback(SmartObjectUserFragments[i].ClaimHandle);
+				SmartObjectSubsystem.UnregisterSlotInvalidationCallback(SmartObjectUserFragments[i].InteractionHandle);
 			}
 		});
 }
\ No newline at end of file
diff --git a/Engine/Plugins/Runtime/MassGameplay/Source/MassSmartObjects/Private/MassSmartObjectUserTrait.cpp b/Engine/Plugins/Runtime/MassGameplay/Source/MassSmartObjects/Private/MassSmartObjectUserTrait.cpp
--- a/Engine/Plugins/Runtime/MassGameplay/Source/MassSmartObjects/Private/MassSmartObjectUserTrait.cpp
+++ b/Engine/Plugins/Runtime/MassGameplay/Source/MassSmartObjects/Private/MassSmartObjectUserTrait.cpp
@@ -7,5 +7,8 @@
 
 void UMassSmartObjectUserTrait::BuildTemplate(FMassEntityTemplateBuildContext& BuildContext, const UWorld& World) const
 {
-	BuildContext.AddFragment<FMassSmartObjectUserFragment>();
+	FMassSmartObjectUserFragment SmartObjectUser;
+	SmartObjectUser.UserTags = UserTags;
+	
+	BuildContext.AddFragment(FConstStructView::Make(SmartObjectUser));
 }
diff --git a/Engine/Plugins/Runtime/MassGameplay/Source/MassSmartObjects/Public/MassSmartObjectFragments.h b/Engine/Plugins/Runtime/MassGameplay/Source/MassSmartObjects/Public/MassSmartObjectFragments.h
--- a/Engine/Plugins/Runtime/MassGameplay/Source/MassSmartObjects/Public/MassSmartObjectFragments.h
+++ b/Engine/Plugins/Runtime/MassGameplay/Source/MassSmartObjects/Public/MassSmartObjectFragments.h
@@ -4,6 +4,7 @@
 
 #include "MassEntityTypes.h"
 #include "MassSmartObjectTypes.h"
+#include "MassSmartObjectRequest.h"
 #include "SmartObjectRuntime.h"
 #include "MassSmartObjectFragments.generated.h"
 
@@ -13,15 +14,15 @@ struct MASSSMARTOBJECTS_API FMassSmartObjectUserFragment : public FMassFragment
 {
 	GENERATED_BODY()
 
+	/** Tags describing the smart object user. */
 	UPROPERTY(Transient)
-	FSmartObjectClaimHandle ClaimHandle;
+	FGameplayTagContainer UserTags;
 
+	/** Claim handle for the currently active smart object interaction. */
 	UPROPERTY(Transient)
-	FVector TargetLocation = FVector::ZeroVector;
-
-	UPROPERTY(Transient)
-	FVector TargetDirection = FVector::ForwardVector;
+	FSmartObjectClaimHandle InteractionHandle;
 
+	/** Status of the current active smart object interaction. */
 	UPROPERTY(Transient)
 	EMassSmartObjectInteractionStatus InteractionStatus = EMassSmartObjectInteractionStatus::Unset;
 
@@ -30,7 +31,7 @@ struct MASSSMARTOBJECTS_API FMassSmartObjectUserFragment : public FMassFragment
 	 * won't look for new interactions (value of 0 indicates no cooldown).
 	 */
 	UPROPERTY(Transient)
-	float CooldownEndTime = 0.f;
+	float InteractionCooldownEndTime = 0.f;
 };
 
 /** Fragment used to process time based smartobject interactions */
diff --git a/Engine/Plugins/Runtime/MassGameplay/Source/MassSmartObjects/Public/MassSmartObjectHandler.h b/Engine/Plugins/Runtime/MassGameplay/Source/MassSmartObjects/Public/MassSmartObjectHandler.h
--- a/Engine/Plugins/Runtime/MassGameplay/Source/MassSmartObjects/Public/MassSmartObjectHandler.h
+++ b/Engine/Plugins/Runtime/MassGameplay/Source/MassSmartObjects/Public/MassSmartObjectHandler.h
@@ -46,7 +46,7 @@ struct MASSSMARTOBJECTS_API FMassSmartObjectHandler
 	 * @param Location The center of the query
 	 * @return Request identifier that can be used to try claiming a result once available
 	 */
-	UE_NODISCARD FMassSmartObjectRequestID FindCandidatesAsync(const FMassEntityHandle RequestingEntity, const FVector& Location) const;
+	UE_NODISCARD FMassSmartObjectRequestID FindCandidatesAsync(const FMassEntityHandle RequestingEntity, const FGameplayTagContainer& UserTags, const FGameplayTagQuery& ActivityRequirements, const FVector& Location) const;
 
 	/**
 	 * Creates an async request to build a list of compatible smart objects
@@ -56,15 +56,15 @@ struct MASSSMARTOBJECTS_API FMassSmartObjectHandler
 	 * @param LaneLocation The lane location as reference for the query
 	 * @return Request identifier that can be used to try claiming a result once available
 	 */
-	UE_NODISCARD FMassSmartObjectRequestID FindCandidatesAsync(const FMassEntityHandle RequestingEntity, const FZoneGraphCompactLaneLocation& LaneLocation) const;
+	UE_NODISCARD FMassSmartObjectRequestID FindCandidatesAsync(const FMassEntityHandle RequestingEntity, const FGameplayTagContainer& UserTags, const FGameplayTagQuery& ActivityRequirements, const FZoneGraphCompactLaneLocation& LaneLocation) const;
 
 	/**
 	 * Provides the result of a previously created request from FindCandidatesAsync to indicate if it has been processed
 	 * and the results can be used by ClaimCandidate.
 	 * @param RequestID A valid request identifier (method will ensure otherwise)
-	 * @return The current request's result
+	 * @return The current request's result, nullptr if request not ready yet.
 	 */
-	UE_NODISCARD FMassSmartObjectRequestResult GetRequestResult(const FMassSmartObjectRequestID& RequestID) const;
+	UE_NODISCARD const FMassSmartObjectCandidateSlots* GetRequestCandidates(const FMassSmartObjectRequestID& RequestID) const;
 
 	/**
 	 * Deletes the request associated to the specified identifier
@@ -73,25 +73,13 @@ struct MASSSMARTOBJECTS_API FMassSmartObjectHandler
 	void RemoveRequest(const FMassSmartObjectRequestID& RequestID) const;
 
 	/**
-	 * Claims the first available smart object from the results of the find request. This method
-	 * can be called without calling IsRequestProcessed first but it will fail and return
-	 * Failed_UnprocessedRequest. It can then be called again on next frame until it succeeds or
-	 * returns a different error.
+	 * Claims the first available smart object from the provided candidates.
 	 * @param Entity MassEntity associated to the user fragment
 	 * @param User Fragment of the user claiming
-	 * @param RequestID A valid request identifier (method will ensure otherwise)
-	 * @return Whether the slot has been successfully claimed or not
-	 */
-	UE_NODISCARD EMassSmartObjectClaimResult ClaimCandidate(const FMassEntityHandle Entity, FMassSmartObjectUserFragment& User, const FMassSmartObjectRequestID& RequestID) const;
-
-	/**
-	 * Claims the first available smart object from the provided results.
-	 * @param Entity MassEntity associated to the user fragment
-	 * @param User Fragment of the user claiming
-	 * @param SearchRequestResult Results of completed search request
+	 * @param Candidates Candidate slots to choose from.
 	 * @return Whether the slot has been successfully claimed or not
 	 */
-	UE_NODISCARD EMassSmartObjectClaimResult ClaimCandidate(const FMassEntityHandle Entity, FMassSmartObjectUserFragment& User, const FMassSmartObjectRequestResult& SearchRequestResult) const;
+	UE_NODISCARD FSmartObjectClaimHandle ClaimCandidate(const FMassEntityHandle Entity, FMassSmartObjectUserFragment& User, const FMassSmartObjectCandidateSlots& Candidates) const;
 
 	/**
 	 * Claims the first available slot holding any type of USmartObjectMassBehaviorDefinition in the smart object
@@ -101,24 +89,33 @@ struct MASSSMARTOBJECTS_API FMassSmartObjectHandler
 	 * @param RequestResult A valid smart object request result (method will ensure otherwise)
 	 * @return Whether the slot has been successfully claimed or not
 	 */
-	bool ClaimSmartObject(const FMassEntityHandle Entity, FMassSmartObjectUserFragment& User, const FSmartObjectRequestResult& RequestResult) const;
+	UE_NODISCARD FSmartObjectClaimHandle ClaimSmartObject(const FMassEntityHandle Entity, FMassSmartObjectUserFragment& User, const FSmartObjectRequestResult& RequestResult) const;
 
 	/**
 	 * Activates the mass gameplay behavior associated to the previously claimed smart object.
 	 * @param Entity MassEntity associated to the user fragment
 	 * @param User Fragment of the user claiming
+	 * @param ClaimHandle claimed smart object slot to use.
 	 * @param Transform Fragment holding the transform of the user claiming
 	 * @return Whether the slot has been successfully claimed or not
 	 */
-	bool UseSmartObject(const FMassEntityHandle Entity, FMassSmartObjectUserFragment& User, const FTransformFragment& Transform) const;
+	bool StartUsingSmartObject(const FMassEntityHandle Entity, FMassSmartObjectUserFragment& User, const FSmartObjectClaimHandle ClaimHandle) const;
+
+	/**
+	 * Deactivates the mass gameplay behavior started using StartUsingSmartObject.
+	 * @param Entity MassEntity associated to the user fragment
+	 * @param User Fragment of the user claiming
+	 * @param NewStatus Reason of the deactivation.
+	 */
+	void StopUsingSmartObject(const FMassEntityHandle Entity, FMassSmartObjectUserFragment& User, const EMassSmartObjectInteractionStatus NewStatus) const;
 
 	/**
 	 * Releases a claimed/in-use smart object and update user fragment.
 	 * @param Entity MassEntity associated to the user fragment
 	 * @param User Fragment of the user claiming
-	 * @param Status The new status for in-progress interaction
+	 * @param ClaimHandle claimed smart object slot to release.
 	 */
-	void ReleaseSmartObject(const FMassEntityHandle Entity, FMassSmartObjectUserFragment& User, const EMassSmartObjectInteractionStatus Status = EMassSmartObjectInteractionStatus::Unset) const;
+	void ReleaseSmartObject(const FMassEntityHandle Entity, FMassSmartObjectUserFragment& User, const FSmartObjectClaimHandle ClaimHandle) const;
 
 private:
 	UMassEntitySubsystem& EntitySubsystem;
diff --git a/Engine/Plugins/Runtime/MassGameplay/Source/MassSmartObjects/Public/MassSmartObjectRequest.h b/Engine/Plugins/Runtime/MassGameplay/Source/MassSmartObjects/Public/MassSmartObjectRequest.h
--- a/Engine/Plugins/Runtime/MassGameplay/Source/MassSmartObjects/Public/MassSmartObjectRequest.h
+++ b/Engine/Plugins/Runtime/MassGameplay/Source/MassSmartObjects/Public/MassSmartObjectRequest.h
@@ -9,15 +9,15 @@
 #include "MassSmartObjectRequest.generated.h"
 
 /**
- * Structure that represents a potential smart object for a MassEntity during the search
+ * Structure that represents a potential smart object slot for a MassEntity during the search
  */
 USTRUCT()
-struct MASSSMARTOBJECTS_API FSmartObjectCandidate
+struct MASSSMARTOBJECTS_API FSmartObjectCandidateSlot
 {
 	GENERATED_BODY()
 
-	FSmartObjectCandidate() = default;
-	FSmartObjectCandidate(const FSmartObjectRequestResult InResult, const float InCost)	: Result(InResult), Cost(InCost) {}
+	FSmartObjectCandidateSlot() = default;
+	FSmartObjectCandidateSlot(const FSmartObjectRequestResult InResult, const float InCost)	: Result(InResult), Cost(InCost) {}
 
 	UPROPERTY(Transient)
 	FSmartObjectRequestResult Result;
@@ -51,19 +51,21 @@ private:
 /**
  * Struct that holds status and results of a candidate finder request
  */
-USTRUCT()
-struct MASSSMARTOBJECTS_API FMassSmartObjectRequestResult
+USTRUCT(BlueprintType)
+struct MASSSMARTOBJECTS_API FMassSmartObjectCandidateSlots
 {
 	GENERATED_BODY()
 
+	void Reset()
+	{
+		NumSlots = 0;
+	}
+	
 	static constexpr uint32 MaxNumCandidates = 4;
-	TStaticArray<FSmartObjectCandidate, MaxNumCandidates> Candidates;
-
-	UPROPERTY(Transient)
-	uint8 NumCandidates = 0;
+	TStaticArray<FSmartObjectCandidateSlot, MaxNumCandidates> Slots;
 
 	UPROPERTY(Transient)
-	bool bProcessed = false;
+	uint8 NumSlots = 0;
 };
 
 /**
@@ -75,7 +77,10 @@ struct MASSSMARTOBJECTS_API FMassSmartObjectRequestResultFragment : public FMass
 	GENERATED_BODY()
 
 	UPROPERTY(Transient)
-	FMassSmartObjectRequestResult Result;
+	FMassSmartObjectCandidateSlots Candidates;
+
+	UPROPERTY(Transient)
+	bool bProcessed = false;
 };
 
 /**
@@ -93,6 +98,12 @@ struct MASSSMARTOBJECTS_API FMassSmartObjectWorldLocationRequestFragment : publi
 
 	UPROPERTY(Transient)
 	FMassEntityHandle RequestingEntity;
+
+	UPROPERTY(Transient)
+	FGameplayTagContainer UserTags;
+
+	UPROPERTY(Transient)
+	FGameplayTagQuery ActivityRequirements;
 };
 
 /**
@@ -109,6 +120,12 @@ struct MASSSMARTOBJECTS_API FMassSmartObjectLaneLocationRequestFragment : public
 
 	UPROPERTY(Transient)
 	FMassEntityHandle RequestingEntity;
+
+	UPROPERTY(Transient)
+	FGameplayTagContainer UserTags;
+
+	UPROPERTY(Transient)
+	FGameplayTagQuery ActivityRequirements;
 };
 
 /**
diff --git a/Engine/Plugins/Runtime/MassGameplay/Source/MassSmartObjects/Public/MassSmartObjectTypes.h b/Engine/Plugins/Runtime/MassGameplay/Source/MassSmartObjects/Public/MassSmartObjectTypes.h
--- a/Engine/Plugins/Runtime/MassGameplay/Source/MassSmartObjects/Public/MassSmartObjectTypes.h
+++ b/Engine/Plugins/Runtime/MassGameplay/Source/MassSmartObjects/Public/MassSmartObjectTypes.h
@@ -10,18 +10,9 @@ namespace UE::Mass::Signals
 	const FName SmartObjectRequestCandidates = FName(TEXT("SmartObjectRequestCandidates"));
 	const FName SmartObjectCandidatesReady = FName(TEXT("SmartObjectCandidatesReady"));
 	const FName SmartObjectInteractionDone = FName(TEXT("SmartObjectInteractionDone"));
+	const FName SmartObjectInteractionAborted = FName(TEXT("SmartObjectInteractionAborted"));
 }
 
-UENUM()
-enum class EMassSmartObjectClaimResult: uint8
-{
-	Unset,
-	Failed_InvalidRequest,
-	Failed_UnprocessedRequest,
-	Failed_NoAvailableCandidate,
-	Succeeded
-};
-
 UENUM()
 enum class EMassSmartObjectInteractionStatus: uint8
 {
diff --git a/Engine/Plugins/Runtime/MassGameplay/Source/MassSmartObjects/Public/MassSmartObjectUserTrait.h b/Engine/Plugins/Runtime/MassGameplay/Source/MassSmartObjects/Public/MassSmartObjectUserTrait.h
--- a/Engine/Plugins/Runtime/MassGameplay/Source/MassSmartObjects/Public/MassSmartObjectUserTrait.h
+++ b/Engine/Plugins/Runtime/MassGameplay/Source/MassSmartObjects/Public/MassSmartObjectUserTrait.h
@@ -3,6 +3,7 @@
 #pragma once
 
 #include "MassEntityTraitBase.h"
+#include "GameplayTagContainer.h"
 #include "MassSmartObjectUserTrait.generated.h"
 
 /**
@@ -15,4 +16,8 @@ class MASSSMARTOBJECTS_API UMassSmartObjectUserTrait : public UMassEntityTraitBa
 
 protected:
 	virtual void BuildTemplate(FMassEntityTemplateBuildContext& BuildContext, const UWorld& World) const override;
+
+	/** Tags describing the SmartObject user. Used when searching smart objects. */
+	UPROPERTY(EditAnywhere, Category = Parameter)
+	FGameplayTagContainer UserTags;
 };
