Fixed MassSmartObjectCandidatesFinderProcessor not registering its queries

diff --git a/Engine/Plugins/Runtime/MassGameplay/Source/MassSmartObjects/Private/MassSmartObjectProcessor.cpp b/Engine/Plugins/Runtime/MassGameplay/Source/MassSmartObjects/Private/MassSmartObjectProcessor.cpp
--- a/Engine/Plugins/Runtime/MassGameplay/Source/MassSmartObjects/Private/MassSmartObjectProcessor.cpp
+++ b/Engine/Plugins/Runtime/MassGameplay/Source/MassSmartObjects/Private/MassSmartObjectProcessor.cpp
@@ -32,12 +32,14 @@ void UMassSmartObjectCandidatesFinderProcessor::ConfigureQueries()
 	WorldRequestQuery.AddRequirement<FMassSmartObjectRequestResultFragment>(EMassFragmentAccess::ReadWrite);
 	WorldRequestQuery.AddTagRequirement<FMassSmartObjectCompletedRequestTag>(EMassFragmentPresence::None);
 	WorldRequestQuery.AddSubsystemRequirement<USmartObjectSubsystem>(EMassFragmentAccess::ReadOnly);
+	WorldRequestQuery.RegisterWithProcessor(*this);
 
 	LaneRequestQuery.AddRequirement<FMassSmartObjectLaneLocationRequestFragment>(EMassFragmentAccess::ReadOnly);
 	LaneRequestQuery.AddRequirement<FMassSmartObjectRequestResultFragment>(EMassFragmentAccess::ReadWrite);
 	LaneRequestQuery.AddTagRequirement<FMassSmartObjectCompletedRequestTag>(EMassFragmentPresence::None);
 	LaneRequestQuery.AddSubsystemRequirement<UZoneGraphSubsystem>(EMassFragmentAccess::ReadOnly);
 	LaneRequestQuery.AddSubsystemRequirement<USmartObjectSubsystem>(EMassFragmentAccess::ReadOnly);
+	LaneRequestQuery.RegisterWithProcessor(*this);
 }
 
 UMassSmartObjectCandidatesFinderProcessor::UMassSmartObjectCandidatesFinderProcessor()
