Marked up some queries and processors as requiring game-thread execution, also updated traits of some of the mass-used subsystems

diff --git a/Engine/Plugins/AI/MassAI/Source/MassAIBehavior/Private/MassStateTreeProcessors.cpp b/Engine/Plugins/AI/MassAI/Source/MassAIBehavior/Private/MassStateTreeProcessors.cpp
--- a/Engine/Plugins/AI/MassAI/Source/MassAIBehavior/Private/MassStateTreeProcessors.cpp
+++ b/Engine/Plugins/AI/MassAI/Source/MassAIBehavior/Private/MassStateTreeProcessors.cpp
@@ -178,6 +178,7 @@ UMassStateTreeFragmentDestructor::UMassStateTreeFragmentDestructor()
 	ExecutionFlags = (int32)(EProcessorExecutionFlags::Standalone | EProcessorExecutionFlags::Server);
 	ObservedType = FMassStateTreeInstanceFragment::StaticStruct();
 	Operation = EMassObservedOperation::Remove;
+	bRequiresGameThreadExecution = true;
 }
 
 void UMassStateTreeFragmentDestructor::Initialize(UObject& Owner)
@@ -238,6 +239,7 @@ UMassStateTreeActivationProcessor::UMassStateTreeActivationProcessor()
 {
 	ExecutionOrder.ExecuteAfter.Add(UE::Mass::ProcessorGroupNames::LOD);
 	ExecutionOrder.ExecuteBefore.Add(UE::Mass::ProcessorGroupNames::Behavior);
+	bRequiresGameThreadExecution = true; // due to UMassStateTreeSubsystem RW access
 }
 
 void UMassStateTreeActivationProcessor::Initialize(UObject& Owner)
diff --git a/Engine/Plugins/AI/MassAI/Source/MassAIBehavior/Public/MassStateTreeSubsystem.h b/Engine/Plugins/AI/MassAI/Source/MassAIBehavior/Public/MassStateTreeSubsystem.h
--- a/Engine/Plugins/AI/MassAI/Source/MassAIBehavior/Public/MassStateTreeSubsystem.h
+++ b/Engine/Plugins/AI/MassAI/Source/MassAIBehavior/Public/MassStateTreeSubsystem.h
@@ -5,6 +5,7 @@
 #include "MassStateTreeTypes.h"
 #include "Subsystems/WorldSubsystem.h"
 #include "StateTreeExecutionContext.h"
+#include "MassExternalSubsystemTraits.h"
 #include "MassStateTreeSubsystem.generated.h"
 
 class UStateTree;
@@ -67,3 +68,12 @@ protected:
 	UPROPERTY(Transient)
 	TArray<FMassStateTreeInstanceDataItem> InstanceDataArray;
 };
+
+template<>
+struct TMassExternalSubsystemTraits<UMassStateTreeSubsystem> final
+{
+	enum
+	{
+		GameThreadOnly = true
+	};
+};
diff --git a/Engine/Plugins/AI/MassCrowd/Source/MassCrowd/Private/MassCrowdServerRepresentationProcessor.cpp b/Engine/Plugins/AI/MassCrowd/Source/MassCrowd/Private/MassCrowdServerRepresentationProcessor.cpp
--- a/Engine/Plugins/AI/MassCrowd/Source/MassCrowd/Private/MassCrowdServerRepresentationProcessor.cpp
+++ b/Engine/Plugins/AI/MassCrowd/Source/MassCrowd/Private/MassCrowdServerRepresentationProcessor.cpp
@@ -9,6 +9,7 @@ UMassCrowdServerRepresentationProcessor::UMassCrowdServerRepresentationProcessor
 	ExecutionFlags = (int32)EProcessorExecutionFlags::Server;
 
 	bAutoRegisterWithProcessingPhases = true;
+	bRequiresGameThreadExecution = true;
 
 	ExecutionOrder.ExecuteAfter.Add(UE::Mass::ProcessorGroupNames::LOD);
 }
diff --git a/Engine/Plugins/AI/MassCrowd/Source/MassCrowd/Private/MassCrowdVisualizationProcessor.cpp b/Engine/Plugins/AI/MassCrowd/Source/MassCrowd/Private/MassCrowdVisualizationProcessor.cpp
--- a/Engine/Plugins/AI/MassCrowd/Source/MassCrowd/Private/MassCrowdVisualizationProcessor.cpp
+++ b/Engine/Plugins/AI/MassCrowd/Source/MassCrowd/Private/MassCrowdVisualizationProcessor.cpp
@@ -62,6 +62,7 @@ void UMassDebugCrowdVisualizationProcessor::ConfigureQueries()
 	EntityQuery.AddRequirement<FTransformFragment>(EMassFragmentAccess::ReadOnly);
 	EntityQuery.AddRequirement<FMassRepresentationFragment>(EMassFragmentAccess::ReadWrite);
 	EntityQuery.AddRequirement<FMassActorFragment>(EMassFragmentAccess::ReadWrite);
+	EntityQuery.RequireMutatingWorldAccess(); // due to UWorld mutable access
 }
 
 void UMassDebugCrowdVisualizationProcessor::Initialize(UObject& Owner)
diff --git a/Engine/Plugins/Runtime/MassGameplay/Source/MassActors/Private/Translators/MassCapsuleComponentTranslators.cpp b/Engine/Plugins/Runtime/MassGameplay/Source/MassActors/Private/Translators/MassCapsuleComponentTranslators.cpp
--- a/Engine/Plugins/Runtime/MassGameplay/Source/MassActors/Private/Translators/MassCapsuleComponentTranslators.cpp
+++ b/Engine/Plugins/Runtime/MassGameplay/Source/MassActors/Private/Translators/MassCapsuleComponentTranslators.cpp
@@ -50,6 +50,7 @@ UMassTransformToActorCapsuleTranslator::UMassTransformToActorCapsuleTranslator()
 	RequiredTags.Add<FMassCapsuleTransformCopyToActorTag>();
 	ExecutionOrder.ExecuteInGroup = UE::Mass::ProcessorGroupNames::UpdateWorldFromMass;
 	ExecutionOrder.ExecuteAfter.Add(UE::Mass::ProcessorGroupNames::Movement);
+	bRequiresGameThreadExecution = true;
 }
 
 void UMassTransformToActorCapsuleTranslator::ConfigureQueries()
@@ -57,6 +58,7 @@ void UMassTransformToActorCapsuleTranslator::ConfigureQueries()
 	AddRequiredTagsToQuery(EntityQuery);
 	EntityQuery.AddRequirement<FCapsuleComponentWrapperFragment>(EMassFragmentAccess::ReadWrite);
 	EntityQuery.AddRequirement<FTransformFragment>(EMassFragmentAccess::ReadOnly);
+	EntityQuery.RequireMutatingWorldAccess(); // due to mutating World by setting component transform
 }
 
 void UMassTransformToActorCapsuleTranslator::Execute(UMassEntitySubsystem& EntitySubsystem, FMassExecutionContext& Context)
diff --git a/Engine/Plugins/Runtime/MassGameplay/Source/MassActors/Private/Translators/MassCharacterMovementTranslators.cpp b/Engine/Plugins/Runtime/MassGameplay/Source/MassActors/Private/Translators/MassCharacterMovementTranslators.cpp
--- a/Engine/Plugins/Runtime/MassGameplay/Source/MassActors/Private/Translators/MassCharacterMovementTranslators.cpp
+++ b/Engine/Plugins/Runtime/MassGameplay/Source/MassActors/Private/Translators/MassCharacterMovementTranslators.cpp
@@ -146,6 +146,7 @@ void UMassCharacterOrientationToActorTranslator::ConfigureQueries()
 	AddRequiredTagsToQuery(EntityQuery);
 	EntityQuery.AddRequirement<FCharacterMovementComponentWrapperFragment>(EMassFragmentAccess::ReadWrite);
 	EntityQuery.AddRequirement<FTransformFragment>(EMassFragmentAccess::ReadOnly);
+	EntityQuery.RequireMutatingWorldAccess(); // due to mutating World by setting component rotation
 }
 
 void UMassCharacterOrientationToActorTranslator::Execute(UMassEntitySubsystem& EntitySubsystem, FMassExecutionContext& Context)
diff --git a/Engine/Plugins/Runtime/MassGameplay/Source/MassActors/Private/Translators/MassSceneComponentLocationTranslator.cpp b/Engine/Plugins/Runtime/MassGameplay/Source/MassActors/Private/Translators/MassSceneComponentLocationTranslator.cpp
--- a/Engine/Plugins/Runtime/MassGameplay/Source/MassActors/Private/Translators/MassSceneComponentLocationTranslator.cpp
+++ b/Engine/Plugins/Runtime/MassGameplay/Source/MassActors/Private/Translators/MassSceneComponentLocationTranslator.cpp
@@ -52,6 +52,7 @@ UMassSceneComponentLocationToActorTranslator::UMassSceneComponentLocationToActor
 	ExecutionOrder.ExecuteInGroup = UE::Mass::ProcessorGroupNames::UpdateWorldFromMass;
 	ExecutionOrder.ExecuteAfter.Add(UE::Mass::ProcessorGroupNames::Movement);
 	RequiredTags.Add<FMassSceneComponentLocationCopyToActorTag>();
+	bRequiresGameThreadExecution = true;
 }
 
 void UMassSceneComponentLocationToActorTranslator::ConfigureQueries()
@@ -59,6 +60,7 @@ void UMassSceneComponentLocationToActorTranslator::ConfigureQueries()
 	AddRequiredTagsToQuery(EntityQuery);
 	EntityQuery.AddRequirement<FMassSceneComponentWrapperFragment>(EMassFragmentAccess::ReadOnly);
 	EntityQuery.AddRequirement<FTransformFragment>(EMassFragmentAccess::ReadWrite);
+	EntityQuery.RequireMutatingWorldAccess(); // due to mutating World by setting actor's location
 }
 
 void UMassSceneComponentLocationToActorTranslator::Execute(UMassEntitySubsystem& EntitySubsystem, FMassExecutionContext& Context)
diff --git a/Engine/Plugins/Runtime/MassGameplay/Source/MassGameplayDebug/Private/AssignDebugVisProcessor.cpp b/Engine/Plugins/Runtime/MassGameplay/Source/MassGameplayDebug/Private/AssignDebugVisProcessor.cpp
--- a/Engine/Plugins/Runtime/MassGameplay/Source/MassGameplayDebug/Private/AssignDebugVisProcessor.cpp
+++ b/Engine/Plugins/Runtime/MassGameplay/Source/MassGameplayDebug/Private/AssignDebugVisProcessor.cpp
@@ -13,6 +13,7 @@ UAssignDebugVisProcessor::UAssignDebugVisProcessor()
 	: EntityQuery(*this)
 {
 	bAutoRegisterWithProcessingPhases = false;
+	bRequiresGameThreadExecution = true; // due to UMassDebuggerSubsystem
 	ObservedType = FSimDebugVisFragment::StaticStruct();
 	Operation = EMassObservedOperation::Add;
 }
diff --git a/Engine/Plugins/Runtime/MassGameplay/Source/MassGameplayDebug/Private/DebugVisLocationProcessor.cpp b/Engine/Plugins/Runtime/MassGameplay/Source/MassGameplayDebug/Private/DebugVisLocationProcessor.cpp
--- a/Engine/Plugins/Runtime/MassGameplay/Source/MassGameplayDebug/Private/DebugVisLocationProcessor.cpp
+++ b/Engine/Plugins/Runtime/MassGameplay/Source/MassGameplayDebug/Private/DebugVisLocationProcessor.cpp
@@ -15,6 +15,7 @@ UDebugVisLocationProcessor::UDebugVisLocationProcessor()
 	: EntityQuery(*this)
 {
 	ExecutionOrder.ExecuteAfter.Add(UE::Mass::ProcessorGroupNames::SyncWorldToMass);
+	bRequiresGameThreadExecution = true; // due to UMassDebuggerSubsystem access
 }
 
 void UDebugVisLocationProcessor::ConfigureQueries()
@@ -75,6 +76,7 @@ UMassProcessor_UpdateDebugVis::UMassProcessor_UpdateDebugVis()
 	: EntityQuery(*this)
 {
 	ExecutionOrder.ExecuteAfter.Add(UE::Mass::ProcessorGroupNames::UpdateWorldFromMass);
+	bRequiresGameThreadExecution = true; // due to UMassDebuggerSubsystem access
 }
 
 void UMassProcessor_UpdateDebugVis::ConfigureQueries() 
diff --git a/Engine/Plugins/Runtime/MassGameplay/Source/MassGameplayDebug/Public/MassDebuggerSubsystem.h b/Engine/Plugins/Runtime/MassGameplay/Source/MassGameplayDebug/Public/MassDebuggerSubsystem.h
--- a/Engine/Plugins/Runtime/MassGameplay/Source/MassGameplayDebug/Public/MassDebuggerSubsystem.h
+++ b/Engine/Plugins/Runtime/MassGameplay/Source/MassGameplayDebug/Public/MassDebuggerSubsystem.h
@@ -71,3 +71,12 @@ protected:
 	AMassDebugVisualizer* DebugVisualizer;
 };
 
+
+template<>
+struct TMassExternalSubsystemTraits<UMassDebuggerSubsystem> final
+{
+	enum
+	{
+		GameThreadOnly = true
+	};
+};
\ No newline at end of file
diff --git a/Engine/Plugins/Runtime/MassGameplay/Source/MassRepresentation/Private/MassRepresentationProcessor.cpp b/Engine/Plugins/Runtime/MassGameplay/Source/MassRepresentation/Private/MassRepresentationProcessor.cpp
--- a/Engine/Plugins/Runtime/MassGameplay/Source/MassRepresentation/Private/MassRepresentationProcessor.cpp
+++ b/Engine/Plugins/Runtime/MassGameplay/Source/MassRepresentation/Private/MassRepresentationProcessor.cpp
@@ -373,6 +373,7 @@ UMassRepresentationFragmentDestructor::UMassRepresentationFragmentDestructor()
 	ObservedType = FMassRepresentationFragment::StaticStruct();
 	Operation = EMassObservedOperation::Remove;
 	ExecutionFlags = (int32)(EProcessorExecutionFlags::All);
+	bRequiresGameThreadExecution = true; // due to FMassRepresentationSubsystemSharedFragment.RepresentationSubsystem use
 }
 
 void UMassRepresentationFragmentDestructor::ConfigureQueries()
diff --git a/Engine/Plugins/Runtime/MassGameplay/Source/MassRepresentation/Public/MassRepresentationFragments.h b/Engine/Plugins/Runtime/MassGameplay/Source/MassRepresentation/Public/MassRepresentationFragments.h
--- a/Engine/Plugins/Runtime/MassGameplay/Source/MassRepresentation/Public/MassRepresentationFragments.h
+++ b/Engine/Plugins/Runtime/MassGameplay/Source/MassRepresentation/Public/MassRepresentationFragments.h
@@ -64,6 +64,16 @@ struct FMassRepresentationSubsystemSharedFragment : public FMassSharedFragment
 	UMassRepresentationSubsystem* RepresentationSubsystem = nullptr;
 };
 
+template<>
+struct TMassExternalSubsystemTraits<FMassRepresentationSubsystemSharedFragment> final
+{
+	enum
+	{
+		GameThreadOnly = true
+	};
+};
+
+
 USTRUCT()
 struct FMassRepresentationParameters : public FMassSharedFragment
 {
@@ -109,6 +119,16 @@ struct FMassRepresentationParameters : public FMassSharedFragment
 	mutable UMassRepresentationActorManagement* CachedRepresentationActorManagement = nullptr;
 };
 
+template<>
+struct TMassExternalSubsystemTraits<FMassRepresentationParameters> final
+{
+	enum
+	{
+		GameThreadOnly = true
+	};
+};
+
+
 inline void FMassRepresentationParameters::ComputeCachedValues() const
 {
 	// Calculate the default representation when actor isn't spawned yet.
