[MassRepresentation] use more verbose names for MassVisualizer actors created from the MassRepresentationSubsystem subclasses.
#rnx

diff --git a/Engine/Plugins/Runtime/MassGameplay/Source/MassRepresentation/Private/MassRepresentationSubsystem.cpp b/Engine/Plugins/Runtime/MassGameplay/Source/MassRepresentation/Private/MassRepresentationSubsystem.cpp
--- a/Engine/Plugins/Runtime/MassGameplay/Source/MassRepresentation/Private/MassRepresentationSubsystem.cpp
+++ b/Engine/Plugins/Runtime/MassGameplay/Source/MassRepresentation/Private/MassRepresentationSubsystem.cpp
@@ -315,6 +315,10 @@ void UMassRepresentationSubsystem::Initialize(FSubsystemCollectionBase& Collecti
 			Visualizer = World->SpawnActor<AMassVisualizer>(SpawnInfo);
 			check(Visualizer);
 			VisualizationComponent = &Visualizer->GetVisualizationComponent();
+
+#if WITH_EDITOR
+			Visualizer->SetActorLabel(FString::Printf(TEXT("%sVisualizer"), *GetClass()->GetName()), /*bMarkDirty*/false);
+#endif
 		}
 
 		UMassSimulationSubsystem* SimSystem = World->GetSubsystem<UMassSimulationSubsystem>();
