Marked UMassUpdateISMProcessor as requiring game-thread execution, due to it relying on FMassRepresentationSubsystemSharedFragment which is marked for gamethread use only.

diff --git a/Engine/Plugins/Runtime/MassGameplay/Source/MassRepresentation/Private/MassUpdateISMProcessor.cpp b/Engine/Plugins/Runtime/MassGameplay/Source/MassRepresentation/Private/MassUpdateISMProcessor.cpp
--- a/Engine/Plugins/Runtime/MassGameplay/Source/MassRepresentation/Private/MassUpdateISMProcessor.cpp
+++ b/Engine/Plugins/Runtime/MassGameplay/Source/MassRepresentation/Private/MassUpdateISMProcessor.cpp
@@ -15,6 +15,7 @@ UMassUpdateISMProcessor::UMassUpdateISMProcessor()
 	ExecutionFlags = (int32)(EProcessorExecutionFlags::Client | EProcessorExecutionFlags::Standalone);
 
 	ExecutionOrder.ExecuteAfter.Add(UE::Mass::ProcessorGroupNames::Representation);
+	bRequiresGameThreadExecution = true;
 }
 
 void UMassUpdateISMProcessor::ConfigureQueries()
