PR #9246: Mass Spawner infinite loop bug fix. (Contributed by [REDACTED_PERSON])

diff --git a/Engine/Plugins/Runtime/MassGameplay/Source/MassSpawner/Private/MassEntityZoneGraphSpawnPointsGenerator.cpp b/Engine/Plugins/Runtime/MassGameplay/Source/MassSpawner/Private/MassEntityZoneGraphSpawnPointsGenerator.cpp
--- a/Engine/Plugins/Runtime/MassGameplay/Source/MassSpawner/Private/MassEntityZoneGraphSpawnPointsGenerator.cpp
+++ b/Engine/Plugins/Runtime/MassGameplay/Source/MassSpawner/Private/MassEntityZoneGraphSpawnPointsGenerator.cpp
@@ -116,8 +116,15 @@ void UMassEntityZoneGraphSpawnPointsGenerator::Generate(UObject& QueryOwner, TCo
 
 void UMassEntityZoneGraphSpawnPointsGenerator::GeneratePointsForZoneGraphData(const ::AZoneGraphData& ZoneGraphData, TArray<FVector>& Locations, const FRandomStream& RandomStream) const
 {
-	const FZoneGraphStorage& ZoneGraphStorage = ZoneGraphData.GetStorage();
+	// Avoid an infinite loop.
+	if (MinGap == 0.0f && MaxGap == 0.0f)
+	{
+		UE_VLOG_UELOG(this, LogMassSpawner, Error, TEXT("You cannot set both Min Gap and Max Gap to 0.0f"));
+		return;						
+	}
 
+	const FZoneGraphStorage &ZoneGraphStorage = ZoneGraphData.GetStorage();
+	
 	// Loop through all lanes
 	for (int32 LaneIndex = 0; LaneIndex < ZoneGraphStorage.Lanes.Num(); ++LaneIndex)
 	{
