Constified FMassReplicationContext's LOD subsystem access

diff --git a/Engine/Plugins/Runtime/MassGameplay/Source/MassReplication/Public/MassReplicationProcessor.h b/Engine/Plugins/Runtime/MassGameplay/Source/MassReplication/Public/MassReplicationProcessor.h
--- a/Engine/Plugins/Runtime/MassGameplay/Source/MassReplication/Public/MassReplicationProcessor.h
+++ b/Engine/Plugins/Runtime/MassGameplay/Source/MassReplication/Public/MassReplicationProcessor.h
@@ -51,14 +51,14 @@ protected:
 
 struct FMassReplicationContext
 {
-	FMassReplicationContext(UWorld& InWorld, UMassLODSubsystem& InLODSubsystem, UMassReplicationSubsystem& InReplicationSubsystem)
+	FMassReplicationContext(UWorld& InWorld, const UMassLODSubsystem& InLODSubsystem, UMassReplicationSubsystem& InReplicationSubsystem)
 		: World(InWorld)
 		, LODSubsystem(InLODSubsystem)
 		, ReplicationSubsystem(InReplicationSubsystem)
 	{}
 
 	UWorld& World;
-	UMassLODSubsystem& LODSubsystem;
+	const UMassLODSubsystem& LODSubsystem;
 	UMassReplicationSubsystem& ReplicationSubsystem;
 };
 
