Added MassProcessor.ProcessorRequirements to track subsystems (and potentially other elements) accessed by a given processor outside of MassEntityQuery.ForEachEntityChunk.
This change allowed us to remove hacky FMassExecutionContext.bSubsystemRequirementsSet.

Updated subsystems access based on MassProcessor.ProcessorRequirements
* UMassSignalSubsystem
* UZoneGraphAnnotationSubsystem
* UMassLODSubsystem
and removed member pointers to these systems.

diff --git a/Engine/Plugins/AI/MassAI/Source/MassAIBehavior/Private/MassStateTreeProcessors.cpp b/Engine/Plugins/AI/MassAI/Source/MassAIBehavior/Private/MassStateTreeProcessors.cpp
--- a/Engine/Plugins/AI/MassAI/Source/MassAIBehavior/Private/MassStateTreeProcessors.cpp
+++ b/Engine/Plugins/AI/MassAI/Source/MassAIBehavior/Private/MassStateTreeProcessors.cpp
@@ -242,11 +242,6 @@ UMassStateTreeActivationProcessor::UMassStateTreeActivationProcessor()
 	bRequiresGameThreadExecution = true; // due to UMassStateTreeSubsystem RW access
 }
 
-void UMassStateTreeActivationProcessor::Initialize(UObject& Owner)
-{
-	SignalSubsystem = UWorld::GetSubsystem<UMassSignalSubsystem>(Owner.GetWorld());
-}
-
 void UMassStateTreeActivationProcessor::ConfigureQueries()
 {
 	EntityQuery.AddRequirement<FMassStateTreeInstanceFragment>(EMassFragmentAccess::ReadWrite);
@@ -254,14 +249,14 @@ void UMassStateTreeActivationProcessor::ConfigureQueries()
 	EntityQuery.AddTagRequirement<FMassStateTreeActivatedTag>(EMassFragmentPresence::None);
 	EntityQuery.AddChunkRequirement<FMassSimulationVariableTickChunkFragment>(EMassFragmentAccess::ReadOnly, EMassFragmentPresence::Optional);
 	EntityQuery.AddSubsystemRequirement<UMassStateTreeSubsystem>(EMassFragmentAccess::ReadWrite);
+
+	ProcessorRequirements.AddSubsystemRequirement<UMassSignalSubsystem>(EMassFragmentAccess::ReadWrite);
 }
 
 void UMassStateTreeActivationProcessor::Execute(UMassEntitySubsystem& EntitySubsystem, FMassExecutionContext& Context)
 {
-	if (SignalSubsystem == nullptr)
-	{
-		return;
-	}
+	ConfigureContextForProcessorUse(Context);
+	UMassSignalSubsystem& SignalSubsystem = Context.GetMutableSubsystemChecked<UMassSignalSubsystem>(EntitySubsystem.GetWorld());
 
 	const UMassBehaviorSettings* BehaviorSettings = GetDefault<UMassBehaviorSettings>();
 	check(BehaviorSettings);
@@ -269,7 +264,7 @@ void UMassStateTreeActivationProcessor::Execute(UMassEntitySubsystem& EntitySubs
 	// StateTree processor relies on signals to be ticked but we need an 'initial tick' to set the tree in the proper state.
 	// The initializer provides that by sending a signal to all new entities that use StateTree.
 
-	FMassStateTreeExecutionContext StateTreeContext(EntitySubsystem, *SignalSubsystem, Context);
+	FMassStateTreeExecutionContext StateTreeContext(EntitySubsystem, SignalSubsystem, Context);
 
 	const float TimeInSeconds = EntitySubsystem.GetWorld()->GetTimeSeconds();
 
@@ -325,8 +320,7 @@ void UMassStateTreeActivationProcessor::Execute(UMassEntitySubsystem& EntitySubs
 	// Signal all entities inside the consolidated list
 	if (EntitiesToSignal.Num())
 	{
-		checkf(SignalSubsystem != nullptr, TEXT("Expecting a valid MassSignalSubsystem when activating state trees."));
-		SignalSubsystem->SignalEntities(UE::Mass::Signals::StateTreeActivate, EntitiesToSignal);
+		SignalSubsystem.SignalEntities(UE::Mass::Signals::StateTreeActivate, EntitiesToSignal);
 	}
 }
 
@@ -351,30 +345,30 @@ void UMassStateTreeProcessor::Initialize(UObject& Owner)
 {
 	Super::Initialize(Owner);
 
-	SignalSubsystem = UWorld::GetSubsystem<UMassSignalSubsystem>(Owner.GetWorld());
+	UMassSignalSubsystem* SignalSubsystem = UWorld::GetSubsystem<UMassSignalSubsystem>(Owner.GetWorld());
 
-	SubscribeToSignal(UE::Mass::Signals::StateTreeActivate);
-	SubscribeToSignal(UE::Mass::Signals::LookAtFinished);
-	SubscribeToSignal(UE::Mass::Signals::NewStateTreeTaskRequired);
-	SubscribeToSignal(UE::Mass::Signals::StandTaskFinished);
-	SubscribeToSignal(UE::Mass::Signals::DelayedTransitionWakeup);
+	SubscribeToSignal(*SignalSubsystem, UE::Mass::Signals::StateTreeActivate);
+	SubscribeToSignal(*SignalSubsystem, UE::Mass::Signals::LookAtFinished);
+	SubscribeToSignal(*SignalSubsystem, UE::Mass::Signals::NewStateTreeTaskRequired);
+	SubscribeToSignal(*SignalSubsystem, UE::Mass::Signals::StandTaskFinished);
+	SubscribeToSignal(*SignalSubsystem, UE::Mass::Signals::DelayedTransitionWakeup);
 
 	// @todo MassStateTree: add a way to register/unregister from enter/exit state (need reference counting)
-	SubscribeToSignal(UE::Mass::Signals::SmartObjectRequestCandidates);
-	SubscribeToSignal(UE::Mass::Signals::SmartObjectCandidatesReady);
-	SubscribeToSignal(UE::Mass::Signals::SmartObjectInteractionDone);
-	SubscribeToSignal(UE::Mass::Signals::SmartObjectInteractionAborted);
+	SubscribeToSignal(*SignalSubsystem, UE::Mass::Signals::SmartObjectRequestCandidates);
+	SubscribeToSignal(*SignalSubsystem, UE::Mass::Signals::SmartObjectCandidatesReady);
+	SubscribeToSignal(*SignalSubsystem, UE::Mass::Signals::SmartObjectInteractionDone);
+	SubscribeToSignal(*SignalSubsystem, UE::Mass::Signals::SmartObjectInteractionAborted);
 
-	SubscribeToSignal(UE::Mass::Signals::FollowPointPathStart);
-	SubscribeToSignal(UE::Mass::Signals::FollowPointPathDone);
-	SubscribeToSignal(UE::Mass::Signals::CurrentLaneChanged);
+	SubscribeToSignal(*SignalSubsystem, UE::Mass::Signals::FollowPointPathStart);
+	SubscribeToSignal(*SignalSubsystem, UE::Mass::Signals::FollowPointPathDone);
+	SubscribeToSignal(*SignalSubsystem, UE::Mass::Signals::CurrentLaneChanged);
 
-	SubscribeToSignal(UE::Mass::Signals::AnnotationTagsChanged);
+	SubscribeToSignal(*SignalSubsystem, UE::Mass::Signals::AnnotationTagsChanged);
 
-	SubscribeToSignal(UE::Mass::Signals::HitReceived);
+	SubscribeToSignal(*SignalSubsystem, UE::Mass::Signals::HitReceived);
 
 	// @todo MassStateTree: move this to its game plugin when possible
-	SubscribeToSignal(UE::Mass::Signals::ContextualAnimTaskFinished);
+	SubscribeToSignal(*SignalSubsystem, UE::Mass::Signals::ContextualAnimTaskFinished);
 }
 
 void UMassStateTreeProcessor::ConfigureQueries()
@@ -382,24 +376,26 @@ void UMassStateTreeProcessor::ConfigureQueries()
 	EntityQuery.AddRequirement<FMassStateTreeInstanceFragment>(EMassFragmentAccess::ReadWrite);
 	EntityQuery.AddConstSharedRequirement<FMassStateTreeSharedFragment>();
 	EntityQuery.AddSubsystemRequirement<UMassStateTreeSubsystem>(EMassFragmentAccess::ReadWrite);
+
+	ProcessorRequirements.AddSubsystemRequirement<UMassSignalSubsystem>(EMassFragmentAccess::ReadWrite);
 }
 
 void UMassStateTreeProcessor::SignalEntities(UMassEntitySubsystem& EntitySubsystem, FMassExecutionContext& Context, FMassSignalNameLookup& EntitySignals)
 {
-	if (SignalSubsystem == nullptr)
-	{
-		return;
-	}
+	UWorld* World = EntitySubsystem.GetWorld();
+	ConfigureContextForProcessorUse(Context);
+	UMassSignalSubsystem& SignalSubsystem = Context.GetMutableSubsystemChecked<UMassSignalSubsystem>(World);
+	
 	QUICK_SCOPE_CYCLE_COUNTER(StateTreeProcessor_Run);
 	CSV_SCOPED_TIMING_STAT_EXCLUSIVE(StateTreeProcessorExecute);
 
 	const float TimeInSeconds = EntitySubsystem.GetWorld()->GetTimeSeconds();
-	FMassStateTreeExecutionContext StateTreeContext(EntitySubsystem, *SignalSubsystem, Context);
+	FMassStateTreeExecutionContext StateTreeContext(EntitySubsystem, SignalSubsystem, Context);
 
 	TArray<FMassEntityHandle> EntitiesToSignal;
 
 	EntityQuery.ForEachEntityChunk(EntitySubsystem, Context,
-		[this, &StateTreeContext, TimeInSeconds, &EntitiesToSignal, World = EntitySubsystem.GetWorld(), &EntitySignals](FMassExecutionContext& Context)
+		[this, &StateTreeContext, TimeInSeconds, &EntitiesToSignal, World, &EntitySignals](FMassExecutionContext& Context)
 		{
 			// Keep stats regarding the amount of tree instances ticked per frame
 			CSV_CUSTOM_STAT(StateTreeProcessor, NumTickedStateTree, Context.GetNumEntities(), ECsvCustomStatOp::Accumulate);
@@ -453,6 +449,6 @@ void UMassStateTreeProcessor::SignalEntities(UMassEntitySubsystem& EntitySubsyst
 
 	if (EntitiesToSignal.Num())
 	{
-		SignalSubsystem->SignalEntities(UE::Mass::Signals::NewStateTreeTaskRequired, EntitiesToSignal);
+		SignalSubsystem.SignalEntities(UE::Mass::Signals::NewStateTreeTaskRequired, EntitiesToSignal);
 	}
 }
diff --git a/Engine/Plugins/AI/MassAI/Source/MassAIBehavior/Private/MassZoneGraphAnnotationProcessors.cpp b/Engine/Plugins/AI/MassAI/Source/MassAIBehavior/Private/MassZoneGraphAnnotationProcessors.cpp
--- a/Engine/Plugins/AI/MassAI/Source/MassAIBehavior/Private/MassZoneGraphAnnotationProcessors.cpp
+++ b/Engine/Plugins/AI/MassAI/Source/MassAIBehavior/Private/MassZoneGraphAnnotationProcessors.cpp
@@ -70,7 +70,8 @@ void UMassZoneGraphAnnotationTagUpdateProcessor::Initialize(UObject& Owner)
 {
 	Super::Initialize(Owner);
 
-	SubscribeToSignal(UE::Mass::Signals::CurrentLaneChanged);
+	UMassSignalSubsystem* SignalSubsystem = UWorld::GetSubsystem<UMassSignalSubsystem>(Owner.GetWorld());
+	SubscribeToSignal(*SignalSubsystem, UE::Mass::Signals::CurrentLaneChanged);
 }
 
 void UMassZoneGraphAnnotationTagUpdateProcessor::ConfigureQueries()
@@ -81,6 +82,8 @@ void UMassZoneGraphAnnotationTagUpdateProcessor::ConfigureQueries()
 	EntityQuery.AddChunkRequirement<FMassZoneGraphAnnotationVariableTickChunkFragment>(EMassFragmentAccess::ReadOnly);
 	EntityQuery.AddChunkRequirement<FMassSimulationVariableTickChunkFragment>(EMassFragmentAccess::ReadOnly, EMassFragmentPresence::Optional);
 	EntityQuery.AddSubsystemRequirement<UZoneGraphAnnotationSubsystem>(EMassFragmentAccess::ReadWrite);
+
+	ProcessorRequirements.AddSubsystemRequirement<UMassSignalSubsystem>(EMassFragmentAccess::ReadWrite);
 }
 
 void UMassZoneGraphAnnotationTagUpdateProcessor::Execute(UMassEntitySubsystem& EntitySubsystem, FMassExecutionContext& Context)
@@ -90,7 +93,8 @@ void UMassZoneGraphAnnotationTagUpdateProcessor::Execute(UMassEntitySubsystem& E
 	// Calling super will update the signals, and call SignalEntities() below.
 	Super::Execute(EntitySubsystem, Context);
 
-	EntityQuery.ForEachEntityChunk(EntitySubsystem, Context, [this, World = EntitySubsystem.GetWorld()](FMassExecutionContext& Context)
+	UWorld* World = EntitySubsystem.GetWorld();
+	EntityQuery.ForEachEntityChunk(EntitySubsystem, Context, [this, World](FMassExecutionContext& Context)
 	{
 		// Periodically update tags.
 		if (!FMassZoneGraphAnnotationVariableTickChunkFragment::UpdateChunk(Context))
@@ -115,7 +119,9 @@ void UMassZoneGraphAnnotationTagUpdateProcessor::Execute(UMassEntitySubsystem& E
 
 	if (TransientEntitiesToSignal.Num())
 	{
-		SignalSubsystem->SignalEntities(UE::Mass::Signals::AnnotationTagsChanged, TransientEntitiesToSignal);
+		ConfigureContextForProcessorUse(Context);
+		UMassSignalSubsystem& SignalSubsystem = Context.GetMutableSubsystemChecked<UMassSignalSubsystem>(World);
+		SignalSubsystem.SignalEntities(UE::Mass::Signals::AnnotationTagsChanged, TransientEntitiesToSignal);
 	}
 }
 
diff --git a/Engine/Plugins/AI/MassAI/Source/MassAIBehavior/Public/MassStateTreeProcessors.h b/Engine/Plugins/AI/MassAI/Source/MassAIBehavior/Public/MassStateTreeProcessors.h
--- a/Engine/Plugins/AI/MassAI/Source/MassAIBehavior/Public/MassStateTreeProcessors.h
+++ b/Engine/Plugins/AI/MassAI/Source/MassAIBehavior/Public/MassStateTreeProcessors.h
@@ -49,12 +49,8 @@ class MASSAIBEHAVIOR_API UMassStateTreeActivationProcessor : public UMassProcess
 public:
 	UMassStateTreeActivationProcessor();
 protected:
-	virtual void Initialize(UObject& Owner) override;
 	virtual void ConfigureQueries() override;
 	virtual void Execute(UMassEntitySubsystem& EntitySubsystem, FMassExecutionContext& Context) override;
-
-	UPROPERTY(Transient)
-	UMassSignalSubsystem* SignalSubsystem = nullptr;
 	
 	FMassEntityQuery EntityQuery;
 };
diff --git a/Engine/Plugins/AI/MassCrowd/Source/MassCrowd/Private/MassCrowdNavigationProcessor.cpp b/Engine/Plugins/AI/MassCrowd/Source/MassCrowd/Private/MassCrowdNavigationProcessor.cpp
--- a/Engine/Plugins/AI/MassCrowd/Source/MassCrowd/Private/MassCrowdNavigationProcessor.cpp
+++ b/Engine/Plugins/AI/MassCrowd/Source/MassCrowd/Private/MassCrowdNavigationProcessor.cpp
@@ -12,6 +12,7 @@
 #include "Annotations/ZoneGraphDisturbanceAnnotation.h"
 #include "VisualLogger/VisualLogger.h"
 #include "MassSimulationLOD.h"
+#include "MassSignalSubsystem.h"
 
 //----------------------------------------------------------------------//
 // UMassCrowdLaneTrackingSignalProcessor
@@ -33,7 +34,8 @@ void UMassCrowdLaneTrackingSignalProcessor::Initialize(UObject& Owner)
 {
 	Super::Initialize(Owner);
 	
-	SubscribeToSignal(UE::Mass::Signals::CurrentLaneChanged);
+	UMassSignalSubsystem* SignalSubsystem = UWorld::GetSubsystem<UMassSignalSubsystem>(Owner.GetWorld());
+	SubscribeToSignal(*SignalSubsystem, UE::Mass::Signals::CurrentLaneChanged);
 }
 
 void UMassCrowdLaneTrackingSignalProcessor::SignalEntities(UMassEntitySubsystem& EntitySubsystem, FMassExecutionContext& Context, FMassSignalNameLookup& EntitySignals)
diff --git a/Engine/Plugins/AI/MassCrowd/Source/MassCrowd/Private/MassCrowdServerRepresentationLODProcessor.cpp b/Engine/Plugins/AI/MassCrowd/Source/MassCrowd/Private/MassCrowdServerRepresentationLODProcessor.cpp
--- a/Engine/Plugins/AI/MassCrowd/Source/MassCrowd/Private/MassCrowdServerRepresentationLODProcessor.cpp
+++ b/Engine/Plugins/AI/MassCrowd/Source/MassCrowd/Private/MassCrowdServerRepresentationLODProcessor.cpp
@@ -35,6 +35,8 @@ void UMassCrowdServerRepresentationLODProcessor::ConfigureQueries()
 	EntityQuery.AddRequirement<FTransformFragment>(EMassFragmentAccess::ReadOnly);
 	EntityQuery.AddRequirement<FMassViewerInfoFragment>(EMassFragmentAccess::ReadOnly);
 	EntityQuery.AddRequirement<FMassRepresentationLODFragment>(EMassFragmentAccess::ReadWrite);
+
+	ProcessorRequirements.AddSubsystemRequirement<UMassLODSubsystem>(EMassFragmentAccess::ReadOnly);
 }
 
 void UMassCrowdServerRepresentationLODProcessor::Initialize(UObject& InOwner)
@@ -48,7 +50,7 @@ void UMassCrowdServerRepresentationLODProcessor::Execute(UMassEntitySubsystem& E
 {
 	TRACE_CPUPROFILER_EVENT_SCOPE(TEXT("CrowdServerRepresentationLOD"))
 
-	// @todo this access patter will be transformed into "computation tasks" in near future. Do not duplicate.
+	ConfigureContextForProcessorUse(Context);
 	const UMassLODSubsystem& LODSubsystem = Context.GetSubsystemChecked<UMassLODSubsystem>(EntitySubsystem.GetWorld());
 	const TArray<FViewerInfo>& Viewers = LODSubsystem.GetViewers();
 	LODCalculator.PrepareExecution(Viewers);
diff --git a/Engine/Plugins/Runtime/MassEntity/Source/MassEntity/Private/MassEntityQuery.cpp b/Engine/Plugins/Runtime/MassEntity/Source/MassEntity/Private/MassEntityQuery.cpp
--- a/Engine/Plugins/Runtime/MassEntity/Source/MassEntity/Private/MassEntityQuery.cpp
+++ b/Engine/Plugins/Runtime/MassEntity/Source/MassEntity/Private/MassEntityQuery.cpp
@@ -312,6 +312,12 @@ void FMassEntityQuery::ParallelForEachEntityChunk(UMassEntitySubsystem& EntitySu
 	ExecutionContext.FlushDeferred(EntitySubsystem);
 }
 
+void FMassEntityQuery::ApplyQueryRequirementsToContext(FMassExecutionContext& ExecutionContext)
+{
+	ExecutionContext.SetSubsystemRequirements(RequiredConstSubsystems, RequiredMutableSubsystems);
+	ExecutionContext.SetRequirements(Requirements, ChunkRequirements, ConstSharedRequirements, SharedRequirements);
+}
+
 int32 FMassEntityQuery::GetNumMatchingEntities(UMassEntitySubsystem& InEntitySubsystem)
 {
 	CacheArchetypes(InEntitySubsystem);
diff --git a/Engine/Plugins/Runtime/MassEntity/Source/MassEntity/Private/MassExecutionContext.cpp b/Engine/Plugins/Runtime/MassEntity/Source/MassEntity/Private/MassExecutionContext.cpp
--- a/Engine/Plugins/Runtime/MassEntity/Source/MassEntity/Private/MassExecutionContext.cpp
+++ b/Engine/Plugins/Runtime/MassEntity/Source/MassEntity/Private/MassExecutionContext.cpp
@@ -19,14 +19,12 @@ void FMassExecutionContext::ClearExecutionData()
 	ChunkSerialModificationNumber = -1;
 	ConstSubsystemsBitSet.Reset();
 	MutableSubsystemsBitSet.Reset();
-	bSubsystemRequirementsSet = false;
 }
 
 void FMassExecutionContext::SetSubsystemRequirements(const FMassExternalSubystemBitSet& RequiredConstSubsystems, const FMassExternalSubystemBitSet& RequiredMutableSubsystems)
 {
 	ConstSubsystemsBitSet = RequiredConstSubsystems;
 	MutableSubsystemsBitSet = RequiredMutableSubsystems;
-	bSubsystemRequirementsSet = true;
 }
 
 void FMassExecutionContext::SetEntityCollection(const FMassArchetypeEntityCollection& InEntityCollection)
diff --git a/Engine/Plugins/Runtime/MassEntity/Source/MassEntity/Private/MassProcessor.cpp b/Engine/Plugins/Runtime/MassEntity/Source/MassEntity/Private/MassProcessor.cpp
--- a/Engine/Plugins/Runtime/MassEntity/Source/MassEntity/Private/MassProcessor.cpp
+++ b/Engine/Plugins/Runtime/MassEntity/Source/MassEntity/Private/MassProcessor.cpp
@@ -100,8 +100,14 @@ public:
 // UMassProcessor 
 //----------------------------------------------------------------------//
 UMassProcessor::UMassProcessor(const FObjectInitializer& ObjectInitializer)
+	: UMassProcessor()
+{
+}
+
+UMassProcessor::UMassProcessor()
 	: ExecutionFlags((int32)(EProcessorExecutionFlags::Server | EProcessorExecutionFlags::Standalone))
 {
+	RegisterQuery(ProcessorRequirements);
 }
 
 void UMassProcessor::SetShouldAutoRegisterWithGlobalList(const bool bAutoRegister)
@@ -148,6 +154,11 @@ void UMassProcessor::CallExecute(UMassEntitySubsystem& EntitySubsystem, FMassExe
 	Execute(EntitySubsystem, Context);
 }
 
+void UMassProcessor::ConfigureContextForProcessorUse(FMassExecutionContext& Context)
+{
+	ProcessorRequirements.ApplyQueryRequirementsToContext(Context);
+}
+
 void UMassProcessor::ExportRequirements(FMassExecutionRequirements& OutRequirements) const
 {
 	for (FMassEntityQuery* Query : OwnedQueries)
diff --git a/Engine/Plugins/Runtime/MassEntity/Source/MassEntity/Public/MassEntityQuery.h b/Engine/Plugins/Runtime/MassEntity/Source/MassEntity/Public/MassEntityQuery.h
--- a/Engine/Plugins/Runtime/MassEntity/Source/MassEntity/Public/MassEntityQuery.h
+++ b/Engine/Plugins/Runtime/MassEntity/Source/MassEntity/Public/MassEntityQuery.h
@@ -124,6 +124,9 @@ public:
 	 */
 	void ParallelForEachEntityChunk(UMassEntitySubsystem& EntitySubsystem, FMassExecutionContext& ExecutionContext, const FMassExecuteFunction& ExecuteFunction);
 
+	/** Configures ExecutionContext with access/presence requirements declared by this query */
+	void ApplyQueryRequirementsToContext(FMassExecutionContext& ExecutionContext);
+
 	/** Will gather all archetypes from InEntitySubsystem matching this->Requirements.
 	 *  Note that no work will be done if the cached data is up to date (as tracked by EntitySubsystemHash and 
 	 *	ArchetypeDataVersion properties). */
diff --git a/Engine/Plugins/Runtime/MassEntity/Source/MassEntity/Public/MassExecutionContext.h b/Engine/Plugins/Runtime/MassEntity/Source/MassEntity/Public/MassExecutionContext.h
--- a/Engine/Plugins/Runtime/MassEntity/Source/MassEntity/Public/MassExecutionContext.h
+++ b/Engine/Plugins/Runtime/MassEntity/Source/MassEntity/Public/MassExecutionContext.h
@@ -60,11 +60,6 @@ private:
 	 * is mainly utilized to avoid numerous small flushes in favor of fewer larger ones. */
 	bool bFlushDeferredCommands = true;
 
-	/** A temporary variable letting us access subsystems via a FMassExecutionContext outside of MassEntityQuery's 
-	 * execution. Will go away once entity-independent calculations as well as query-bound calculations are being 
-	 * performed with a common tasks interface */
-	bool bSubsystemRequirementsSet = false;
-
 	TArrayView<FFragmentView> GetMutableRequirements() { return FragmentViews; }
 	TArrayView<FChunkFragmentView> GetMutableChunkRequirements() { return ChunkFragmentViews; }
 	TArrayView<FConstSharedFragmentView> GetMutableConstSharedRequirements() { return ConstSharedFragmentViews; }
@@ -250,7 +245,7 @@ public:
 	{
 		// @todo consider getting this directly from entity subsystem - it should cache all the used system
 		const uint32 SystemIndex = FMassExternalSubystemBitSet::GetTypeIndex<T>();
-		if (bSubsystemRequirementsSet == false || ensure(MutableSubsystemsBitSet.IsBitSet(SystemIndex)))
+		if (ensure(MutableSubsystemsBitSet.IsBitSet(SystemIndex)))
 		{
 			return GetSubsystemInternal<T>(World, SystemIndex);
 		}
@@ -270,7 +265,7 @@ public:
 	const T* GetSubsystem(const UWorld* World)
 	{
 		const uint32 SystemIndex = FMassExternalSubystemBitSet::GetTypeIndex<T>();
-		if (bSubsystemRequirementsSet == false || ensure(ConstSubsystemsBitSet.IsBitSet(SystemIndex) || MutableSubsystemsBitSet.IsBitSet(SystemIndex)))
+		if (ensure(ConstSubsystemsBitSet.IsBitSet(SystemIndex) || MutableSubsystemsBitSet.IsBitSet(SystemIndex)))
 		{
 			return GetSubsystemInternal<T>(World, SystemIndex);
 		}
@@ -284,9 +279,6 @@ public:
 		check(InstancePtr);
 		return *InstancePtr;
 	}
-	
-
-	void SetSubsystemRequirements(const FMassExternalSubystemBitSet& RequiredConstSubsystems, const FMassExternalSubystemBitSet& RequiredMutableSubsystems);
 
 	template<typename T>
 	T* GetMutableSubsystem(const UWorld* World, const TSubclassOf<UWorldSubsystem> SubsystemClass)
@@ -350,6 +342,8 @@ public:
 	}
 
 protected:
+	void SetSubsystemRequirements(const FMassExternalSubystemBitSet& RequiredConstSubsystems, const FMassExternalSubystemBitSet& RequiredMutableSubsystems);
+
 	void SetRequirements(TConstArrayView<FMassFragmentRequirement> InRequirements, 
 		TConstArrayView<FMassFragmentRequirement> InChunkRequirements, 
 		TConstArrayView<FMassFragmentRequirement> InConstSharedRequirements, 
diff --git a/Engine/Plugins/Runtime/MassEntity/Source/MassEntity/Public/MassProcessor.h b/Engine/Plugins/Runtime/MassEntity/Source/MassEntity/Public/MassProcessor.h
--- a/Engine/Plugins/Runtime/MassEntity/Source/MassEntity/Public/MassProcessor.h
+++ b/Engine/Plugins/Runtime/MassEntity/Source/MassEntity/Public/MassProcessor.h
@@ -45,7 +45,8 @@ class MASSENTITY_API UMassProcessor : public UObject
 {
 	GENERATED_BODY()
 public:
-	UMassProcessor(const FObjectInitializer& ObjectInitializer = FObjectInitializer::Get());
+	UMassProcessor();
+	UMassProcessor(const FObjectInitializer& ObjectInitializer);
 
 	virtual void Initialize(UObject& Owner) {}
 	virtual FGraphEventRef DispatchProcessorTasks(UMassEntitySubsystem& EntitySubsystem, FMassExecutionContext& ExecutionContext, const FGraphEventArray& Prerequisites = FGraphEventArray());
@@ -56,6 +57,10 @@ public:
 	bool ShouldExecute(const EProcessorExecutionFlags CurrentExecutionFlags) const { return (GetExecutionFlags() & CurrentExecutionFlags) != EProcessorExecutionFlags::None; }
 	void CallExecute(UMassEntitySubsystem& EntitySubsystem, FMassExecutionContext& Context);
 	
+	/** Configures Context to accept access requests performed outside of FMassEntityQuery::ForEachEntityChunk calls.
+	 *  The access is declared via ProcessorQuery member variable */
+	void ConfigureContextForProcessorUse(FMassExecutionContext& Context);
+
 	bool AllowDuplicates() const { return bAllowDuplicates; }
 
 	virtual void DebugOutputDescription(FOutputDevice& Ar, int32 Indent = 0) const;
@@ -135,6 +140,9 @@ protected:
 	friend UMassCompositeProcessor;
 	friend FMassDebugger;
 
+	/** A query representing elements this processor is accessing in Execute function outside of query execution */
+	FMassEntityQuery ProcessorRequirements;
+
 private:
 	/** Stores processor's queries registered via RegisterQuery. 
 	 *  @note that it's safe to store pointers here since RegisterQuery does verify that a given registered query is 
diff --git a/Engine/Plugins/Runtime/MassGameplay/Source/MassLOD/Private/MassLODCollectorProcessor.cpp b/Engine/Plugins/Runtime/MassGameplay/Source/MassLOD/Private/MassLODCollectorProcessor.cpp
--- a/Engine/Plugins/Runtime/MassGameplay/Source/MassLOD/Private/MassLODCollectorProcessor.cpp
+++ b/Engine/Plugins/Runtime/MassGameplay/Source/MassLOD/Private/MassLODCollectorProcessor.cpp
@@ -49,6 +49,8 @@ void UMassLODCollectorProcessor::ConfigureQueries()
 	EntityQuery_NotVisibleRangeAndOffLOD.AddTagRequirement<FMassVisibilityCulledByDistanceTag>(EMassFragmentPresence::All);
 	EntityQuery_NotVisibleRangeAndOffLOD.AddTagRequirement<FMassOffLODTag>(EMassFragmentPresence::All);
 	EntityQuery_NotVisibleRangeAndOffLOD.RegisterWithProcessor(*this);
+
+	ProcessorRequirements.AddSubsystemRequirement<UMassLODSubsystem>(EMassFragmentAccess::ReadOnly);
 }
 
 template <bool bLocalViewersOnly>
@@ -78,7 +80,8 @@ void UMassLODCollectorProcessor::ExecuteInternal(UMassEntitySubsystem& EntitySub
 
 void UMassLODCollectorProcessor::Execute(UMassEntitySubsystem& EntitySubsystem, FMassExecutionContext& Context)
 {
-	// @todo this access patter will be transformed into "computation tasks" in near future. Do not duplicate.
+	ConfigureContextForProcessorUse(Context);
+
 	const UMassLODSubsystem& LODSubsystem = Context.GetSubsystemChecked<UMassLODSubsystem>(EntitySubsystem.GetWorld());
 	const TArray<FViewerInfo>& Viewers = LODSubsystem.GetViewers();
 	Collector.PrepareExecution(Viewers);
diff --git a/Engine/Plugins/Runtime/MassGameplay/Source/MassLOD/Private/MassSimulationLOD.cpp b/Engine/Plugins/Runtime/MassGameplay/Source/MassLOD/Private/MassSimulationLOD.cpp
--- a/Engine/Plugins/Runtime/MassGameplay/Source/MassLOD/Private/MassSimulationLOD.cpp
+++ b/Engine/Plugins/Runtime/MassGameplay/Source/MassLOD/Private/MassSimulationLOD.cpp
@@ -109,6 +109,8 @@ void UMassSimulationLODProcessor::ConfigureQueries()
 		const FMassSimulationLODParameters& LODParams = Context.GetConstSharedFragment<FMassSimulationLODParameters>();
 		return LODParams.bSetLODTags;
 	});
+
+	ProcessorRequirements.AddSubsystemRequirement<UMassLODSubsystem>(EMassFragmentAccess::ReadOnly);
 }
 
 void UMassSimulationLODProcessor::Execute(UMassEntitySubsystem& EntitySubsystem, FMassExecutionContext& Context)
@@ -117,7 +119,9 @@ void UMassSimulationLODProcessor::Execute(UMassEntitySubsystem& EntitySubsystem,
 
 	{
 		TRACE_CPUPROFILER_EVENT_SCOPE(TEXT("PrepareExecution"));
-		// @todo this access patter will be transformed into "computation tasks" in near future. Do not duplicate.
+		
+		ConfigureContextForProcessorUse(Context);
+
 		const UMassLODSubsystem& LODSubsystem = Context.GetSubsystemChecked<UMassLODSubsystem>(EntitySubsystem.GetWorld());
 		const TArray<FViewerInfo>& Viewers = LODSubsystem.GetViewers();
 
diff --git a/Engine/Plugins/Runtime/MassGameplay/Source/MassLOD/Public/MassLODSubsystem.h b/Engine/Plugins/Runtime/MassGameplay/Source/MassLOD/Public/MassLODSubsystem.h
--- a/Engine/Plugins/Runtime/MassGameplay/Source/MassLOD/Public/MassLODSubsystem.h
+++ b/Engine/Plugins/Runtime/MassGameplay/Source/MassLOD/Public/MassLODSubsystem.h
@@ -7,7 +7,7 @@
 #include "IndexedHandle.h"
 #include "MassLODTypes.h"
 #include "Subsystems/WorldSubsystem.h"
-
+#include "MassExternalSubsystemTraits.h"
 #include "MassLODSubsystem.generated.h"
 
 class UMassLODSubsystem;
@@ -137,3 +137,11 @@ private:
 
 };
 
+template<>
+struct TMassExternalSubsystemTraits<UMassLODSubsystem> final
+{
+	enum
+	{
+		GameThreadOnly = true
+	};
+};
diff --git a/Engine/Plugins/Runtime/MassGameplay/Source/MassReplication/Private/MassReplicationProcessor.cpp b/Engine/Plugins/Runtime/MassGameplay/Source/MassReplication/Private/MassReplicationProcessor.cpp
--- a/Engine/Plugins/Runtime/MassGameplay/Source/MassReplication/Private/MassReplicationProcessor.cpp
+++ b/Engine/Plugins/Runtime/MassGameplay/Source/MassReplication/Private/MassReplicationProcessor.cpp
@@ -59,6 +59,8 @@ void UMassReplicationProcessor::ConfigureQueries()
 	EntityQuery.AddRequirement<FMassReplicatedAgentFragment>(EMassFragmentAccess::ReadWrite);
 	EntityQuery.AddConstSharedRequirement<FMassReplicationParameters>();
 	EntityQuery.AddSharedRequirement<FMassReplicationSharedFragment>(EMassFragmentAccess::ReadWrite);
+
+	ProcessorRequirements.AddSubsystemRequirement<UMassLODSubsystem>(EMassFragmentAccess::ReadOnly);
 }
 
 void UMassReplicationProcessor::Initialize(UObject& Owner)
@@ -162,7 +164,8 @@ void UMassReplicationProcessor::Execute(UMassEntitySubsystem& EntitySubsystem, F
 		PrepareExecution(EntitySubsystem);
 	}
 
-	UMassLODSubsystem& LODSubsystem = Context.GetMutableSubsystemChecked<UMassLODSubsystem>(EntitySubsystem.GetWorld());
+	ConfigureContextForProcessorUse(Context);
+	const UMassLODSubsystem& LODSubsystem = Context.GetSubsystemChecked<UMassLODSubsystem>(EntitySubsystem.GetWorld());
 	const TArray<FViewerInfo>& AllViewersInfo = LODSubsystem.GetViewers();
 	const TArray<FMassClientHandle>& ClientHandles = ReplicationSubsystem->GetClientReplicationHandles();
 	for (const FMassClientHandle ClientHandle : ClientHandles)
diff --git a/Engine/Plugins/Runtime/MassGameplay/Source/MassRepresentation/Private/MassVisualizationLODProcessor.cpp b/Engine/Plugins/Runtime/MassGameplay/Source/MassRepresentation/Private/MassVisualizationLODProcessor.cpp
--- a/Engine/Plugins/Runtime/MassGameplay/Source/MassRepresentation/Private/MassVisualizationLODProcessor.cpp
+++ b/Engine/Plugins/Runtime/MassGameplay/Source/MassRepresentation/Private/MassVisualizationLODProcessor.cpp
@@ -46,6 +46,8 @@ void UMassVisualizationLODProcessor::ConfigureQueries()
 
 	DebugEntityQuery = BaseQuery;
 	DebugEntityQuery.RegisterWithProcessor(*this);
+
+	ProcessorRequirements.AddSubsystemRequirement<UMassLODSubsystem>(EMassFragmentAccess::ReadOnly);
 }
 
 void UMassVisualizationLODProcessor::Execute(UMassEntitySubsystem& EntitySubsystem, FMassExecutionContext& Context)
@@ -63,7 +65,7 @@ void UMassVisualizationLODProcessor::Execute(UMassEntitySubsystem& EntitySubsyst
 
 	{
 		TRACE_CPUPROFILER_EVENT_SCOPE(PrepareExecution)
-		// @todo this access patter will be transformed into "computation tasks" in near future. Do not duplicate.
+		ConfigureContextForProcessorUse(Context);
 		const UMassLODSubsystem& LODSubsystem = Context.GetSubsystemChecked<UMassLODSubsystem>(EntitySubsystem.GetWorld());
 		const TArray<FViewerInfo>& Viewers = LODSubsystem.GetViewers();
 		EntitySubsystem.ForEachSharedFragment<FMassVisualizationLODSharedFragment>([this, &Viewers](FMassVisualizationLODSharedFragment& LODSharedFragment)
diff --git a/Engine/Plugins/Runtime/MassGameplay/Source/MassSignals/Private/MassSignalProcessorBase.cpp b/Engine/Plugins/Runtime/MassGameplay/Source/MassSignals/Private/MassSignalProcessorBase.cpp
--- a/Engine/Plugins/Runtime/MassGameplay/Source/MassSignals/Private/MassSignalProcessorBase.cpp
+++ b/Engine/Plugins/Runtime/MassGameplay/Source/MassSignals/Private/MassSignalProcessorBase.cpp
@@ -11,14 +11,9 @@ UMassSignalProcessorBase::UMassSignalProcessorBase(const FObjectInitializer& Obj
 {
 }
 
-void UMassSignalProcessorBase::Initialize(UObject& Owner)
-{
-	SignalSubsystem = UWorld::GetSubsystem<UMassSignalSubsystem>(Owner.GetWorld());
-}
-
 void UMassSignalProcessorBase::BeginDestroy()
 {
-	if (SignalSubsystem)
+	if (UMassSignalSubsystem* SignalSubsystem = UWorld::GetSubsystem<UMassSignalSubsystem>(GetWorld()))
 	{
 		for (const FName& SignalName : RegisteredSignals)
 		{
@@ -29,12 +24,11 @@ void UMassSignalProcessorBase::BeginDestroy()
 	Super::BeginDestroy();
 }
 
-void UMassSignalProcessorBase::SubscribeToSignal(const FName SignalName)
+void UMassSignalProcessorBase::SubscribeToSignal(UMassSignalSubsystem& SignalSubsystem, const FName SignalName)
 {
-	check(SignalSubsystem);
 	check(!RegisteredSignals.Contains(SignalName));
 	RegisteredSignals.Add(SignalName);
-	SignalSubsystem->GetSignalDelegateByName(SignalName).AddUObject(this, &UMassSignalProcessorBase::OnSignalReceived);
+	SignalSubsystem.GetSignalDelegateByName(SignalName).AddUObject(this, &UMassSignalProcessorBase::OnSignalReceived);
 }
 
 void UMassSignalProcessorBase::Execute(UMassEntitySubsystem& EntitySubsystem, FMassExecutionContext& Context)
diff --git a/Engine/Plugins/Runtime/MassGameplay/Source/MassSignals/Public/MassSignalProcessorBase.h b/Engine/Plugins/Runtime/MassGameplay/Source/MassSignals/Public/MassSignalProcessorBase.h
--- a/Engine/Plugins/Runtime/MassGameplay/Source/MassSignals/Public/MassSignalProcessorBase.h
+++ b/Engine/Plugins/Runtime/MassGameplay/Source/MassSignals/Public/MassSignalProcessorBase.h
@@ -23,12 +23,6 @@ public:
 	UMassSignalProcessorBase(const FObjectInitializer& ObjectInitializer = FObjectInitializer::Get());
 
 protected:
-	/**
-	 * Initialize the processor
-	 * @param Owner of the Processor
-	 */
-	virtual void Initialize(UObject& Owner) override;
-
 	virtual void BeginDestroy() override;
 
 	/** Configure the owned FMassEntityQuery instances to express processor queries requirements */
@@ -60,11 +54,7 @@ protected:
 	 * To receive notification about a particular signal, you need to subscribe to it.
 	 * @param SignalName is the name of the signal to receive notification about
 	 */
-	void SubscribeToSignal(const FName SignalName);
-
-	
-	UPROPERTY(Transient)
-	UMassSignalSubsystem* SignalSubsystem = nullptr;
+	void SubscribeToSignal(UMassSignalSubsystem& SignalSubsystem, const FName SignalName);
 
 	FMassEntityQuery EntityQuery;
 
diff --git a/Engine/Plugins/Runtime/MassGameplay/Source/MassSignals/Public/MassSignalSubsystem.h b/Engine/Plugins/Runtime/MassGameplay/Source/MassSignals/Public/MassSignalSubsystem.h
--- a/Engine/Plugins/Runtime/MassGameplay/Source/MassSignals/Public/MassSignalSubsystem.h
+++ b/Engine/Plugins/Runtime/MassGameplay/Source/MassSignals/Public/MassSignalSubsystem.h
@@ -109,3 +109,12 @@ protected:
 
 	TArray<FDelayedSignal> DelayedSignals;
 };
+
+template<>
+struct TMassExternalSubsystemTraits<UMassSignalSubsystem> final
+{
+	enum
+	{
+		GameThreadOnly = true
+	};
+};
diff --git a/Engine/Plugins/Runtime/MassGameplay/Source/MassSmartObjects/Private/MassSmartObjectProcessor.cpp b/Engine/Plugins/Runtime/MassGameplay/Source/MassSmartObjects/Private/MassSmartObjectProcessor.cpp
--- a/Engine/Plugins/Runtime/MassGameplay/Source/MassSmartObjects/Private/MassSmartObjectProcessor.cpp
+++ b/Engine/Plugins/Runtime/MassGameplay/Source/MassSmartObjects/Private/MassSmartObjectProcessor.cpp
@@ -15,18 +15,13 @@
 #include "SmartObjectSubsystem.h"
 #include "VisualLogger/VisualLogger.h"
 #include "ZoneGraphAnnotationSubsystem.h"
+#include "MassGameplayExternalTraits.h"
 #include "ZoneGraphSubsystem.h"
 #include "MassGameplayExternalTraits.h"
 
 //----------------------------------------------------------------------//
 // UMassSmartObjectCandidatesFinderProcessor
 //----------------------------------------------------------------------//
-void UMassSmartObjectCandidatesFinderProcessor::Initialize(UObject& Owner)
-{
-	SignalSubsystem = UWorld::GetSubsystem<UMassSignalSubsystem>(Owner.GetWorld());
-	AnnotationSubsystem = UWorld::GetSubsystem<UZoneGraphAnnotationSubsystem>(GetWorld());
-}
-
 void UMassSmartObjectCandidatesFinderProcessor::ConfigureQueries()
 {
 	WorldRequestQuery.AddRequirement<FMassSmartObjectWorldLocationRequestFragment>(EMassFragmentAccess::ReadOnly);
@@ -41,6 +36,9 @@ void UMassSmartObjectCandidatesFinderProcessor::ConfigureQueries()
 	LaneRequestQuery.AddSubsystemRequirement<UZoneGraphSubsystem>(EMassFragmentAccess::ReadOnly);
 	LaneRequestQuery.AddSubsystemRequirement<USmartObjectSubsystem>(EMassFragmentAccess::ReadOnly);
 	LaneRequestQuery.RegisterWithProcessor(*this);
+
+	ProcessorRequirements.AddSubsystemRequirement<UMassSignalSubsystem>(EMassFragmentAccess::ReadWrite);
+	ProcessorRequirements.AddSubsystemRequirement<UZoneGraphAnnotationSubsystem>(EMassFragmentAccess::ReadOnly);
 }
 
 UMassSmartObjectCandidatesFinderProcessor::UMassSmartObjectCandidatesFinderProcessor()
@@ -56,9 +54,12 @@ UMassSmartObjectCandidatesFinderProcessor::UMassSmartObjectCandidatesFinderProce
 
 void UMassSmartObjectCandidatesFinderProcessor::Execute(UMassEntitySubsystem& EntitySubsystem, FMassExecutionContext& Context)
 {
-	checkf(SignalSubsystem != nullptr, TEXT("MassSignalSubsystem should exist when executing processors."));
-	checkf(AnnotationSubsystem != nullptr, TEXT("ZoneGraphAnnotationSubsystem should exist when executing processors."));
+	UWorld* World = EntitySubsystem.GetWorld();
 
+	ConfigureContextForProcessorUse(Context);
+	UMassSignalSubsystem& SignalSubsystem = Context.GetMutableSubsystemChecked<UMassSignalSubsystem>(World);
+	const UZoneGraphAnnotationSubsystem& AnnotationSubsystem = Context.GetSubsystemChecked<UZoneGraphAnnotationSubsystem>(World);
+	
 	// Create filter
 	FSmartObjectRequestFilter Filter;
 	Filter.BehaviorDefinitionClass = USmartObjectMassBehaviorDefinition::StaticClass();
@@ -87,7 +88,7 @@ void UMassSmartObjectCandidatesFinderProcessor::Execute(UMassEntitySubsystem& En
 	};
 
 	// Process world location based requests
-	WorldRequestQuery.ForEachEntityChunk(EntitySubsystem, Context, [this, &Filter, &EntitiesToSignal, &BeginRequestProcessing, &EndRequestProcessing, World = EntitySubsystem.GetWorld()](FMassExecutionContext& Context)
+	WorldRequestQuery.ForEachEntityChunk(EntitySubsystem, Context, [this, &Filter, &EntitiesToSignal, &BeginRequestProcessing, &EndRequestProcessing, World](FMassExecutionContext& Context)
 	{
 		const USmartObjectSubsystem& SmartObjectSubsystem = Context.GetSubsystemChecked<USmartObjectSubsystem>(World);
 
@@ -154,10 +155,10 @@ void UMassSmartObjectCandidatesFinderProcessor::Execute(UMassEntitySubsystem& En
 
 	// Process lane based requests
 	const FZoneGraphTag SmartObjectTag = GetDefault<UMassSmartObjectSettings>()->SmartObjectTag;
-	USmartObjectZoneAnnotations* Annotations = Cast<USmartObjectZoneAnnotations>(AnnotationSubsystem->GetFirstAnnotationForTag(SmartObjectTag));
+	USmartObjectZoneAnnotations* Annotations = Cast<USmartObjectZoneAnnotations>(AnnotationSubsystem.GetFirstAnnotationForTag(SmartObjectTag));
 
 	LaneRequestQuery.ForEachEntityChunk(EntitySubsystem, Context,
-		[this, Annotations, &Filter, SmartObjectTag, &EntitiesToSignal, &BeginRequestProcessing, &EndRequestProcessing, World = EntitySubsystem.GetWorld()](FMassExecutionContext& Context)
+		[&AnnotationSubsystem, Annotations, &Filter, SmartObjectTag, &EntitiesToSignal, &BeginRequestProcessing, &EndRequestProcessing, World = EntitySubsystem.GetWorld()](FMassExecutionContext& Context)
 		{
 #if WITH_MASSGAMEPLAY_DEBUG
 			const UZoneGraphSubsystem& ZoneGraphSubsystem = Context.GetSubsystemChecked<UZoneGraphSubsystem>(World);
@@ -218,7 +219,7 @@ void UMassSmartObjectCandidatesFinderProcessor::Execute(UMassEntitySubsystem& En
 				}
 
 				// Fetch current annotations for the specified lane and look for the smart object tag
-				const FZoneGraphTagMask LaneMask = AnnotationSubsystem->GetAnnotationTags(RequestLaneHandle);
+				const FZoneGraphTagMask LaneMask = AnnotationSubsystem.GetAnnotationTags(RequestLaneHandle);
 				if (!LaneMask.Contains(SmartObjectTag))
 				{
 					continue;
@@ -297,23 +298,20 @@ void UMassSmartObjectCandidatesFinderProcessor::Execute(UMassEntitySubsystem& En
 	// Signal entities that their search results are ready
 	if (EntitiesToSignal.Num())
 	{
-		SignalSubsystem->SignalEntities(UE::Mass::Signals::SmartObjectCandidatesReady, EntitiesToSignal);
+		SignalSubsystem.SignalEntities(UE::Mass::Signals::SmartObjectCandidatesReady, EntitiesToSignal);
 	}
 }
 
 //----------------------------------------------------------------------//
 // UMassSmartObjectTimedBehaviorProcessor
 //----------------------------------------------------------------------//
-void UMassSmartObjectTimedBehaviorProcessor::Initialize(UObject& Owner)
-{
-	SignalSubsystem = UWorld::GetSubsystem<UMassSignalSubsystem>(Owner.GetWorld());
-}
-
 void UMassSmartObjectTimedBehaviorProcessor::ConfigureQueries()
 {
 	EntityQuery.AddRequirement<FMassSmartObjectUserFragment>(EMassFragmentAccess::ReadWrite);
 	EntityQuery.AddRequirement<FMassSmartObjectTimedBehaviorFragment>(EMassFragmentAccess::ReadWrite);
 	EntityQuery.AddSubsystemRequirement<USmartObjectSubsystem>(EMassFragmentAccess::ReadOnly);
+
+	ProcessorRequirements.AddSubsystemRequirement<UMassSignalSubsystem>(EMassFragmentAccess::ReadWrite);
 }
 
 UMassSmartObjectTimedBehaviorProcessor::UMassSmartObjectTimedBehaviorProcessor()
@@ -324,13 +322,12 @@ UMassSmartObjectTimedBehaviorProcessor::UMassSmartObjectTimedBehaviorProcessor()
 
 void UMassSmartObjectTimedBehaviorProcessor::Execute(UMassEntitySubsystem& EntitySubsystem, FMassExecutionContext& Context)
 {
-	checkf(SignalSubsystem != nullptr, TEXT("MassSignalSubsystem should exist when executing processors."));
-
+	UWorld* World = EntitySubsystem.GetWorld();
 	TArray<FMassEntityHandle> ToRelease;
 
 	QUICK_SCOPE_CYCLE_COUNTER(UMassProcessor_SmartObjectTestBehavior_Run);
 
-	EntityQuery.ForEachEntityChunk(EntitySubsystem, Context, [this, &ToRelease, World = EntitySubsystem.GetWorld()](FMassExecutionContext& Context)
+	EntityQuery.ForEachEntityChunk(EntitySubsystem, Context, [this, &ToRelease, World](FMassExecutionContext& Context)
 	{
 		const USmartObjectSubsystem& SmartObjectSubsystem = Context.GetSubsystemChecked<USmartObjectSubsystem>(World);
 
@@ -379,9 +376,11 @@ void UMassSmartObjectTimedBehaviorProcessor::Execute(UMassEntitySubsystem& Entit
 		}
 	});
 
-	for (const FMassEntityHandle EntityToRelease : ToRelease)
+	if (ToRelease.Num())
 	{
-		SignalSubsystem->SignalEntity(UE::Mass::Signals::SmartObjectInteractionDone, EntityToRelease);
+		ConfigureContextForProcessorUse(Context);
+		UMassSignalSubsystem& SignalSubsystem = Context.GetMutableSubsystemChecked<UMassSignalSubsystem>(World);
+		SignalSubsystem.SignalEntities(UE::Mass::Signals::SmartObjectInteractionDone, ToRelease);
 	}
 }
 
diff --git a/Engine/Plugins/Runtime/MassGameplay/Source/MassSmartObjects/Public/MassSmartObjectProcessor.h b/Engine/Plugins/Runtime/MassGameplay/Source/MassSmartObjects/Public/MassSmartObjectProcessor.h
--- a/Engine/Plugins/Runtime/MassGameplay/Source/MassSmartObjects/Public/MassSmartObjectProcessor.h
+++ b/Engine/Plugins/Runtime/MassGameplay/Source/MassSmartObjects/Public/MassSmartObjectProcessor.h
@@ -19,7 +19,6 @@ public:
 	UMassSmartObjectCandidatesFinderProcessor();
 
 protected:
-	virtual void Initialize(UObject& Owner) override;
 	virtual void ConfigureQueries() override;
 	virtual void Execute(UMassEntitySubsystem& EntitySubsystem, FMassExecutionContext& Context) override;
 
@@ -27,12 +26,6 @@ protected:
 	UPROPERTY(EditDefaultsOnly, Category = SmartObject, config)
 	float SearchExtents = 5000.f;
 
-	UPROPERTY(Transient)
-	UMassSignalSubsystem* SignalSubsystem;
-
-	UPROPERTY(Transient)
-	UZoneGraphAnnotationSubsystem* AnnotationSubsystem;
-
 	/** Query to fetch and process requests to find smart objects using spacial query around a given world location. */
 	FMassEntityQuery WorldRequestQuery;
 
@@ -49,13 +42,9 @@ public:
 	UMassSmartObjectTimedBehaviorProcessor();
 
 protected:
-	virtual void Initialize(UObject& Owner) override;
 	virtual void Execute(UMassEntitySubsystem& EntitySubsystem, FMassExecutionContext& Context) override;
 	virtual void ConfigureQueries() override;
 
-	UPROPERTY(Transient)
-	UMassSignalSubsystem* SignalSubsystem;
-
 	FMassEntityQuery EntityQuery;
 };
 
