Separated "entity management" and USubsystem stuff of UMassEntityManager. Now UMassEntityManager hosts a shared ref of newly created FMassEntityManager, and users are free to create their own FMassEnt

diff --git a/Engine/Plugins/Runtime/MassGameplay/Source/MassMovement/Private/MassMovementProcessors.cpp b/Engine/Plugins/Runtime/MassGameplay/Source/MassMovement/Private/MassMovementProcessors.cpp
--- a/Engine/Plugins/Runtime/MassGameplay/Source/MassMovement/Private/MassMovementProcessors.cpp
+++ b/Engine/Plugins/Runtime/MassGameplay/Source/MassMovement/Private/MassMovementProcessors.cpp
@@ -32,7 +32,7 @@ void UMassApplyMovementProcessor::ConfigureQueries()
 	EntityQuery.AddConstSharedRequirement<FMassMovementParameters>(EMassFragmentPresence::All);
 }
 
-void UMassApplyMovementProcessor::Execute(UMassEntitySubsystem& EntitySubsystem,
+void UMassApplyMovementProcessor::Execute(FMassEntityManager& EntityManager,
 													FMassExecutionContext& Context)
 {
 	// Clamp max delta time to avoid force explosion on large time steps (i.e. during initialization).
@@ -40,7 +40,7 @@ void UMassApplyMovementProcessor::Execute(UMassEntitySubsystem& EntitySubsystem,
 
 	QUICK_SCOPE_CYCLE_COUNTER(HighRes);
 
-	EntityQuery.ForEachEntityChunk(EntitySubsystem, Context, [this, DeltaTime](FMassExecutionContext& Context)
+	EntityQuery.ForEachEntityChunk(EntityManager, Context, [this, DeltaTime](FMassExecutionContext& Context)
 	{
 		const int32 NumEntities = Context.GetNumEntities();
 
diff --git a/Engine/Plugins/Runtime/MassGameplay/Source/MassMovement/Private/MassMovementTrait.cpp b/Engine/Plugins/Runtime/MassGameplay/Source/MassMovement/Private/MassMovementTrait.cpp
--- a/Engine/Plugins/Runtime/MassGameplay/Source/MassMovement/Private/MassMovementTrait.cpp
+++ b/Engine/Plugins/Runtime/MassGameplay/Source/MassMovement/Private/MassMovementTrait.cpp
@@ -1,15 +1,16 @@
-﻿// Copyright Epic Games, Inc. All Rights Reserved.
+// Copyright Epic Games, Inc. All Rights Reserved.
 #include "Movement/MassMovementTrait.h"
 #include "MassEntityTemplateRegistry.h"
 #include "MassCommonFragments.h"
 #include "MassMovementFragments.h"
 #include "MassMovementTypes.h"
 #include "Engine/World.h"
+#include "MassEntityUtils.h"
+
 
 void UMassMovementTrait::BuildTemplate(FMassEntityTemplateBuildContext& BuildContext, const UWorld& World) const
 {
-	UMassEntitySubsystem* EntitySubsystem = UWorld::GetSubsystem<UMassEntitySubsystem>(&World);
-	check(EntitySubsystem);
+	FMassEntityManager& EntityManager = UE::Mass::Utils::GetEntityManagerChecked(World);
 
 	BuildContext.RequireFragment<FAgentRadiusFragment>();
 	BuildContext.RequireFragment<FTransformFragment>();
@@ -17,6 +18,6 @@ void UMassMovementTrait::BuildTemplate(FMassEntityTemplateBuildContext& BuildCon
 	BuildContext.AddFragment<FMassVelocityFragment>();
 	BuildContext.AddFragment<FMassForceFragment>();
 
-	const FConstSharedStruct MovementFragment = EntitySubsystem->GetOrCreateConstSharedFragment(UE::StructUtils::GetStructCrc32(FConstStructView::Make(Movement)), Movement);
+	const FConstSharedStruct MovementFragment = EntityManager.GetOrCreateConstSharedFragment(UE::StructUtils::GetStructCrc32(FConstStructView::Make(Movement)), Movement);
 	BuildContext.AddConstSharedFragment(MovementFragment);
 }
diff --git a/Engine/Plugins/Runtime/MassGameplay/Source/MassMovement/Private/MassSimpleMovementTrait.cpp b/Engine/Plugins/Runtime/MassGameplay/Source/MassMovement/Private/MassSimpleMovementTrait.cpp
--- a/Engine/Plugins/Runtime/MassGameplay/Source/MassMovement/Private/MassSimpleMovementTrait.cpp
+++ b/Engine/Plugins/Runtime/MassGameplay/Source/MassMovement/Private/MassSimpleMovementTrait.cpp
@@ -40,9 +40,9 @@ void UMassSimpleMovementProcessor::ConfigureQueries()
 	EntityQuery.SetChunkFilter(&FMassSimulationVariableTickChunkFragment::ShouldTickChunkThisFrame);
 }
 
-void UMassSimpleMovementProcessor::Execute(UMassEntitySubsystem& EntitySubsystem, FMassExecutionContext& Context)
+void UMassSimpleMovementProcessor::Execute(FMassEntityManager& EntityManager, FMassExecutionContext& Context)
 {
-	EntityQuery.ForEachEntityChunk(EntitySubsystem, Context, ([this](FMassExecutionContext& Context)
+	EntityQuery.ForEachEntityChunk(EntityManager, Context, ([this](FMassExecutionContext& Context)
 		{
 			const TConstArrayView<FMassVelocityFragment> VelocitiesList = Context.GetFragmentView<FMassVelocityFragment>();
 			const TArrayView<FTransformFragment> TransformsList = Context.GetMutableFragmentView<FTransformFragment>();
diff --git a/Engine/Plugins/Runtime/MassGameplay/Source/MassMovement/Private/MassVelocityRandomizerTrait.cpp b/Engine/Plugins/Runtime/MassGameplay/Source/MassMovement/Private/MassVelocityRandomizerTrait.cpp
--- a/Engine/Plugins/Runtime/MassGameplay/Source/MassMovement/Private/MassVelocityRandomizerTrait.cpp
+++ b/Engine/Plugins/Runtime/MassGameplay/Source/MassMovement/Private/MassVelocityRandomizerTrait.cpp
@@ -34,10 +34,10 @@ void UMassRandomVelocityInitializer::ConfigureQueries()
 	EntityQuery.AddRequirement<FMassVelocityFragment>(EMassFragmentAccess::ReadWrite);
 }
 
-void UMassRandomVelocityInitializer::Execute(UMassEntitySubsystem& EntitySubsystem, FMassExecutionContext& Context)
+void UMassRandomVelocityInitializer::Execute(FMassEntityManager& EntityManager, FMassExecutionContext& Context)
 {
 	// note: the author is aware that the vectors produced below are not distributed uniformly, but it's good enough
-	EntityQuery.ForEachEntityChunk(EntitySubsystem, Context, ([this](FMassExecutionContext& Context)
+	EntityQuery.ForEachEntityChunk(EntityManager, Context, ([this](FMassExecutionContext& Context)
 		{
 			const TArrayView<FMassVelocityFragment> VelocitiesList = Context.GetMutableFragmentView<FMassVelocityFragment>();
 			for (FMassVelocityFragment& VelocityFragment : VelocitiesList)
diff --git a/Engine/Plugins/Runtime/MassGameplay/Source/MassReplication/Private/MassClientBubbleSerializerBase.cpp b/Engine/Plugins/Runtime/MassGameplay/Source/MassReplication/Private/MassClientBubbleSerializerBase.cpp
--- a/Engine/Plugins/Runtime/MassGameplay/Source/MassReplication/Private/MassClientBubbleSerializerBase.cpp
+++ b/Engine/Plugins/Runtime/MassGameplay/Source/MassReplication/Private/MassClientBubbleSerializerBase.cpp
@@ -2,6 +2,8 @@
 
 #include "MassClientBubbleSerializerBase.h"
 #include "MassClientBubbleHandler.h"
+#include "MassEntityUtils.h"
+
 
 void FMassClientBubbleSerializerBase::InitializeForWorld(UWorld& InWorld)
 {
@@ -12,8 +14,7 @@ void FMassClientBubbleSerializerBase::InitializeForWorld(UWorld& InWorld)
 
 	ReplicationSubsystem = InWorld.GetSubsystem<UMassReplicationSubsystem>();
 
-	EntitySystem = UWorld::GetSubsystem<UMassEntitySubsystem>(World);
-	check(EntitySystem);
+	EntityManager = UE::Mass::Utils::GetEntityManagerChecked(InWorld).AsShared();
 
 	checkf(ClientHandler, TEXT("ClientHandler must be setup! See SetClientHandler()"));
 }
diff --git a/Engine/Plugins/Runtime/MassGameplay/Source/MassReplication/Private/MassReplicationFragments.cpp b/Engine/Plugins/Runtime/MassGameplay/Source/MassReplication/Private/MassReplicationFragments.cpp
--- a/Engine/Plugins/Runtime/MassGameplay/Source/MassReplication/Private/MassReplicationFragments.cpp
+++ b/Engine/Plugins/Runtime/MassGameplay/Source/MassReplication/Private/MassReplicationFragments.cpp
@@ -22,17 +22,17 @@ void UMassNetworkIDFragmentInitializer::ConfigureQueries()
 	EntityQuery.AddSubsystemRequirement<UMassReplicationSubsystem>(EMassFragmentAccess::ReadWrite);
 }
 
-void UMassNetworkIDFragmentInitializer::Execute(UMassEntitySubsystem& EntitySubsystem, FMassExecutionContext& Context)
+void UMassNetworkIDFragmentInitializer::Execute(FMassEntityManager& EntityManager, FMassExecutionContext& Context)
 {
 	QUICK_SCOPE_CYCLE_COUNTER(MassProcessor_InitNetworkID_Run);
 
-	const UWorld* World = EntitySubsystem.GetWorld();
+	const UWorld* World = EntityManager.GetWorld();
 	const ENetMode NetMode = World->GetNetMode();
 
 	if (NetMode != NM_Client)
 	{
 #if UE_REPLICATION_COMPILE_SERVER_CODE
-		EntityQuery.ForEachEntityChunk(EntitySubsystem, Context, [World = EntitySubsystem.GetWorld()](FMassExecutionContext& Context)
+		EntityQuery.ForEachEntityChunk(EntityManager, Context, [World = EntityManager.GetWorld()](FMassExecutionContext& Context)
 			{
 				UMassReplicationSubsystem& ReplicationSubsystem = Context.GetMutableSubsystemChecked<UMassReplicationSubsystem>(World);
 
diff --git a/Engine/Plugins/Runtime/MassGameplay/Source/MassReplication/Private/MassReplicationGridProcessor.cpp b/Engine/Plugins/Runtime/MassGameplay/Source/MassReplication/Private/MassReplicationGridProcessor.cpp
--- a/Engine/Plugins/Runtime/MassGameplay/Source/MassReplication/Private/MassReplicationGridProcessor.cpp
+++ b/Engine/Plugins/Runtime/MassGameplay/Source/MassReplication/Private/MassReplicationGridProcessor.cpp
@@ -44,11 +44,11 @@ void UMassReplicationGridProcessor::ConfigureQueries()
 	RemoveFromGridEntityQuery.AddTagRequirement<FMassInReplicationGridTag>(EMassFragmentPresence::All);
 }
 
-void UMassReplicationGridProcessor::Execute(UMassEntitySubsystem& EntitySubsystem, FMassExecutionContext& Context)
+void UMassReplicationGridProcessor::Execute(FMassEntityManager& EntityManager, FMassExecutionContext& Context)
 {
-	UWorld* World = EntitySubsystem.GetWorld();
+	UWorld* World = EntityManager.GetWorld();
 
-	AddToGridEntityQuery.ForEachEntityChunk(EntitySubsystem, Context, [World](FMassExecutionContext& Context)
+	AddToGridEntityQuery.ForEachEntityChunk(EntityManager, Context, [World](FMassExecutionContext& Context)
 	{
 		UMassReplicationSubsystem& ReplicationSubsystem = Context.GetMutableSubsystemChecked<UMassReplicationSubsystem>(World);
 		FReplicationHashGrid2D& ReplicationGrid = ReplicationSubsystem.GetGridMutable();
@@ -72,7 +72,7 @@ void UMassReplicationGridProcessor::Execute(UMassEntitySubsystem& EntitySubsyste
 		}
 	});
 
-	UpdateGridEntityQuery.ForEachEntityChunk(EntitySubsystem, Context, [World](FMassExecutionContext& Context)
+	UpdateGridEntityQuery.ForEachEntityChunk(EntityManager, Context, [World](FMassExecutionContext& Context)
 	{
 		UMassReplicationSubsystem& ReplicationSubsystem = Context.GetMutableSubsystemChecked<UMassReplicationSubsystem>(World);
 		FReplicationHashGrid2D& ReplicationGrid = ReplicationSubsystem.GetGridMutable();
@@ -103,7 +103,7 @@ void UMassReplicationGridProcessor::Execute(UMassEntitySubsystem& EntitySubsyste
 		}
 	});
 
-	RemoveFromGridEntityQuery.ForEachEntityChunk(EntitySubsystem, Context, [World](FMassExecutionContext& Context)
+	RemoveFromGridEntityQuery.ForEachEntityChunk(EntityManager, Context, [World](FMassExecutionContext& Context)
 	{
 		UMassReplicationSubsystem& ReplicationSubsystem = Context.GetMutableSubsystemChecked<UMassReplicationSubsystem>(World);
 		FReplicationHashGrid2D& ReplicationGrid = ReplicationSubsystem.GetGridMutable();
@@ -138,9 +138,9 @@ void UMassReplicationGridRemoverProcessor::ConfigureQueries()
 	EntityQuery.AddSubsystemRequirement<UMassReplicationSubsystem>(EMassFragmentAccess::ReadWrite);
 }
 
-void UMassReplicationGridRemoverProcessor::Execute(UMassEntitySubsystem& EntitySubsystem, FMassExecutionContext& Context)
+void UMassReplicationGridRemoverProcessor::Execute(FMassEntityManager& EntityManager, FMassExecutionContext& Context)
 {
-	EntityQuery.ForEachEntityChunk(EntitySubsystem, Context, [World = EntitySubsystem.GetWorld()](FMassExecutionContext& Context)
+	EntityQuery.ForEachEntityChunk(EntityManager, Context, [World = EntityManager.GetWorld()](FMassExecutionContext& Context)
 	{
 		UMassReplicationSubsystem& ReplicationSubsystem = Context.GetMutableSubsystemChecked<UMassReplicationSubsystem>(World);
 		FReplicationHashGrid2D& ReplicationGrid = ReplicationSubsystem.GetGridMutable();
diff --git a/Engine/Plugins/Runtime/MassGameplay/Source/MassReplication/Private/MassReplicationProcessor.cpp b/Engine/Plugins/Runtime/MassGameplay/Source/MassReplication/Private/MassReplicationProcessor.cpp
--- a/Engine/Plugins/Runtime/MassGameplay/Source/MassReplication/Private/MassReplicationProcessor.cpp
+++ b/Engine/Plugins/Runtime/MassGameplay/Source/MassReplication/Private/MassReplicationProcessor.cpp
@@ -75,7 +75,7 @@ void UMassReplicationProcessor::Initialize(UObject& Owner)
 #endif //UE_REPLICATION_COMPILE_SERVER_CODE
 }
 
-void UMassReplicationProcessor::PrepareExecution(UMassEntitySubsystem& EntitySubsystem)
+void UMassReplicationProcessor::PrepareExecution(FMassEntityManager& EntityManager)
 {
 #if UE_REPLICATION_COMPILE_SERVER_CODE
 
@@ -84,7 +84,7 @@ void UMassReplicationProcessor::PrepareExecution(UMassEntitySubsystem& EntitySub
 	//first synchronize clients and viewers
 	ReplicationSubsystem->SynchronizeClientsAndViewers();
 
-	EntitySubsystem.ForEachSharedFragment<FMassReplicationSharedFragment>([this](FMassReplicationSharedFragment& RepSharedFragment)
+	EntityManager.ForEachSharedFragment<FMassReplicationSharedFragment>([this](FMassReplicationSharedFragment& RepSharedFragment)
 	{
 		if (!RepSharedFragment.bEntityQueryInitialized)
 		{
@@ -152,19 +152,19 @@ void UMassReplicationProcessor::PrepareExecution(UMassEntitySubsystem& EntitySub
 #endif //UE_REPLICATION_COMPILE_SERVER_CODE
 }
 
-void UMassReplicationProcessor::Execute(UMassEntitySubsystem& EntitySubsystem, FMassExecutionContext& Context)
+void UMassReplicationProcessor::Execute(FMassEntityManager& EntityManager, FMassExecutionContext& Context)
 {
 #if UE_REPLICATION_COMPILE_SERVER_CODE
-	UWorld* World = EntitySubsystem.GetWorld();
+	UWorld* World = EntityManager.GetWorld();
 	check(World);
 	check(ReplicationSubsystem);
 
 	{
 		QUICK_SCOPE_CYCLE_COUNTER(UMassReplicationProcessor_Preperation);
-		PrepareExecution(EntitySubsystem);
+		PrepareExecution(EntityManager);
 	}
 
-	const UMassLODSubsystem& LODSubsystem = Context.GetSubsystemChecked<UMassLODSubsystem>(EntitySubsystem.GetWorld());
+	const UMassLODSubsystem& LODSubsystem = Context.GetSubsystemChecked<UMassLODSubsystem>(EntityManager.GetWorld());
 	const TArray<FViewerInfo>& AllViewersInfo = LODSubsystem.GetViewers();
 	const TArray<FMassClientHandle>& ClientHandles = ReplicationSubsystem->GetClientReplicationHandles();
 	for (const FMassClientHandle ClientHandle : ClientHandles)
@@ -190,7 +190,7 @@ void UMassReplicationProcessor::Execute(UMassEntitySubsystem& EntitySubsystem, F
 		// Prepare LOD collector and calculator
 		// Remember the max LOD distance from each
 		float MaxLODDistance = 0.0f;
-		EntitySubsystem.ForEachSharedFragment<FMassReplicationSharedFragment>([&Viewers,&MaxLODDistance](FMassReplicationSharedFragment& RepSharedFragment)
+		EntityManager.ForEachSharedFragment<FMassReplicationSharedFragment>([&Viewers,&MaxLODDistance](FMassReplicationSharedFragment& RepSharedFragment)
 		{
 			RepSharedFragment.LODCollector.PrepareExecution(Viewers);
 			RepSharedFragment.LODCalculator.PrepareExecution(Viewers);
@@ -206,7 +206,7 @@ void UMassReplicationProcessor::Execute(UMassEntitySubsystem& EntitySubsystem, F
 			ReplicationSubsystem->GetGrid().Query(Bounds, EntitiesInRange);
 		}
 
-		EntityQuery.CacheArchetypes(EntitySubsystem);
+		EntityQuery.CacheArchetypes(EntityManager);
 		if (EntityQuery.GetArchetypes().Num() > 0)
 		{
 			// EntitySet stores array of entities per specified archetype, may contain duplicates.
@@ -228,13 +228,13 @@ void UMassReplicationProcessor::Execute(UMassEntitySubsystem& EntitySubsystem, F
 				Set.Archetype = Archetype;
 			}
 
-			auto BuildEntitySet = [&EntitySets, &EntitySubsystem](const TArray<FMassEntityHandle>& Entities)
+			auto BuildEntitySet = [&EntitySets, &EntityManager](const TArray<FMassEntityHandle>& Entities)
 			{
 				FEntitySet* PrevSet = Entities.Num() ? &EntitySets[0] : nullptr;
 				for (const FMassEntityHandle Entity : Entities)
 				{
 					// Add to set of supported archetypes. Dont process if we don't care about the type.
-					const FMassArchetypeHandle Archetype = EntitySubsystem.GetArchetypeForEntity(Entity);
+					const FMassArchetypeHandle Archetype = EntityManager.GetArchetypeForEntity(Entity);
 					FEntitySet* Set = PrevSet && PrevSet->Archetype == Archetype ? PrevSet : EntitySets.FindByPredicate([&Archetype](const FEntitySet& Set) { return Archetype == Set.Archetype; });
 					if (Set != nullptr)
 					{
@@ -259,7 +259,7 @@ void UMassReplicationProcessor::Execute(UMassEntitySubsystem& EntitySubsystem, F
 
 				{
 					QUICK_SCOPE_CYCLE_COUNTER(UMassReplicationProcessor_SyncToMass);
-					SyncClientData.ForEachEntityChunk(EntitySubsystem, Context, [&ClientReplicationInfo](FMassExecutionContext& Context)
+					SyncClientData.ForEachEntityChunk(EntityManager, Context, [&ClientReplicationInfo](FMassExecutionContext& Context)
 					{
 						const TArrayView<FMassReplicationLODFragment> ViewerLODList = Context.GetMutableFragmentView<FMassReplicationLODFragment>();
 						TArrayView<FMassReplicatedAgentFragment> ReplicatedAgentList = Context.GetMutableFragmentView<FMassReplicatedAgentFragment>();
@@ -287,7 +287,7 @@ void UMassReplicationProcessor::Execute(UMassEntitySubsystem& EntitySubsystem, F
 
 				{
 					QUICK_SCOPE_CYCLE_COUNTER(UMassReplicationProcessor_LODCollection);
-					CollectViewerInfoQuery.ForEachEntityChunk(EntitySubsystem, Context, [](FMassExecutionContext& Context)
+					CollectViewerInfoQuery.ForEachEntityChunk(EntityManager, Context, [](FMassExecutionContext& Context)
 					{
 						const TConstArrayView<FTransformFragment> LocationList = Context.GetFragmentView<FTransformFragment>();
 						const TArrayView<FMassReplicationViewerInfoFragment> ViewersInfoList = Context.GetMutableFragmentView<FMassReplicationViewerInfoFragment>();
@@ -298,7 +298,7 @@ void UMassReplicationProcessor::Execute(UMassEntitySubsystem& EntitySubsystem, F
 
 				{
 					QUICK_SCOPE_CYCLE_COUNTER(UMassReplicationProcessor_LODCaculation);
-					CalculateLODQuery.ForEachEntityChunk(EntitySubsystem, Context, [](FMassExecutionContext& Context)
+					CalculateLODQuery.ForEachEntityChunk(EntityManager, Context, [](FMassExecutionContext& Context)
 					{
 						const TConstArrayView<FMassReplicationViewerInfoFragment> ViewersInfoList = Context.GetFragmentView<FMassReplicationViewerInfoFragment>();
 						const TArrayView<FMassReplicationLODFragment> ViewerLODList = Context.GetMutableFragmentView<FMassReplicationLODFragment>();
@@ -311,7 +311,7 @@ void UMassReplicationProcessor::Execute(UMassEntitySubsystem& EntitySubsystem, F
 
 			{
 				QUICK_SCOPE_CYCLE_COUNTER(UMassReplicationProcessor_LODAdjustDistance);
-				EntitySubsystem.ForEachSharedFragment<FMassReplicationSharedFragment>([](FMassReplicationSharedFragment& RepSharedFragment)
+				EntityManager.ForEachSharedFragment<FMassReplicationSharedFragment>([](FMassReplicationSharedFragment& RepSharedFragment)
 				{
 					RepSharedFragment.bHasAdjustedDistancesFromCount = RepSharedFragment.LODCalculator.AdjustDistancesFromCount();
 				});
@@ -327,7 +327,7 @@ void UMassReplicationProcessor::Execute(UMassEntitySubsystem& EntitySubsystem, F
 
 				{
 					QUICK_SCOPE_CYCLE_COUNTER(UMassReplicationProcessor_LODAdjustLODFromCount);
-					AdjustLODDistancesQuery.ForEachEntityChunk(EntitySubsystem, Context, [](FMassExecutionContext& Context)
+					AdjustLODDistancesQuery.ForEachEntityChunk(EntityManager, Context, [](FMassExecutionContext& Context)
 					{
 						const TConstArrayView<FMassReplicationViewerInfoFragment> ViewersInfoList = Context.GetFragmentView<FMassReplicationViewerInfoFragment>();
 						const TArrayView<FMassReplicationLODFragment> ViewerLODList = Context.GetMutableFragmentView<FMassReplicationLODFragment>();
@@ -339,11 +339,11 @@ void UMassReplicationProcessor::Execute(UMassEntitySubsystem& EntitySubsystem, F
 				{
 					QUICK_SCOPE_CYCLE_COUNTER(UMassReplicationProcessor_ProcessClientReplication);
 					FMassReplicationContext ReplicationContext(*World, LODSubsystem, *ReplicationSubsystem);
-					EntitySubsystem.ForEachSharedFragment<FMassReplicationSharedFragment>([&EntitySubsystem, &Context, &ReplicationContext, &ClientHandle](FMassReplicationSharedFragment& RepSharedFragment)
+					EntityManager.ForEachSharedFragment<FMassReplicationSharedFragment>([&EntityManager, &Context, &ReplicationContext, &ClientHandle](FMassReplicationSharedFragment& RepSharedFragment)
 					{
 						RepSharedFragment.CurrentClientHandle = ClientHandle;
 
-						RepSharedFragment.EntityQuery.ForEachEntityChunk(EntitySubsystem, Context, [&ReplicationContext, &RepSharedFragment](FMassExecutionContext& Context)
+						RepSharedFragment.EntityQuery.ForEachEntityChunk(EntityManager, Context, [&ReplicationContext, &RepSharedFragment](FMassExecutionContext& Context)
 						{
 							RepSharedFragment.CachedReplicator->ProcessClientReplication(Context, ReplicationContext);
 						});
@@ -352,7 +352,7 @@ void UMassReplicationProcessor::Execute(UMassEntitySubsystem& EntitySubsystem, F
 
 				{
 					QUICK_SCOPE_CYCLE_COUNTER(UMassReplicationProcessor_SyncFromMass);
-					SyncClientData.ForEachEntityChunk(EntitySubsystem, Context, [&ClientReplicationInfo](FMassExecutionContext& Context)
+					SyncClientData.ForEachEntityChunk(EntityManager, Context, [&ClientReplicationInfo](FMassExecutionContext& Context)
 					{
 						TArrayView<FMassReplicatedAgentFragment> ReplicatedAgentList = Context.GetMutableFragmentView<FMassReplicatedAgentFragment>();
 
@@ -369,9 +369,9 @@ void UMassReplicationProcessor::Execute(UMassEntitySubsystem& EntitySubsystem, F
 				// Optional debug display
 				if (UE::Mass::Replication::DebugClientReplicationLOD == ClientHandle.GetIndex())
 				{
-					EntitySubsystem.ForEachSharedFragment<FMassReplicationSharedFragment>([World, &EntitySubsystem, &Context](FMassReplicationSharedFragment& RepSharedFragment)
+					EntityManager.ForEachSharedFragment<FMassReplicationSharedFragment>([World, &EntityManager, &Context](FMassReplicationSharedFragment& RepSharedFragment)
 					{
-						RepSharedFragment.EntityQuery.ForEachEntityChunk(EntitySubsystem, Context, [World, &RepSharedFragment](FMassExecutionContext& Context)
+						RepSharedFragment.EntityQuery.ForEachEntityChunk(EntityManager, Context, [World, &RepSharedFragment](FMassExecutionContext& Context)
 						{
 							const TConstArrayView<FTransformFragment> TransformList = Context.GetFragmentView<FTransformFragment>();
 							const TConstArrayView<FMassReplicationLODFragment> ViewerLODList = Context.GetFragmentView<FMassReplicationLODFragment>();
diff --git a/Engine/Plugins/Runtime/MassGameplay/Source/MassReplication/Private/MassReplicationSubsystem.cpp b/Engine/Plugins/Runtime/MassGameplay/Source/MassReplication/Private/MassReplicationSubsystem.cpp
--- a/Engine/Plugins/Runtime/MassGameplay/Source/MassReplication/Private/MassReplicationSubsystem.cpp
+++ b/Engine/Plugins/Runtime/MassGameplay/Source/MassReplication/Private/MassReplicationSubsystem.cpp
@@ -9,6 +9,7 @@
 #include "MassClientBubbleHandler.h"
 #include "MassClientBubbleInfoBase.h"
 #include "MassReplicationSettings.h"
+#include "MassEntityUtils.h"
 
 
 uint32 UMassReplicationSubsystem::CurrentNetMassCounter = 0;
@@ -18,14 +19,12 @@ void UMassReplicationSubsystem::Initialize(FSubsystemCollectionBase& Collection)
 	Super::Initialize(Collection);
 
 	World = GetWorld();
-
 	check(World);
 
 	MassLODSubsystem = Collection.InitializeDependency<UMassLODSubsystem>();
-
 	check(MassLODSubsystem);
 
-	EntitySystem = UWorld::GetSubsystem<UMassEntitySubsystem>(GetWorld());
+	EntityManager = UE::Mass::Utils::GetEntityManagerChecked(*World).AsShared();
 }
 
 void UMassReplicationSubsystem::Deinitialize()
@@ -48,7 +47,7 @@ void UMassReplicationSubsystem::Deinitialize()
 
 	World = nullptr;
 	MassLODSubsystem = nullptr;
-	EntitySystem = nullptr;
+	EntityManager.Reset();
 }
 
 UMassReplicationSubsystem::UMassReplicationSubsystem()
@@ -58,8 +57,8 @@ UMassReplicationSubsystem::UMassReplicationSubsystem()
 
 FMassNetworkID UMassReplicationSubsystem::GetNetIDFromHandle(const FMassEntityHandle Handle) const
 {
-	check(EntitySystem);
-	const FMassNetworkIDFragment& Data = EntitySystem->GetFragmentDataChecked<FMassNetworkIDFragment>(Handle);
+	check(EntityManager);
+	const FMassNetworkIDFragment& Data = EntityManager->GetFragmentDataChecked<FMassNetworkIDFragment>(Handle);
 	return Data.NetID;
 }
 
diff --git a/Engine/Plugins/Runtime/MassGameplay/Source/MassReplication/Private/MassReplicationTrait.cpp b/Engine/Plugins/Runtime/MassGameplay/Source/MassReplication/Private/MassReplicationTrait.cpp
--- a/Engine/Plugins/Runtime/MassGameplay/Source/MassReplication/Private/MassReplicationTrait.cpp
+++ b/Engine/Plugins/Runtime/MassGameplay/Source/MassReplication/Private/MassReplicationTrait.cpp
@@ -7,6 +7,7 @@
 #include "MassSimulationLOD.h"
 #include "MassReplicationTypes.h"
 #include "MassReplicationSubsystem.h"
+#include "MassEntityUtils.h"
 
 
 void UMassReplicationTrait::BuildTemplate(FMassEntityTemplateBuildContext& BuildContext, const UWorld& World) const
@@ -25,15 +26,14 @@ void UMassReplicationTrait::BuildTemplate(FMassEntityTemplateBuildContext& Build
 	BuildContext.AddFragment<FMassReplicationLODFragment>();
 	BuildContext.AddFragment<FMassReplicationGridCellLocationFragment>();
 
-	UMassEntitySubsystem* EntitySubsystem = UWorld::GetSubsystem<UMassEntitySubsystem>(&World);
-	check(EntitySubsystem);
+	FMassEntityManager& EntityManager = UE::Mass::Utils::GetEntityManagerChecked(World);
 
 	UMassReplicationSubsystem* ReplicationSubsystem = UWorld::GetSubsystem<UMassReplicationSubsystem>(&World);
 	check(ReplicationSubsystem);
 
 	uint32 ParamsHash = UE::StructUtils::GetStructCrc32(FConstStructView::Make(Params));
-	FConstSharedStruct ParamsFragment = EntitySubsystem->GetOrCreateConstSharedFragment(ParamsHash, Params);
+	FConstSharedStruct ParamsFragment = EntityManager.GetOrCreateConstSharedFragment(ParamsHash, Params);
 	BuildContext.AddConstSharedFragment(ParamsFragment);
-	FSharedStruct SharedFragment = EntitySubsystem->GetOrCreateSharedFragment<FMassReplicationSharedFragment>(ParamsHash, *ReplicationSubsystem, Params);
+	FSharedStruct SharedFragment = EntityManager.GetOrCreateSharedFragment<FMassReplicationSharedFragment>(ParamsHash, *ReplicationSubsystem, Params);
 	BuildContext.AddSharedFragment(SharedFragment);
 }
\ No newline at end of file
diff --git a/Engine/Plugins/Runtime/MassGameplay/Source/MassReplication/Public/MassClientBubbleHandler.h b/Engine/Plugins/Runtime/MassGameplay/Source/MassReplication/Public/MassClientBubbleHandler.h
--- a/Engine/Plugins/Runtime/MassGameplay/Source/MassReplication/Public/MassClientBubbleHandler.h
+++ b/Engine/Plugins/Runtime/MassGameplay/Source/MassReplication/Public/MassClientBubbleHandler.h
@@ -18,7 +18,7 @@
 #include "MassClientBubbleHandler.generated.h"
 
 class UWorld;
-class UMassEntitySubsystem;
+struct FMassEntityManager;
 
 namespace UE::Mass::Replication
 {
@@ -428,8 +428,7 @@ void TClientBubbleHandlerBase<AgentArrayItem>::PostReplicatedAddHelper(const TAr
 {
 	TArray<FMassEntityHandle> EntitiesDestroy;
 
-	UMassEntitySubsystem* EntitySystem = Serializer->GetEntitySystem();
-	check(EntitySystem);
+	FMassEntityManager& EntityManager = Serializer->GetEntityManagerChecked();
 
 	UMassReplicationSubsystem* ReplicationSubsystem = Serializer->GetReplicationSubsystem();
 	check(ReplicationSubsystem);
@@ -455,7 +454,7 @@ void TClientBubbleHandlerBase<AgentArrayItem>::PostReplicatedAddHelper(const TAr
 			// If EntityData IsSet() it means we have had multiple Adds without a remove and we treat this add as modifying an existing agent.
 			if (EntityInfo->Entity.IsSet())
 			{
-				FMassEntityView EntityView(*EntitySystem, EntityInfo->Entity);
+				FMassEntityView EntityView(EntityManager, EntityInfo->Entity);
 
 				SetModifiedEntityData(EntityView, AddedItem.Agent);
 			}
@@ -507,8 +506,7 @@ void TClientBubbleHandlerBase<AgentArrayItem>::PostReplicatedAddEntitiesHelper(c
 {
 	check(Serializer);
 
-	UMassEntitySubsystem* EntitySystem = Serializer->GetEntitySystem();
-	check(EntitySystem);
+	FMassEntityManager& EntityManager = Serializer->GetEntityManagerChecked();
 
 	UMassSpawnerSubsystem* SpawnerSubsystem = Serializer->GetSpawnerSubsystem();
 	check(SpawnerSubsystem);
@@ -552,7 +550,7 @@ void TClientBubbleHandlerBase<AgentArrayItem>::PostReplicatedAddEntitiesHelper(c
 		int32 AgentsSpawnIdx = 0;
 
 		Query.ForEachEntityChunk(FMassArchetypeEntityCollection(ArchetypeHandle, Entities, FMassArchetypeEntityCollection::NoDuplicates)
-								, *EntitySystem, ExecContext, [&AgentsSpawn, &AgentsSpawnIdx, this, ReplicationSubsystem, &ExecContext, &CacheFragmentViewsForSpawnQuery, &SetSpawnedEntityData, &EntitySystem](FMassExecutionContext& Context)
+								, EntityManager, ExecContext, [&AgentsSpawn, &AgentsSpawnIdx, this, ReplicationSubsystem, &ExecContext, &CacheFragmentViewsForSpawnQuery, &SetSpawnedEntityData, &EntityManager](FMassExecutionContext& Context)
 			{
 				CacheFragmentViewsForSpawnQuery(ExecContext);
 
@@ -568,7 +566,7 @@ void TClientBubbleHandlerBase<AgentArrayItem>::PostReplicatedAddEntitiesHelper(c
 					NetIDFragment.NetID = AgentSpawn.GetNetID();
 					ReplicationSubsystem->SetEntity(NetIDFragment.NetID, Entity);
 
-					FMassEntityView EntityView(*EntitySystem, Entity);
+					FMassEntityView EntityView(EntityManager, Entity);
 					SetSpawnedEntityData(EntityView, AgentSpawn, i);
 
 					++AgentsSpawnIdx;
@@ -582,8 +580,7 @@ void TClientBubbleHandlerBase<AgentArrayItem>::PostReplicatedAddEntitiesHelper(c
 template<typename AgentArrayItem>
 void TClientBubbleHandlerBase<AgentArrayItem>::PostReplicatedChangeHelper(const TArrayView<int32> ChangedIndices, FSetModifiedEntityDataFunction SetModifiedEntityData)
 {
-	UMassEntitySubsystem* EntitySystem = Serializer->GetEntitySystem();
-	check(EntitySystem);
+	FMassEntityManager& EntityManager = Serializer->GetEntityManagerChecked();
 
 	UMassReplicationSubsystem* ReplicationSubsystem = Serializer->GetReplicationSubsystem();
 	check(ReplicationSubsystem);
@@ -601,7 +598,7 @@ void TClientBubbleHandlerBase<AgentArrayItem>::PostReplicatedChangeHelper(const
 		// Currently we don't think this should be needed, but are leaving it in for bomb proofing.
 		if (ensure(EntityInfo->ReplicationID == ChangedItem.ReplicationID))
 		{
-			FMassEntityView EntityView(*EntitySystem, EntityInfo->Entity);
+			FMassEntityView EntityView(EntityManager, EntityInfo->Entity);
 			SetModifiedEntityData(EntityView, ChangedItem.Agent);
 		}
 	}
@@ -650,8 +647,7 @@ template<typename AgentArrayItem>
 void TClientBubbleHandlerBase<AgentArrayItem>::DebugValidateBubbleOnClient()
 {
 #if UE_ALLOW_DEBUG_REPLICATION
-	UMassEntitySubsystem* EntitySystem = Serializer->GetEntitySystem();
-	check(EntitySystem);
+	const FMassEntityManager& EntityManager = Serializer->GetEntityManagerChecked();
 
 	UMassReplicationSubsystem* ReplicationSubsystem = Serializer->GetReplicationSubsystem();
 	check(ReplicationSubsystem);
@@ -671,12 +667,12 @@ void TClientBubbleHandlerBase<AgentArrayItem>::DebugValidateBubbleOnClient()
 		{
 			if (EntityInfo->ReplicationID == Item.ReplicationID)
 			{
-				const bool bIsEntityValid = EntitySystem->IsEntityValid(EntityInfo->Entity);
+				const bool bIsEntityValid = EntityManager.IsEntityValid(EntityInfo->Entity);
 				checkf(bIsEntityValid, TEXT("Must be valid entity if at latest ReplciationID"));
 
 				if (bIsEntityValid)
 				{
-					const FMassNetworkIDFragment& FragmentNetID = EntitySystem->GetFragmentDataChecked<FMassNetworkIDFragment>(EntityInfo->Entity);
+					const FMassNetworkIDFragment& FragmentNetID = EntityManager.GetFragmentDataChecked<FMassNetworkIDFragment>(EntityInfo->Entity);
 
 					checkf(FragmentNetID.NetID == Agent.GetNetID(), TEXT("Fragment and Agent NetID do not match!"));
 				}
@@ -713,8 +709,7 @@ void TClientBubbleHandlerBase<AgentArrayItem>::DebugValidateBubbleOnServer()
 #if UE_ALLOW_DEBUG_REPLICATION
 	using namespace UE::Mass::Replication;
 
-	UMassEntitySubsystem* EntitySystem = Serializer->GetEntitySystem();
-	check(EntitySystem);
+	const FMassEntityManager& EntityManager = Serializer->GetEntityManagerChecked();
 
 	for (int32 OuterIdx = 0; OuterIdx < (*Agents).Num(); ++OuterIdx)
 	{
@@ -737,14 +732,14 @@ void TClientBubbleHandlerBase<AgentArrayItem>::DebugValidateBubbleOnServer()
 		const FMassAgentLookupData& LookupData = AgentLookupArray[OuterItem.GetHandle().GetIndex()];
 
 		checkf(LookupData.AgentsIdx == OuterIdx, TEXT("Agent index must match lookup data!"));
-		checkf(EntitySystem->IsEntityValid(LookupData.Entity), TEXT("Must be valid entity"));
+		checkf(EntityManager.IsEntityValid(LookupData.Entity), TEXT("Must be valid entity"));
 
-		const FMassNetworkIDFragment& FragmentNetID = EntitySystem->GetFragmentDataChecked<FMassNetworkIDFragment>(LookupData.Entity);
+		const FMassNetworkIDFragment& FragmentNetID = EntityManager.GetFragmentDataChecked<FMassNetworkIDFragment>(LookupData.Entity);
 
 		checkf(FragmentNetID.NetID == OuterItem.Agent.GetNetID(), TEXT("Fragment and Agent NetID do not match!"));
 		checkf(LookupData.NetID == OuterItem.Agent.GetNetID(), TEXT("LookupData and Agent NetID do not match!"));
 
-		const FReplicationTemplateIDFragment& FragmentTemplateID = EntitySystem->GetFragmentDataChecked<FReplicationTemplateIDFragment>(LookupData.Entity);
+		const FReplicationTemplateIDFragment& FragmentTemplateID = EntityManager.GetFragmentDataChecked<FReplicationTemplateIDFragment>(LookupData.Entity);
 		checkf(FragmentTemplateID.ID == OuterItem.Agent.GetTemplateID(), TEXT("Agent TemplateID different to Fragment!"));
 	}
 
diff --git a/Engine/Plugins/Runtime/MassGameplay/Source/MassReplication/Public/MassClientBubbleSerializerBase.h b/Engine/Plugins/Runtime/MassGameplay/Source/MassReplication/Public/MassClientBubbleSerializerBase.h
--- a/Engine/Plugins/Runtime/MassGameplay/Source/MassReplication/Public/MassClientBubbleSerializerBase.h
+++ b/Engine/Plugins/Runtime/MassGameplay/Source/MassReplication/Public/MassClientBubbleSerializerBase.h
@@ -10,7 +10,7 @@
 
 class UWorld;
 class UMassSpawnerSubsystem;
-class UMassEntitySubsystem;
+struct FMassEntityManager;
 class IClientBubbleHandlerInterface;
 
 /**
@@ -35,7 +35,7 @@ public:
 	UWorld* GetWorld() const { return World; }
 	UMassSpawnerSubsystem* GetSpawnerSubsystem() const { return SpawnerSubsystem; }
 	UMassReplicationSubsystem* GetReplicationSubsystem() const { return ReplicationSubsystem; }
-	UMassEntitySubsystem* GetEntitySystem() const { return EntitySystem; }
+	FMassEntityManager& GetEntityManagerChecked() const { check(EntityManager); return *EntityManager.Get(); }
 
 	void SetClientHandler(IClientBubbleHandlerInterface& InClientHandler) { ClientHandler = &InClientHandler; }
 	IClientBubbleHandlerInterface* GetClientHandler() const { return ClientHandler; }
@@ -47,8 +47,7 @@ private:
 	UPROPERTY(Transient)
 	UMassSpawnerSubsystem* SpawnerSubsystem = nullptr;
 
-	UPROPERTY(Transient)
-	UMassEntitySubsystem* EntitySystem = nullptr;
+	TSharedPtr<FMassEntityManager> EntityManager;
 
 	UPROPERTY(Transient)
 	UMassReplicationSubsystem* ReplicationSubsystem = nullptr;
diff --git a/Engine/Plugins/Runtime/MassGameplay/Source/MassReplication/Public/MassReplicationProcessor.h b/Engine/Plugins/Runtime/MassGameplay/Source/MassReplication/Public/MassReplicationProcessor.h
--- a/Engine/Plugins/Runtime/MassGameplay/Source/MassReplication/Public/MassReplicationProcessor.h
+++ b/Engine/Plugins/Runtime/MassGameplay/Source/MassReplication/Public/MassReplicationProcessor.h
@@ -32,9 +32,9 @@ public:
 protected:
 	virtual void ConfigureQueries() override;
 	virtual void Initialize(UObject& Owner) override;
-	virtual void Execute(UMassEntitySubsystem& EntitySubsystem, FMassExecutionContext& Context) override;
+	virtual void Execute(FMassEntityManager& EntityManager, FMassExecutionContext& Context) override;
 
-	void PrepareExecution(UMassEntitySubsystem& EntitySubsystem);
+	void PrepareExecution(FMassEntityManager& EntityManager);
 
 protected:
 
diff --git a/Engine/Plugins/Runtime/MassGameplay/Source/MassReplication/Public/MassReplicationSubsystem.h b/Engine/Plugins/Runtime/MassGameplay/Source/MassReplication/Public/MassReplicationSubsystem.h
--- a/Engine/Plugins/Runtime/MassGameplay/Source/MassReplication/Public/MassReplicationSubsystem.h
+++ b/Engine/Plugins/Runtime/MassGameplay/Source/MassReplication/Public/MassReplicationSubsystem.h
@@ -9,7 +9,7 @@
 
 #include "MassReplicationSubsystem.generated.h"
 
-class UMassEntitySubsystem;
+struct FMassEntityManager;
 class AMassCrowdClientBubbleInfo;
 class UMassLODSubsystem;
 class AMassClientBubbleInfoBase;
@@ -283,13 +283,12 @@ protected:
 	bool SynchronizeClients(const TArray<FViewerInfo>& Viewers);
 	void SynchronizeClientViewers(const TArray<FViewerInfo>& Viewers);
 
-	static void DebugCheckArraysAreInSync();
+	void DebugCheckArraysAreInSync();
 
 protected:
 	static uint32 CurrentNetMassCounter;
 
-	UPROPERTY()
-	UMassEntitySubsystem* EntitySystem;
+	TSharedPtr<FMassEntityManager> EntityManager;
 
 	/** Clients free list FMassClientHandle manager, handles will to the indices of FMassClientReplicationData::ClientBubbles */
 	FMassClientHandleManager ClientHandleManager;
diff --git a/Engine/Plugins/Runtime/MassGameplay/Source/MassRepresentation/Private/MassRepresentationActorManagement.cpp b/Engine/Plugins/Runtime/MassGameplay/Source/MassRepresentation/Private/MassRepresentationActorManagement.cpp
--- a/Engine/Plugins/Runtime/MassGameplay/Source/MassRepresentation/Private/MassRepresentationActorManagement.cpp
+++ b/Engine/Plugins/Runtime/MassGameplay/Source/MassRepresentation/Private/MassRepresentationActorManagement.cpp
@@ -13,11 +13,11 @@ float UMassRepresentationActorManagement::GetSpawnPriority(const FMassRepresenta
 	return Representation.LODSignificance - (Representation.Visibility == EMassVisibility::CanBeSeen ? 1.0f : 0.0f);
 }
 
-AActor* UMassRepresentationActorManagement::GetOrSpawnActor(UMassRepresentationSubsystem& RepresentationSubsystem, UMassEntitySubsystem& EntitySubsystem, const FMassEntityHandle MassAgent, FMassActorFragment& ActorInfo, const FTransform& Transform, const int16 TemplateActorIndex, FMassActorSpawnRequestHandle& SpawnRequestHandle, const float Priority) const
+AActor* UMassRepresentationActorManagement::GetOrSpawnActor(UMassRepresentationSubsystem& RepresentationSubsystem, FMassEntityManager& EntityManager, const FMassEntityHandle MassAgent, FMassActorFragment& ActorInfo, const FTransform& Transform, const int16 TemplateActorIndex, FMassActorSpawnRequestHandle& SpawnRequestHandle, const float Priority) const
 {
 	return RepresentationSubsystem.GetOrSpawnActorFromTemplate(MassAgent, Transform, TemplateActorIndex, SpawnRequestHandle, Priority,
-		FMassActorPreSpawnDelegate::CreateUObject(this, &UMassRepresentationActorManagement::OnPreActorSpawn, &EntitySubsystem),
-		FMassActorPostSpawnDelegate::CreateUObject(this, &UMassRepresentationActorManagement::OnPostActorSpawn, &EntitySubsystem));
+		FMassActorPreSpawnDelegate::CreateUObject(this, &UMassRepresentationActorManagement::OnPreActorSpawn, &EntityManager),
+		FMassActorPostSpawnDelegate::CreateUObject(this, &UMassRepresentationActorManagement::OnPostActorSpawn, &EntityManager));
 }
 
 
@@ -31,7 +31,7 @@ void UMassRepresentationActorManagement::SetActorEnabled(const EMassActorEnabled
 	if (Actor.GetActorEnableCollision() != bEnabled)
 	{
 		// Deferring this as there is a callback internally that could end up doing things outside of the game thread and will fire checks(Chaos mostly)
-		CommandBuffer.PushCommand<FMassDeferredSetCommand>([&Actor, bEnabled](UMassEntitySubsystem&)
+		CommandBuffer.PushCommand<FMassDeferredSetCommand>([&Actor, bEnabled](FMassEntityManager&)
 		{
 			Actor.SetActorEnableCollision(bEnabled);
 		});
@@ -42,19 +42,19 @@ void UMassRepresentationActorManagement::TeleportActor(const FTransform& Transfo
 {
 	if (!Actor.GetTransform().Equals(Transform))
 	{
-		CommandBuffer.PushCommand<FMassDeferredSetCommand>([&Actor, Transform](UMassEntitySubsystem&)
+		CommandBuffer.PushCommand<FMassDeferredSetCommand>([&Actor, Transform](FMassEntityManager&)
 		{
 			Actor.SetActorTransform(Transform, /*bSweep*/false, /*OutSweepHitResult*/nullptr, ETeleportType::TeleportPhysics);
 		});
 	}
 }
 
-void UMassRepresentationActorManagement::OnPreActorSpawn(const FMassActorSpawnRequestHandle& SpawnRequestHandle, FConstStructView SpawnRequest, UMassEntitySubsystem* EntitySubsystem) const
+void UMassRepresentationActorManagement::OnPreActorSpawn(const FMassActorSpawnRequestHandle& SpawnRequestHandle, FConstStructView SpawnRequest, FMassEntityManager* EntityManager) const
 {
-	check(EntitySubsystem);
+	check(EntityManager);
 
 	const FMassActorSpawnRequest& MassActorSpawnRequest = SpawnRequest.Get<FMassActorSpawnRequest>();
-	const FMassEntityView EntityView(*EntitySubsystem, MassActorSpawnRequest.MassAgent);
+	const FMassEntityView EntityView(*EntityManager, MassActorSpawnRequest.MassAgent);
 	FMassActorFragment& ActorInfo = EntityView.GetFragmentData<FMassActorFragment>();
 	FMassRepresentationFragment& Representation = EntityView.GetFragmentData<FMassRepresentationFragment>();
 	UMassRepresentationSubsystem* RepresentationSubsystem = EntityView.GetSharedFragmentData<FMassRepresentationSubsystemSharedFragment>().RepresentationSubsystem;
@@ -80,15 +80,15 @@ void UMassRepresentationActorManagement::OnPreActorSpawn(const FMassActorSpawnRe
 	}
 }
 
-EMassActorSpawnRequestAction UMassRepresentationActorManagement::OnPostActorSpawn(const FMassActorSpawnRequestHandle& SpawnRequestHandle, FConstStructView SpawnRequest, UMassEntitySubsystem* EntitySubsystem) const
+EMassActorSpawnRequestAction UMassRepresentationActorManagement::OnPostActorSpawn(const FMassActorSpawnRequestHandle& SpawnRequestHandle, FConstStructView SpawnRequest, FMassEntityManager* EntityManager) const
 {
-	check(EntitySubsystem);
+	check(EntityManager);
 
 	const FMassActorSpawnRequest& MassActorSpawnRequest = SpawnRequest.Get<FMassActorSpawnRequest>();
 	checkf(MassActorSpawnRequest.SpawnedActor, TEXT("Expecting valid spawned actor"));
 
 	// Might be already done if the actor has a MassAgentComponent via the callback OnMassAgentComponentEntityAssociated on the MassRepresentationSubsystem
-	FMassActorFragment& ActorInfo = EntitySubsystem->GetFragmentDataChecked<FMassActorFragment>(MassActorSpawnRequest.MassAgent);
+	FMassActorFragment& ActorInfo = EntityManager->GetFragmentDataChecked<FMassActorFragment>(MassActorSpawnRequest.MassAgent);
 	if (ActorInfo.IsValid())
 	{
 		// If already set, make sure it is pointing to the same actor.
@@ -102,9 +102,9 @@ EMassActorSpawnRequestAction UMassRepresentationActorManagement::OnPostActorSpaw
 	return EMassActorSpawnRequestAction::Keep;
 }
 
-void UMassRepresentationActorManagement::ReleaseAnyActorOrCancelAnySpawning(UMassEntitySubsystem& EntitySubsystem, const FMassEntityHandle MassAgent)
+void UMassRepresentationActorManagement::ReleaseAnyActorOrCancelAnySpawning(FMassEntityManager& EntityManager, const FMassEntityHandle MassAgent)
 {
-	FMassEntityView EntityView(EntitySubsystem, MassAgent);
+	FMassEntityView EntityView(EntityManager, MassAgent);
 	FMassActorFragment& ActorInfo = EntityView.GetFragmentData<FMassActorFragment>();
 	FMassRepresentationFragment& Representation = EntityView.GetFragmentData<FMassRepresentationFragment>();
 	UMassRepresentationSubsystem* RepresentationSubsystem = EntityView.GetSharedFragmentData<FMassRepresentationSubsystemSharedFragment>().RepresentationSubsystem;
diff --git a/Engine/Plugins/Runtime/MassGameplay/Source/MassRepresentation/Private/MassRepresentationProcessor.cpp b/Engine/Plugins/Runtime/MassGameplay/Source/MassRepresentation/Private/MassRepresentationProcessor.cpp
--- a/Engine/Plugins/Runtime/MassGameplay/Source/MassRepresentation/Private/MassRepresentationProcessor.cpp
+++ b/Engine/Plugins/Runtime/MassGameplay/Source/MassRepresentation/Private/MassRepresentationProcessor.cpp
@@ -10,6 +10,8 @@
 #include "MassEntityView.h"
 #include "Engine/World.h"
 #include "MassRepresentationActorManagement.h"
+#include "MassEntityUtils.h"
+
 
 namespace UE::MassRepresentation
 {
@@ -44,7 +46,8 @@ void UMassRepresentationProcessor::Initialize(UObject& Owner)
 	Super::Initialize(Owner);
 
 	World = Owner.GetWorld();
-	CachedEntitySubsystem = UWorld::GetSubsystem<UMassEntitySubsystem>(World);
+	check(World);
+	EntityManager = UE::Mass::Utils::GetEntityManagerChecked(*World).AsShared();
 }
 
 void UMassRepresentationProcessor::UpdateRepresentation(FMassExecutionContext& Context)
@@ -57,6 +60,9 @@ void UMassRepresentationProcessor::UpdateRepresentation(FMassExecutionContext& C
 
 	UMassActorSubsystem* MassActorSubsystem = Context.GetMutableSubsystem<UMassActorSubsystem>(RepresentationSubsystem->GetWorld());
 
+	check(EntityManager);
+	FMassEntityManager& CachedEntityManager = *EntityManager.Get();
+
 	const TConstArrayView<FTransformFragment> TransformList = Context.GetFragmentView<FTransformFragment>();
 	const TArrayView<FMassRepresentationFragment> RepresentationList = Context.GetMutableFragmentView<FMassRepresentationFragment>();
 	const TConstArrayView<FMassRepresentationLODFragment> RepresentationLODList = Context.GetFragmentView<FMassRepresentationLODFragment>();
@@ -108,7 +114,7 @@ void UMassRepresentationProcessor::UpdateRepresentation(FMassExecutionContext& C
 				// If we already queued spawn request but have changed our mind, continue with it but once we get the actor back, disable it immediately
 				if (Representation.ActorSpawnRequestHandle.IsValid())
 				{
-					Actor = RepresentationActorManagement->GetOrSpawnActor(*RepresentationSubsystem, *CachedEntitySubsystem, MassAgent, ActorInfo, TransformFragment.GetTransform(), Representation.LowResTemplateActorIndex, Representation.ActorSpawnRequestHandle, RepresentationActorManagement->GetSpawnPriority(RepresentationLOD));
+					Actor = RepresentationActorManagement->GetOrSpawnActor(*RepresentationSubsystem, CachedEntityManager, MassAgent, ActorInfo, TransformFragment.GetTransform(), Representation.LowResTemplateActorIndex, Representation.ActorSpawnRequestHandle, RepresentationActorManagement->GetSpawnPriority(RepresentationLOD));
 				}
 			}
 			if (Actor != nullptr)
@@ -155,7 +161,7 @@ void UMassRepresentationProcessor::UpdateRepresentation(FMassExecutionContext& C
 							!RepresentationSubsystem->DoesActorMatchTemplate(*Actor, WantedTemplateActorIndex) ||
 							Representation.ActorSpawnRequestHandle.IsValid())
 						{
-							NewActor = RepresentationActorManagement->GetOrSpawnActor(*RepresentationSubsystem, *CachedEntitySubsystem, MassAgent, ActorInfo, TransformFragment.GetTransform(), WantedTemplateActorIndex, Representation.ActorSpawnRequestHandle, RepresentationActorManagement->GetSpawnPriority(RepresentationLOD));
+							NewActor = RepresentationActorManagement->GetOrSpawnActor(*RepresentationSubsystem, CachedEntityManager, MassAgent, ActorInfo, TransformFragment.GetTransform(), WantedTemplateActorIndex, Representation.ActorSpawnRequestHandle, RepresentationActorManagement->GetSpawnPriority(RepresentationLOD));
 						}
 						else
 						{
@@ -222,7 +228,7 @@ void UMassRepresentationProcessor::UpdateRepresentation(FMassExecutionContext& C
 	}
 }
 
-void UMassRepresentationProcessor::Execute(UMassEntitySubsystem& InEntitySubsystem, FMassExecutionContext& Context)
+void UMassRepresentationProcessor::Execute(FMassEntityManager& InEntitySubsystem, FMassExecutionContext& Context)
 {
 	// Update entities representation
 	EntityQuery.ForEachEntityChunk(InEntitySubsystem, Context, [this](FMassExecutionContext& Context)
@@ -255,7 +261,7 @@ bool UMassRepresentationProcessor::ReleaseActorOrCancelSpawning(UMassRepresentat
 			TObjectKey<const AActor> ActorKey(Actor); 
 			if (MassActorSubsystem)
 			{
-				CommandBuffer.PushCommand<FMassDeferredSetCommand>([MassActorSubsystem, ActorKey](UMassEntitySubsystem&)
+				CommandBuffer.PushCommand<FMassDeferredSetCommand>([MassActorSubsystem, ActorKey](FMassEntityManager&)
 				{
 					MassActorSubsystem->RemoveHandleForActor(ActorKey);
 				});
@@ -355,10 +361,10 @@ void UMassVisualizationProcessor::UpdateEntityVisibility(const FMassEntityHandle
 	}
 }
 
-void UMassVisualizationProcessor::Execute(UMassEntitySubsystem& EntitySubsystem, FMassExecutionContext& Context)
+void UMassVisualizationProcessor::Execute(FMassEntityManager& InEntityManager, FMassExecutionContext& Context)
 {
 	// Update entities visualization
-	EntityQuery.ForEachEntityChunk(EntitySubsystem, Context, [this](FMassExecutionContext& Context)
+	EntityQuery.ForEachEntityChunk(InEntityManager, Context, [this](FMassExecutionContext& Context)
 	{
 		UpdateVisualization(Context);
 	});
@@ -384,9 +390,9 @@ void UMassRepresentationFragmentDestructor::ConfigureQueries()
 	EntityQuery.AddSharedRequirement<FMassRepresentationSubsystemSharedFragment>(EMassFragmentAccess::ReadWrite);
 }
 
-void UMassRepresentationFragmentDestructor::Execute(UMassEntitySubsystem& EntitySubsystem, FMassExecutionContext& Context)
+void UMassRepresentationFragmentDestructor::Execute(FMassEntityManager& EntityManager, FMassExecutionContext& Context)
 {
-	EntityQuery.ForEachEntityChunk(EntitySubsystem, Context, [this](FMassExecutionContext& Context)
+	EntityQuery.ForEachEntityChunk(EntityManager, Context, [this](FMassExecutionContext& Context)
 	{
 		UMassRepresentationSubsystem* RepresentationSubsystem = Context.GetMutableSharedFragment<FMassRepresentationSubsystemSharedFragment>().RepresentationSubsystem;
 		check(RepresentationSubsystem);
diff --git a/Engine/Plugins/Runtime/MassGameplay/Source/MassRepresentation/Private/MassRepresentationSubsystem.cpp b/Engine/Plugins/Runtime/MassGameplay/Source/MassRepresentation/Private/MassRepresentationSubsystem.cpp
--- a/Engine/Plugins/Runtime/MassGameplay/Source/MassRepresentation/Private/MassRepresentationSubsystem.cpp
+++ b/Engine/Plugins/Runtime/MassGameplay/Source/MassRepresentation/Private/MassRepresentationSubsystem.cpp
@@ -17,6 +17,8 @@
 #include "MassEntityView.h"
 #include "WorldPartition/WorldPartitionSubsystem.h"
 #include "WorldPartition/WorldPartitionRuntimeCell.h"
+#include "MassEntityUtils.h"
+
 
 int16 UMassRepresentationSubsystem::FindOrAddStaticMeshDesc(const FStaticMeshInstanceVisualizationDesc& Desc)
 {
@@ -302,7 +304,8 @@ void UMassRepresentationSubsystem::Initialize(FSubsystemCollectionBase& Collecti
 
 	if (UWorld* World = GetWorld())
 	{
-		EntitySubsystem = UWorld::GetSubsystem<UMassEntitySubsystem>(World);
+		EntityManager = UE::Mass::Utils::GetEntityManagerChecked(*World).AsShared();
+
 		ActorSpawnerSubsystem = World->GetSubsystem<UMassActorSpawnerSubsystem>();
 		WorldPartitionSubsystem = World->GetSubsystem<UWorldPartitionSubsystem>();
 
@@ -374,32 +377,32 @@ void UMassRepresentationSubsystem::OnProcessingPhaseStarted(const float DeltaSec
 
 void UMassRepresentationSubsystem::OnMassAgentComponentEntityAssociated(const UMassAgentComponent& AgentComponent)
 {
-	check(EntitySubsystem);
+	check(EntityManager);
 
 	const FMassEntityHandle MassAgent = AgentComponent.GetEntityHandle();
-	checkf(EntitySubsystem->IsEntityValid(MassAgent), TEXT("Expecting a valid mass entity"));
-	if (EntitySubsystem->IsEntityValid(MassAgent) && AgentComponent.IsNetSimulating())
+	checkf(EntityManager->IsEntityValid(MassAgent), TEXT("Expecting a valid mass entity"));
+	if (EntityManager->IsEntityValid(MassAgent) && AgentComponent.IsNetSimulating())
 	{
 		// Check if this mass agent already handled by this sub system, if yes than release any local spawned actor or cancel any spawn requests
 		if (HandledMassAgents.Find(MassAgent))
 		{
-			UMassRepresentationActorManagement::ReleaseAnyActorOrCancelAnySpawning(*EntitySubsystem, MassAgent);
+			UMassRepresentationActorManagement::ReleaseAnyActorOrCancelAnySpawning(*EntityManager.Get(), MassAgent);
 		}
 	}
 }
 
 void UMassRepresentationSubsystem::OnMassAgentComponentEntityDetaching(const UMassAgentComponent& AgentComponent)
 {
-	check(EntitySubsystem);
+	check(EntityManager);
 
 	AActor* ComponentOwner = AgentComponent.GetOwner();
 	check(ComponentOwner);
 
 	const FMassEntityHandle MassAgent = AgentComponent.GetEntityHandle();
-	checkf(EntitySubsystem->IsEntityValid(MassAgent), TEXT("Expecting a valid mass entity"));
-	if (EntitySubsystem->IsEntityValid(MassAgent) && AgentComponent.IsNetSimulating())
+	checkf(EntityManager->IsEntityValid(MassAgent), TEXT("Expecting a valid mass entity"));
+	if (EntityManager->IsEntityValid(MassAgent) && AgentComponent.IsNetSimulating())
 	{
-		const FMassEntityView EntityView(*EntitySubsystem, MassAgent);
+		const FMassEntityView EntityView(*EntityManager.Get(), MassAgent);
 		if (FMassRepresentationFragment* Representation = EntityView.GetFragmentDataPtr<FMassRepresentationFragment>())
 		{
 			// Force a reevaluate of the current representation
diff --git a/Engine/Plugins/Runtime/MassGameplay/Source/MassRepresentation/Private/MassUpdateISMProcessor.cpp b/Engine/Plugins/Runtime/MassGameplay/Source/MassRepresentation/Private/MassUpdateISMProcessor.cpp
--- a/Engine/Plugins/Runtime/MassGameplay/Source/MassRepresentation/Private/MassUpdateISMProcessor.cpp
+++ b/Engine/Plugins/Runtime/MassGameplay/Source/MassRepresentation/Private/MassUpdateISMProcessor.cpp
@@ -28,9 +28,9 @@ void UMassUpdateISMProcessor::ConfigureQueries()
 	EntityQuery.AddSharedRequirement<FMassRepresentationSubsystemSharedFragment>(EMassFragmentAccess::ReadWrite);
 }
 
-void UMassUpdateISMProcessor::Execute(UMassEntitySubsystem& EntitySubsystem, FMassExecutionContext& Context)
+void UMassUpdateISMProcessor::Execute(FMassEntityManager& EntityManager, FMassExecutionContext& Context)
 {
-	EntityQuery.ForEachEntityChunk(EntitySubsystem, Context, [](FMassExecutionContext& Context)
+	EntityQuery.ForEachEntityChunk(EntityManager, Context, [](FMassExecutionContext& Context)
 	{
 		UMassRepresentationSubsystem* RepresentationSubsystem = Context.GetSharedFragment<FMassRepresentationSubsystemSharedFragment>().RepresentationSubsystem;
 		check(RepresentationSubsystem);
diff --git a/Engine/Plugins/Runtime/MassGameplay/Source/MassRepresentation/Private/MassVisualizationLODProcessor.cpp b/Engine/Plugins/Runtime/MassGameplay/Source/MassRepresentation/Private/MassVisualizationLODProcessor.cpp
--- a/Engine/Plugins/Runtime/MassGameplay/Source/MassRepresentation/Private/MassVisualizationLODProcessor.cpp
+++ b/Engine/Plugins/Runtime/MassGameplay/Source/MassRepresentation/Private/MassVisualizationLODProcessor.cpp
@@ -50,11 +50,11 @@ void UMassVisualizationLODProcessor::ConfigureQueries()
 	ProcessorRequirements.AddSubsystemRequirement<UMassLODSubsystem>(EMassFragmentAccess::ReadOnly);
 }
 
-void UMassVisualizationLODProcessor::Execute(UMassEntitySubsystem& EntitySubsystem, FMassExecutionContext& Context)
+void UMassVisualizationLODProcessor::Execute(FMassEntityManager& EntityManager, FMassExecutionContext& Context)
 {
 	if (bForceOFFLOD)
 	{
-		CloseEntityQuery.ForEachEntityChunk(EntitySubsystem, Context, [this](FMassExecutionContext& Context)
+		CloseEntityQuery.ForEachEntityChunk(EntityManager, Context, [this](FMassExecutionContext& Context)
 		{
 			FMassVisualizationLODSharedFragment& LODSharedFragment = Context.GetMutableSharedFragment<FMassVisualizationLODSharedFragment>();
 			TArrayView<FMassRepresentationLODFragment> RepresentationLODList = Context.GetMutableFragmentView<FMassRepresentationLODFragment>();
@@ -65,9 +65,9 @@ void UMassVisualizationLODProcessor::Execute(UMassEntitySubsystem& EntitySubsyst
 
 	{
 		TRACE_CPUPROFILER_EVENT_SCOPE(PrepareExecution)
-		const UMassLODSubsystem& LODSubsystem = Context.GetSubsystemChecked<UMassLODSubsystem>(EntitySubsystem.GetWorld());
+		const UMassLODSubsystem& LODSubsystem = Context.GetSubsystemChecked<UMassLODSubsystem>(EntityManager.GetWorld());
 		const TArray<FViewerInfo>& Viewers = LODSubsystem.GetViewers();
-		EntitySubsystem.ForEachSharedFragment<FMassVisualizationLODSharedFragment>([this, &Viewers](FMassVisualizationLODSharedFragment& LODSharedFragment)
+		EntityManager.ForEachSharedFragment<FMassVisualizationLODSharedFragment>([this, &Viewers](FMassVisualizationLODSharedFragment& LODSharedFragment)
 		{
 			if (FilterTag == LODSharedFragment.FilterTag)
 			{
@@ -86,13 +86,13 @@ void UMassVisualizationLODProcessor::Execute(UMassEntitySubsystem& EntitySubsyst
 			TConstArrayView<FMassViewerInfoFragment> ViewerInfoList = Context.GetFragmentView<FMassViewerInfoFragment>();
 			LODSharedFragment.LODCalculator.CalculateLOD(Context, ViewerInfoList, RepresentationLODList);
 		};
-		CloseEntityQuery.ForEachEntityChunk(EntitySubsystem, Context, CalculateLOD);
-		FarEntityQuery.ForEachEntityChunk(EntitySubsystem, Context, CalculateLOD);
+		CloseEntityQuery.ForEachEntityChunk(EntityManager, Context, CalculateLOD);
+		FarEntityQuery.ForEachEntityChunk(EntityManager, Context, CalculateLOD);
 	}
 
 	{
 		TRACE_CPUPROFILER_EVENT_SCOPE(AdjustDistanceAndLODFromCount)
-		EntitySubsystem.ForEachSharedFragment<FMassVisualizationLODSharedFragment>([this](FMassVisualizationLODSharedFragment& LODSharedFragment)
+		EntityManager.ForEachSharedFragment<FMassVisualizationLODSharedFragment>([this](FMassVisualizationLODSharedFragment& LODSharedFragment)
 		{
 			if (FilterTag == LODSharedFragment.FilterTag)
 			{
@@ -100,7 +100,7 @@ void UMassVisualizationLODProcessor::Execute(UMassEntitySubsystem& EntitySubsyst
 			}
 		});
 
-		CloseEntityAdjustDistanceQuery.ForEachEntityChunk(EntitySubsystem, Context, [this](FMassExecutionContext& Context)
+		CloseEntityAdjustDistanceQuery.ForEachEntityChunk(EntityManager, Context, [this](FMassExecutionContext& Context)
 		{
 			FMassVisualizationLODSharedFragment& LODSharedFragment = Context.GetMutableSharedFragment<FMassVisualizationLODSharedFragment>();
 			TConstArrayView<FMassViewerInfoFragment> ViewerInfoList = Context.GetFragmentView<FMassViewerInfoFragment>();
@@ -114,8 +114,8 @@ void UMassVisualizationLODProcessor::Execute(UMassEntitySubsystem& EntitySubsyst
 	if (UE::MassRepresentation::bDebugRepresentationLOD)
 	{
 		TRACE_CPUPROFILER_EVENT_SCOPE(DebugDisplayLOD)
-		UWorld* World = EntitySubsystem.GetWorld();
-		DebugEntityQuery.ForEachEntityChunk(EntitySubsystem, Context, [World](FMassExecutionContext& Context)
+		UWorld* World = EntityManager.GetWorld();
+		DebugEntityQuery.ForEachEntityChunk(EntityManager, Context, [World](FMassExecutionContext& Context)
 		{
 			FMassVisualizationLODSharedFragment& LODSharedFragment = Context.GetMutableSharedFragment<FMassVisualizationLODSharedFragment>();
 			TConstArrayView<FMassRepresentationLODFragment> RepresentationLODList = Context.GetFragmentView<FMassRepresentationLODFragment>();
diff --git a/Engine/Plugins/Runtime/MassGameplay/Source/MassRepresentation/Private/MassVisualizationTrait.cpp b/Engine/Plugins/Runtime/MassGameplay/Source/MassRepresentation/Private/MassVisualizationTrait.cpp
--- a/Engine/Plugins/Runtime/MassGameplay/Source/MassRepresentation/Private/MassVisualizationTrait.cpp
+++ b/Engine/Plugins/Runtime/MassGameplay/Source/MassRepresentation/Private/MassVisualizationTrait.cpp
@@ -9,6 +9,8 @@
 #include "Engine/World.h"
 #include "MassLODFragments.h"
 #include "MassActorSubsystem.h"
+#include "MassEntityUtils.h"
+
 
 UMassVisualizationTrait::UMassVisualizationTrait()
 {
@@ -52,8 +54,7 @@ void UMassVisualizationTrait::BuildTemplate(FMassEntityTemplateBuildContext& Bui
 	BuildContext.RequireFragment<FTransformFragment>();
 	BuildContext.RequireFragment<FMassActorFragment>();
 
-	UMassEntitySubsystem* EntitySubsystem = UWorld::GetSubsystem<UMassEntitySubsystem>(&World);
-	check(EntitySubsystem);
+	FMassEntityManager& EntityManager = UE::Mass::Utils::GetEntityManagerChecked(World);
 
 	UMassRepresentationSubsystem* RepresentationSubsystem = Cast<UMassRepresentationSubsystem>(World.GetSubsystemBase(RepresentationSubsystemClass));
 	if (RepresentationSubsystem == nullptr)
@@ -66,7 +67,7 @@ void UMassVisualizationTrait::BuildTemplate(FMassEntityTemplateBuildContext& Bui
 	FMassRepresentationSubsystemSharedFragment Subsystem;
 	Subsystem.RepresentationSubsystem = RepresentationSubsystem;
 	uint32 SubsystemHash = UE::StructUtils::GetStructCrc32(FConstStructView::Make(Subsystem));
-	FSharedStruct SubsystemFragment = EntitySubsystem->GetOrCreateSharedFragment<FMassRepresentationSubsystemSharedFragment>(SubsystemHash, Subsystem);
+	FSharedStruct SubsystemFragment = EntityManager.GetOrCreateSharedFragment<FMassRepresentationSubsystemSharedFragment>(SubsystemHash, Subsystem);
 	BuildContext.AddSharedFragment(SubsystemFragment);
 
 	if (!Params.RepresentationActorManagementClass)
@@ -74,7 +75,7 @@ void UMassVisualizationTrait::BuildTemplate(FMassEntityTemplateBuildContext& Bui
 		UE_LOG(LogMassRepresentation, Error, TEXT("Expecting a valid class for the representation actor management"));
 	}
 	uint32 ParamsHash = UE::StructUtils::GetStructCrc32(FConstStructView::Make(Params));
-	FConstSharedStruct ParamsFragment = EntitySubsystem->GetOrCreateConstSharedFragment<FMassRepresentationParameters>(ParamsHash, Params);
+	FConstSharedStruct ParamsFragment = EntityManager.GetOrCreateConstSharedFragment<FMassRepresentationParameters>(ParamsHash, Params);
 	ParamsFragment.Get<FMassRepresentationParameters>().ComputeCachedValues();
 	BuildContext.AddConstSharedFragment(ParamsFragment);
 
@@ -84,9 +85,9 @@ void UMassVisualizationTrait::BuildTemplate(FMassEntityTemplateBuildContext& Bui
 	RepresentationFragment.LowResTemplateActorIndex = LowResTemplateActor.Get() ? RepresentationSubsystem->FindOrAddTemplateActor(LowResTemplateActor.Get()) : INDEX_NONE;
 
 	uint32 LODParamsHash = UE::StructUtils::GetStructCrc32(FConstStructView::Make(LODParams));
-	FConstSharedStruct LODParamsFragment = EntitySubsystem->GetOrCreateConstSharedFragment<FMassVisualizationLODParameters>(LODParamsHash, LODParams);
+	FConstSharedStruct LODParamsFragment = EntityManager.GetOrCreateConstSharedFragment<FMassVisualizationLODParameters>(LODParamsHash, LODParams);
 	BuildContext.AddConstSharedFragment(LODParamsFragment);
-	FSharedStruct LODSharedFragment = EntitySubsystem->GetOrCreateSharedFragment<FMassVisualizationLODSharedFragment>(LODParamsHash, LODParams);
+	FSharedStruct LODSharedFragment = EntityManager.GetOrCreateSharedFragment<FMassVisualizationLODSharedFragment>(LODParamsHash, LODParams);
 	BuildContext.AddSharedFragment(LODSharedFragment);
 
 	BuildContext.AddFragment<FMassRepresentationLODFragment>();
diff --git a/Engine/Plugins/Runtime/MassGameplay/Source/MassRepresentation/Public/MassRepresentationActorManagement.h b/Engine/Plugins/Runtime/MassGameplay/Source/MassRepresentation/Public/MassRepresentationActorManagement.h
--- a/Engine/Plugins/Runtime/MassGameplay/Source/MassRepresentation/Public/MassRepresentationActorManagement.h
+++ b/Engine/Plugins/Runtime/MassGameplay/Source/MassRepresentation/Public/MassRepresentationActorManagement.h
@@ -6,7 +6,7 @@
 
 struct FMassActorSpawnRequestHandle;
 struct FMassCommandBuffer;
-class UMassEntitySubsystem;
+struct FMassEntityManager;
 enum class EMassActorSpawnRequestAction : uint8;
 enum class EMassActorEnabledType : uint8;
 struct FMassActorSpawnRequestHandle;
@@ -42,7 +42,7 @@ public:
 	 * @param Priority of this spawn request in comparison with the others, lower value means higher priority
 	 * @return the actor spawned
 	 */
-	virtual AActor* GetOrSpawnActor(UMassRepresentationSubsystem& RepresentationSubsystem, UMassEntitySubsystem& EntitySubsystem, const FMassEntityHandle MassAgent, FMassActorFragment& ActorInfo, const FTransform& Transform, const int16 TemplateActorIndex, FMassActorSpawnRequestHandle& SpawnRequestHandle, const float Priority) const;
+	virtual AActor* GetOrSpawnActor(UMassRepresentationSubsystem& RepresentationSubsystem, FMassEntityManager& EntitySubsystem, const FMassEntityHandle MassAgent, FMassActorFragment& ActorInfo, const FTransform& Transform, const int16 TemplateActorIndex, FMassActorSpawnRequestHandle& SpawnRequestHandle, const float Priority) const;
 
 	/**
 	 * Enable/disable a spawned actor for a mass entity
@@ -69,7 +69,7 @@ public:
 	 * @param SpawnRequest of the actor that is about to spawn
 	 * @param EntitySubsystem to use to retrieve the mass agent fragments
 	 */
-	virtual void OnPreActorSpawn(const FMassActorSpawnRequestHandle& SpawnRequestHandle, FConstStructView SpawnRequest, UMassEntitySubsystem* EntitySubsystem) const;
+	virtual void OnPreActorSpawn(const FMassActorSpawnRequestHandle& SpawnRequestHandle, FConstStructView SpawnRequest, FMassEntityManager* EntitySubsystem) const;
 
 	/**
 	 * Method that will be bound to a delegate used post-spawn to notify and let the requester configure the actor
@@ -78,7 +78,7 @@ public:
 	 * @param EntitySubsystem to use to retrieve the mass agent fragments
 	 * @return The action to take on the spawn request, either keep it there or remove it.
 	 */
-	virtual EMassActorSpawnRequestAction OnPostActorSpawn(const FMassActorSpawnRequestHandle& SpawnRequestHandle, FConstStructView SpawnRequest, UMassEntitySubsystem* EntitySubsystem) const;
+	virtual EMassActorSpawnRequestAction OnPostActorSpawn(const FMassActorSpawnRequestHandle& SpawnRequestHandle, FConstStructView SpawnRequest, FMassEntityManager* EntitySubsystem) const;
 
 	/**
 	 * Static methods to Release an actor or cancel its spawning (calls ReleaseAnyActorOrCancelAnySpawning)
@@ -88,7 +88,7 @@ public:
 	 * @param MassAgent is the handle to the associated mass agent
 	 * @return True if actor was release or spawning request was canceled
 	 */
-	static void ReleaseAnyActorOrCancelAnySpawning(UMassEntitySubsystem& EntitySubsystem, const FMassEntityHandle MassAgent);
+	static void ReleaseAnyActorOrCancelAnySpawning(FMassEntityManager& EntityManager, const FMassEntityHandle MassAgent);
 
 	/**
 	 * static Release an actor or cancel its spawning
diff --git a/Engine/Plugins/Runtime/MassGameplay/Source/MassRepresentation/Public/MassRepresentationProcessor.h b/Engine/Plugins/Runtime/MassGameplay/Source/MassRepresentation/Public/MassRepresentationProcessor.h
--- a/Engine/Plugins/Runtime/MassGameplay/Source/MassRepresentation/Public/MassRepresentationProcessor.h
+++ b/Engine/Plugins/Runtime/MassGameplay/Source/MassRepresentation/Public/MassRepresentationProcessor.h
@@ -37,7 +37,7 @@ protected:
 	 * @param EntitySubsystem is the system to execute the lambdas on each entity chunk
 	 * @param Context is the execution context to be passed when executing the lambdas
 	 */
-	virtual void Execute(UMassEntitySubsystem& EntitySubsystem, FMassExecutionContext& Context) override;
+	virtual void Execute(FMassEntityManager& EntityManager, FMassExecutionContext& Context) override;
 
 	/**
 	 * Release the actor to the subsystem, will only release it the actor or spawn request matches the template actor
@@ -62,8 +62,7 @@ protected:
 	UPROPERTY(Transient)
 	UWorld* World;
 
-	UPROPERTY(Transient)
-	UMassEntitySubsystem* CachedEntitySubsystem;
+	TSharedPtr<FMassEntityManager> EntityManager;
 
 	FMassEntityQuery EntityQuery;
 };
@@ -83,7 +82,7 @@ protected:
 	 * @param EntitySubsystem is the system to execute the lambdas on each entity chunk
 	 * @param Context is the execution context to be passed when executing the lambdas
 	 */
-	virtual void Execute(UMassEntitySubsystem& EntitySubsystem, FMassExecutionContext& Context) override;
+	virtual void Execute(FMassEntityManager& EntityManager, FMassExecutionContext& Context) override;
 
 	/**
 	 * Updates chunk visibility info for later chunk logic optimization
@@ -120,7 +119,7 @@ public:
 
 protected:
 	virtual void ConfigureQueries() override;
-	virtual void Execute(UMassEntitySubsystem& EntitySubsystem, FMassExecutionContext& Context) override;
+	virtual void Execute(FMassEntityManager& EntityManager, FMassExecutionContext& Context) override;
 
 	FMassEntityQuery EntityQuery;
 };
\ No newline at end of file
diff --git a/Engine/Plugins/Runtime/MassGameplay/Source/MassRepresentation/Public/MassRepresentationSubsystem.h b/Engine/Plugins/Runtime/MassGameplay/Source/MassRepresentation/Public/MassRepresentationSubsystem.h
--- a/Engine/Plugins/Runtime/MassGameplay/Source/MassRepresentation/Public/MassRepresentationSubsystem.h
+++ b/Engine/Plugins/Runtime/MassGameplay/Source/MassRepresentation/Public/MassRepresentationSubsystem.h
@@ -18,7 +18,7 @@ struct FMassInstancedStaticMeshInfo;
 struct FMassActorSpawnRequestHandle;
 class UMassActorSpawnerSubsystem;
 class UMassAgentComponent;
-class UMassEntitySubsystem;
+struct FMassEntityManager;
 enum class EMassProcessingPhase : uint8;
 class UWorldPartitionSubsystem;
 
@@ -149,8 +149,7 @@ protected:
 	UPROPERTY(Transient)
 	UMassActorSpawnerSubsystem* ActorSpawnerSubsystem;
 
-	UPROPERTY(Transient)
-	UMassEntitySubsystem* EntitySubsystem;
+	TSharedPtr<FMassEntityManager> EntityManager;
 
 	UPROPERTY(Transient)
 	UWorldPartitionSubsystem* WorldPartitionSubsystem;
diff --git a/Engine/Plugins/Runtime/MassGameplay/Source/MassSignals/Private/MassSignalProcessorBase.cpp b/Engine/Plugins/Runtime/MassGameplay/Source/MassSignals/Private/MassSignalProcessorBase.cpp
--- a/Engine/Plugins/Runtime/MassGameplay/Source/MassSignals/Private/MassSignalProcessorBase.cpp
+++ b/Engine/Plugins/Runtime/MassGameplay/Source/MassSignals/Private/MassSignalProcessorBase.cpp
@@ -31,7 +31,7 @@ void UMassSignalProcessorBase::SubscribeToSignal(UMassSignalSubsystem& SignalSub
 	SignalSubsystem.GetSignalDelegateByName(SignalName).AddUObject(this, &UMassSignalProcessorBase::OnSignalReceived);
 }
 
-void UMassSignalProcessorBase::Execute(UMassEntitySubsystem& EntitySubsystem, FMassExecutionContext& Context)
+void UMassSignalProcessorBase::Execute(FMassEntityManager& EntityManager, FMassExecutionContext& Context)
 {
 	QUICK_SCOPE_CYCLE_COUNTER(SignalEntities);
 
@@ -47,7 +47,7 @@ void UMassSignalProcessorBase::Execute(UMassEntitySubsystem& EntitySubsystem, FM
 		return;
 	}
 
-	EntityQuery.CacheArchetypes(EntitySubsystem);
+	EntityQuery.CacheArchetypes(EntityManager);
 	if (EntityQuery.GetArchetypes().Num() > 0)
 	{
 		// EntitySet stores unique array of entities per specified archetype.
@@ -97,7 +97,7 @@ void UMassSignalProcessorBase::Execute(UMassEntitySubsystem& EntitySubsystem, FM
 				for (const FMassEntityHandle Entity : Entities)
 				{
 					// Add to set of supported archetypes. Dont process if we don't care about the type.
-					FMassArchetypeHandle Archetype = EntitySubsystem.GetArchetypeForEntity(Entity);
+					FMassArchetypeHandle Archetype = EntityManager.GetArchetypeForEntity(Entity);
 					FEntitySet* Set = PrevSet->Archetype == Archetype ? PrevSet : EntitySets.FindByPredicate([&Archetype](const FEntitySet& Set) { return Archetype == Set.Archetype; });
 					if (Set != nullptr)
 					{
@@ -118,7 +118,7 @@ void UMassSignalProcessorBase::Execute(UMassEntitySubsystem& EntitySubsystem, FM
 				if (Set.Entities.Num() > 0)
 				{
 					Context.SetEntityCollection(FMassArchetypeEntityCollection(Set.Archetype, Set.Entities, FMassArchetypeEntityCollection::FoldDuplicates));
-					SignalEntities(EntitySubsystem, Context, SignalNameLookup);
+					SignalEntities(EntityManager, Context, SignalNameLookup);
 					Context.ClearEntityCollection();
 				}
 				Set.Reset();
diff --git a/Engine/Plugins/Runtime/MassGameplay/Source/MassSignals/Private/MassSignalSubsystem.cpp b/Engine/Plugins/Runtime/MassGameplay/Source/MassSignals/Private/MassSignalSubsystem.cpp
--- a/Engine/Plugins/Runtime/MassGameplay/Source/MassSignals/Private/MassSignalSubsystem.cpp
+++ b/Engine/Plugins/Runtime/MassGameplay/Source/MassSignals/Private/MassSignalSubsystem.cpp
@@ -77,7 +77,7 @@ void UMassSignalSubsystem::SignalEntityDeferred(FMassExecutionContext& Context,
 void UMassSignalSubsystem::SignalEntitiesDeferred(FMassExecutionContext& Context, FName SignalName, TConstArrayView<FMassEntityHandle> Entities)
 {
 	checkf(Entities.Num() > 0, TEXT("Expecting entities to signal"));
-	Context.Defer().PushCommand<FMassDeferredSetCommand>([SignalName, Entities](const UMassEntitySubsystem& System)
+	Context.Defer().PushCommand<FMassDeferredSetCommand>([SignalName, Entities](const FMassEntityManager& System)
 	{
 		UMassSignalSubsystem* SignalSubsystem = UWorld::GetSubsystem<UMassSignalSubsystem>(System.GetWorld());
 		SignalSubsystem->SignalEntities(SignalName, Entities);
@@ -96,7 +96,7 @@ void UMassSignalSubsystem::DelaySignalEntityDeferred(FMassExecutionContext& Cont
 void UMassSignalSubsystem::DelaySignalEntitiesDeferred(FMassExecutionContext& Context, FName SignalName, TConstArrayView<FMassEntityHandle> Entities, const float DelayInSeconds)
 {
 	checkf(Entities.Num() > 0, TEXT("Expecting entities to signal"));
-	Context.Defer().PushCommand<FMassDeferredSetCommand>([SignalName, Entities, DelayInSeconds](const UMassEntitySubsystem& System)
+	Context.Defer().PushCommand<FMassDeferredSetCommand>([SignalName, Entities, DelayInSeconds](const FMassEntityManager& System)
 	{
 		UMassSignalSubsystem* SignalSubsystem = UWorld::GetSubsystem<UMassSignalSubsystem>(System.GetWorld());
 		SignalSubsystem->DelaySignalEntities(SignalName, Entities, DelayInSeconds);
diff --git a/Engine/Plugins/Runtime/MassGameplay/Source/MassSignals/Public/MassSignalProcessorBase.h b/Engine/Plugins/Runtime/MassGameplay/Source/MassSignals/Public/MassSignalProcessorBase.h
--- a/Engine/Plugins/Runtime/MassGameplay/Source/MassSignals/Public/MassSignalProcessorBase.h
+++ b/Engine/Plugins/Runtime/MassGameplay/Source/MassSignals/Public/MassSignalProcessorBase.h
@@ -34,7 +34,7 @@ protected:
 	 * @param Context is the execution context to be passed when executing the lambdas
 	 * @param EntitySignals Look up to retrieve for each entities their raised signal via GetSignalsForEntity
 	 */
-	 virtual void SignalEntities(UMassEntitySubsystem& EntitySubsystem, FMassExecutionContext& Context, FMassSignalNameLookup& EntitySignals) PURE_VIRTUAL(UMassSignalProcessorBase::SignalEntities, );
+	 virtual void SignalEntities(FMassEntityManager& EntityManager, FMassExecutionContext& Context, FMassSignalNameLookup& EntitySignals) PURE_VIRTUAL(UMassSignalProcessorBase::SignalEntities, );
 
 	/**
 	 * Callback that is being called when new signal is raised
@@ -48,7 +48,7 @@ protected:
 	 * @param EntitySubsystem is the system to execute the lambdas on each entity chunk
 	 * @param Context is the execution context to be passed when executing the lambdas
 	 */
-	virtual void Execute(UMassEntitySubsystem& EntitySubsystem, FMassExecutionContext& Context) override;
+	virtual void Execute(FMassEntityManager& EntityManager, FMassExecutionContext& Context) override;
 
 	/**
 	 * To receive notification about a particular signal, you need to subscribe to it.
diff --git a/Engine/Plugins/Runtime/MassGameplay/Source/MassSimulation/Private/MassSimulationSubsystem.cpp b/Engine/Plugins/Runtime/MassGameplay/Source/MassSimulation/Private/MassSimulationSubsystem.cpp
--- a/Engine/Plugins/Runtime/MassGameplay/Source/MassSimulation/Private/MassSimulationSubsystem.cpp
+++ b/Engine/Plugins/Runtime/MassGameplay/Source/MassSimulation/Private/MassSimulationSubsystem.cpp
@@ -57,8 +57,9 @@ void UMassSimulationSubsystem::Initialize(FSubsystemCollectionBase& Collection)
 {
 	Super::Initialize(Collection);
 
-	EntitySubsystem = Collection.InitializeDependency<UMassEntitySubsystem>();
+	UMassEntitySubsystem* EntitySubsystem = Collection.InitializeDependency<UMassEntitySubsystem>();
 	check(EntitySubsystem);
+	EntityManager = EntitySubsystem->GetMutableEntityManager().AsShared();
 	
 	GetOnProcessingPhaseStarted(EMassProcessingPhase::PrePhysics).AddUObject(this, &UMassSimulationSubsystem::OnProcessingPhaseStarted, EMassProcessingPhase::PrePhysics);
 }
@@ -130,10 +131,10 @@ void UMassSimulationSubsystem::OnProcessingPhaseStarted(const float DeltaSeconds
 		case EMassProcessingPhase::PrePhysics:
 			{
 				TRACE_CPUPROFILER_EVENT_SCOPE(DoEntityCompation);
-				check(EntitySubsystem);
+				check(EntityManager);
 				if (UE::MassSimulation::bDoEntityCompaction)
 				{
-					EntitySubsystem->DoEntityCompaction(GET_MASS_CONFIG_VALUE(DesiredEntityCompactionTimeSlicePerTick));
+					EntityManager->DoEntityCompaction(GET_MASS_CONFIG_VALUE(DesiredEntityCompactionTimeSlicePerTick));
 				}
 			}
 			break;
diff --git a/Engine/Plugins/Runtime/MassGameplay/Source/MassSimulation/Public/MassSimulationSubsystem.h b/Engine/Plugins/Runtime/MassGameplay/Source/MassSimulation/Public/MassSimulationSubsystem.h
--- a/Engine/Plugins/Runtime/MassGameplay/Source/MassSimulation/Public/MassSimulationSubsystem.h
+++ b/Engine/Plugins/Runtime/MassGameplay/Source/MassSimulation/Public/MassSimulationSubsystem.h
@@ -9,7 +9,7 @@
 #include "MassSimulationSubsystem.generated.h"
 
 
-class UMassEntitySubsystem;
+struct FMassEntityManager;
 class UMassProcessingPhaseManager;
 
 DECLARE_LOG_CATEGORY_EXTERN(LogMassSim, Log, All);
@@ -23,7 +23,6 @@ public:
 	
 	UMassSimulationSubsystem(const FObjectInitializer& ObjectInitializer = FObjectInitializer::Get());
 
-	//UMassEntitySubsystem* GetEntitySubsystem() const { return CachedEntitySubsystem; }
 	const UMassProcessingPhaseManager& GetPhaseManager() const { check(PhaseManager); return *PhaseManager; }
 
 	FMassProcessingPhase::FOnPhaseEvent& GetOnProcessingPhaseStarted(const EMassProcessingPhase Phase) const;
@@ -53,8 +52,7 @@ protected:
 
 protected:
 
-	UPROPERTY()
-	UMassEntitySubsystem* EntitySubsystem;
+	TSharedPtr<FMassEntityManager> EntityManager;
 
 	UPROPERTY()
 	UMassProcessingPhaseManager* PhaseManager;
diff --git a/Engine/Plugins/Runtime/MassGameplay/Source/MassSmartObjects/Private/MassSmartObjectHandler.cpp b/Engine/Plugins/Runtime/MassGameplay/Source/MassSmartObjects/Private/MassSmartObjectHandler.cpp
--- a/Engine/Plugins/Runtime/MassGameplay/Source/MassSmartObjects/Private/MassSmartObjectHandler.cpp
+++ b/Engine/Plugins/Runtime/MassGameplay/Source/MassSmartObjects/Private/MassSmartObjectHandler.cpp
@@ -35,7 +35,7 @@ namespace UE::Mass::SmartObject
 	struct FPayload
 	{
 		FMassEntityHandle Entity;
-		TWeakObjectPtr<UMassEntitySubsystem> EntitySubsystem;
+		TSharedPtr<FMassEntityManager> EntityManager;
 		TWeakObjectPtr<USmartObjectSubsystem> SmartObjectSubsystem;
 		TWeakObjectPtr<UMassSignalSubsystem> SignalSubsystem;
 	};
@@ -43,17 +43,17 @@ namespace UE::Mass::SmartObject
 	void OnSlotInvalidated(const FSmartObjectClaimHandle& ClaimHandle, const ESmartObjectSlotState State, FPayload Payload)
 	{
 		USmartObjectSubsystem* SmartObjectSubsystem = Payload.SmartObjectSubsystem.Get();
-		UMassEntitySubsystem* EntitySubsystem = Payload.EntitySubsystem.Get();
+		FMassEntityManager* EntityManager = Payload.EntityManager.Get();
 		UMassSignalSubsystem* SignalSubsystem = Payload.SignalSubsystem.Get();
-		if (SmartObjectSubsystem != nullptr && EntitySubsystem != nullptr && SignalSubsystem != nullptr && EntitySubsystem->IsEntityActive(Payload.Entity))
+		if (SmartObjectSubsystem != nullptr && EntityManager != nullptr && SignalSubsystem != nullptr && EntityManager->IsEntityActive(Payload.Entity))
 		{
-			const FMassEntityView EntityView(*EntitySubsystem, Payload.Entity);
+			const FMassEntityView EntityView(*EntityManager, Payload.Entity);
 			const FMassBehaviorEntityContext Context(EntityView, *SmartObjectSubsystem);
 
-			StopInteraction(EntitySubsystem->Defer(), Context, ClaimHandle, EMassSmartObjectInteractionStatus::Aborted);
-			if (EntitySubsystem->IsProcessing() == false)
+			StopInteraction(EntityManager->Defer(), Context, ClaimHandle, EMassSmartObjectInteractionStatus::Aborted);
+			if (EntityManager->IsProcessing() == false)
 			{
-				EntitySubsystem->FlushCommands();
+				EntityManager->FlushCommands();
 			}
 
 			SignalSubsystem->SignalEntity(UE::Mass::Signals::SmartObjectInteractionAborted, Payload.Entity);
@@ -66,7 +66,7 @@ namespace UE::Mass::SmartObject
 //----------------------------------------------------------------------//
 FMassSmartObjectRequestID FMassSmartObjectHandler::FindCandidatesAsync(const FMassEntityHandle RequestingEntity, const FGameplayTagContainer& UserTags, const FGameplayTagQuery& ActivityRequirements, const FVector& Location) const
 {
-	const FMassEntityHandle ReservedEntity = EntitySubsystem.ReserveEntity();
+	const FMassEntityHandle ReservedEntity = EntityManager.ReserveEntity();
 
 	FMassSmartObjectWorldLocationRequestFragment RequestFragment;
 	RequestFragment.SearchOrigin = Location;
@@ -83,7 +83,7 @@ FMassSmartObjectRequestID FMassSmartObjectHandler::FindCandidatesAsync(const FMa
 
 FMassSmartObjectRequestID FMassSmartObjectHandler::FindCandidatesAsync(const FMassEntityHandle RequestingEntity, const FGameplayTagContainer& UserTags, const FGameplayTagQuery& ActivityRequirements, const FZoneGraphCompactLaneLocation& LaneLocation) const
 {
-	const FMassEntityHandle ReservedEntity = EntitySubsystem.ReserveEntity();
+	const FMassEntityHandle ReservedEntity = EntityManager.ReserveEntity();
 
 	FMassSmartObjectLaneLocationRequestFragment RequestFragment;
 	RequestFragment.CompactLaneLocation = LaneLocation;
@@ -101,18 +101,18 @@ FMassSmartObjectRequestID FMassSmartObjectHandler::FindCandidatesAsync(const FMa
 const FMassSmartObjectCandidateSlots* FMassSmartObjectHandler::GetRequestCandidates(const FMassSmartObjectRequestID& RequestID) const
 {
 	const FMassEntityHandle RequestEntity = static_cast<FMassEntityHandle>(RequestID);
-	if (!ensureMsgf(EntitySubsystem.IsEntityValid(RequestEntity), TEXT("Invalid request.")))
+	if (!ensureMsgf(EntityManager.IsEntityValid(RequestEntity), TEXT("Invalid request.")))
 	{
 		return nullptr;
 	}
 
 	// Check if entity is built by now.
-	if (!EntitySubsystem.IsEntityBuilt(RequestEntity))
+	if (!EntityManager.IsEntityBuilt(RequestEntity))
 	{
 		return nullptr;
 	}
 
-	const FMassSmartObjectRequestResultFragment& RequestFragment = EntitySubsystem.GetFragmentDataChecked<FMassSmartObjectRequestResultFragment>(RequestEntity);
+	const FMassSmartObjectRequestResultFragment& RequestFragment = EntityManager.GetFragmentDataChecked<FMassSmartObjectRequestResultFragment>(RequestEntity);
 
 	return RequestFragment.bProcessed ? &RequestFragment.Candidates : nullptr;
 }
@@ -171,7 +171,7 @@ FSmartObjectClaimHandle FMassSmartObjectHandler::ClaimSmartObject(const FMassEnt
 	// Callback will be unregistered by UMassSmartObjectUserFragmentDeinitializer
 	UE::Mass::SmartObject::FPayload Payload;
 	Payload.Entity = Entity;
-	Payload.EntitySubsystem = &EntitySubsystem;
+	Payload.EntityManager = EntityManager.AsShared();
 	Payload.SmartObjectSubsystem = &SmartObjectSubsystem;
 	Payload.SignalSubsystem = &SignalSubsystem;
 	SmartObjectSubsystem.RegisterSlotInvalidationCallback(ClaimHandle, FOnSlotInvalidated::CreateStatic(&UE::Mass::SmartObject::OnSlotInvalidated, Payload));
@@ -201,7 +201,7 @@ bool FMassSmartObjectHandler::StartUsingSmartObject(
 	}
 
 	// Activate behavior
-	const FMassEntityView EntityView(EntitySubsystem, Entity);
+	const FMassEntityView EntityView(EntityManager, Entity);
 	const FMassBehaviorEntityContext Context(EntityView, SmartObjectSubsystem);
 	BehaviorDefinition->Activate(ExecutionContext.Defer(), Context);
 
@@ -261,7 +261,7 @@ void FMassSmartObjectHandler::StopUsingSmartObject(const FMassEntityHandle Entit
 	}
 #endif
 
-	const FMassEntityView EntityView(EntitySubsystem, Entity);
+	const FMassEntityView EntityView(EntityManager, Entity);
 	const FMassBehaviorEntityContext Context(EntityView, SmartObjectSubsystem);
 	UE::Mass::SmartObject::StopInteraction(ExecutionContext.Defer(), Context, User.InteractionHandle, NewStatus);
 }
@@ -286,7 +286,7 @@ void FMassSmartObjectHandler::ReleaseSmartObject(const FMassEntityHandle Entity,
 	}
 #endif
 
-	const FMassEntityView EntityView(EntitySubsystem, Entity);
+	const FMassEntityView EntityView(EntityManager, Entity);
 	const FMassBehaviorEntityContext Context(EntityView, SmartObjectSubsystem);
 
 	SmartObjectSubsystem.UnregisterSlotInvalidationCallback(ClaimHandle);
diff --git a/Engine/Plugins/Runtime/MassGameplay/Source/MassSmartObjects/Private/MassSmartObjectProcessor.cpp b/Engine/Plugins/Runtime/MassGameplay/Source/MassSmartObjects/Private/MassSmartObjectProcessor.cpp
--- a/Engine/Plugins/Runtime/MassGameplay/Source/MassSmartObjects/Private/MassSmartObjectProcessor.cpp
+++ b/Engine/Plugins/Runtime/MassGameplay/Source/MassSmartObjects/Private/MassSmartObjectProcessor.cpp
@@ -52,9 +52,9 @@ UMassSmartObjectCandidatesFinderProcessor::UMassSmartObjectCandidatesFinderProce
 	ExecutionOrder.ExecuteBefore.Add(UE::Mass::ProcessorGroupNames::Behavior);
 }
 
-void UMassSmartObjectCandidatesFinderProcessor::Execute(UMassEntitySubsystem& EntitySubsystem, FMassExecutionContext& Context)
+void UMassSmartObjectCandidatesFinderProcessor::Execute(FMassEntityManager& EntityManager, FMassExecutionContext& Context)
 {
-	UWorld* World = EntitySubsystem.GetWorld();
+	UWorld* World = EntityManager.GetWorld();
 
 	UMassSignalSubsystem& SignalSubsystem = Context.GetMutableSubsystemChecked<UMassSignalSubsystem>(World);
 	const UZoneGraphAnnotationSubsystem& AnnotationSubsystem = Context.GetSubsystemChecked<UZoneGraphAnnotationSubsystem>(World);
@@ -87,7 +87,7 @@ void UMassSmartObjectCandidatesFinderProcessor::Execute(UMassEntitySubsystem& En
 	};
 
 	// Process world location based requests
-	WorldRequestQuery.ForEachEntityChunk(EntitySubsystem, Context, [this, &Filter, &EntitiesToSignal, &BeginRequestProcessing, &EndRequestProcessing, World](FMassExecutionContext& Context)
+	WorldRequestQuery.ForEachEntityChunk(EntityManager, Context, [this, &Filter, &EntitiesToSignal, &BeginRequestProcessing, &EndRequestProcessing, World](FMassExecutionContext& Context)
 	{
 		const USmartObjectSubsystem& SmartObjectSubsystem = Context.GetSubsystemChecked<USmartObjectSubsystem>(World);
 
@@ -156,8 +156,8 @@ void UMassSmartObjectCandidatesFinderProcessor::Execute(UMassEntitySubsystem& En
 	const FZoneGraphTag SmartObjectTag = GetDefault<UMassSmartObjectSettings>()->SmartObjectTag;
 	USmartObjectZoneAnnotations* Annotations = Cast<USmartObjectZoneAnnotations>(AnnotationSubsystem.GetFirstAnnotationForTag(SmartObjectTag));
 
-	LaneRequestQuery.ForEachEntityChunk(EntitySubsystem, Context,
-		[&AnnotationSubsystem, Annotations, &Filter, SmartObjectTag, &EntitiesToSignal, &BeginRequestProcessing, &EndRequestProcessing, World = EntitySubsystem.GetWorld()](FMassExecutionContext& Context)
+	LaneRequestQuery.ForEachEntityChunk(EntityManager, Context,
+		[&AnnotationSubsystem, Annotations, &Filter, SmartObjectTag, &EntitiesToSignal, &BeginRequestProcessing, &EndRequestProcessing, World = EntityManager.GetWorld()](FMassExecutionContext& Context)
 		{
 #if WITH_MASSGAMEPLAY_DEBUG
 			const UZoneGraphSubsystem& ZoneGraphSubsystem = Context.GetSubsystemChecked<UZoneGraphSubsystem>(World);
@@ -319,14 +319,14 @@ UMassSmartObjectTimedBehaviorProcessor::UMassSmartObjectTimedBehaviorProcessor()
 	ExecutionOrder.ExecuteInGroup = UE::Mass::ProcessorGroupNames::SyncWorldToMass;
 }
 
-void UMassSmartObjectTimedBehaviorProcessor::Execute(UMassEntitySubsystem& EntitySubsystem, FMassExecutionContext& Context)
+void UMassSmartObjectTimedBehaviorProcessor::Execute(FMassEntityManager& EntityManager, FMassExecutionContext& Context)
 {
-	UWorld* World = EntitySubsystem.GetWorld();
+	UWorld* World = EntityManager.GetWorld();
 	TArray<FMassEntityHandle> ToRelease;
 
 	QUICK_SCOPE_CYCLE_COUNTER(UMassProcessor_SmartObjectTestBehavior_Run);
 
-	EntityQuery.ForEachEntityChunk(EntitySubsystem, Context, [this, &ToRelease, World](FMassExecutionContext& Context)
+	EntityQuery.ForEachEntityChunk(EntityManager, Context, [this, &ToRelease, World](FMassExecutionContext& Context)
 	{
 		const USmartObjectSubsystem& SmartObjectSubsystem = Context.GetSubsystemChecked<USmartObjectSubsystem>(World);
 
@@ -399,9 +399,9 @@ void UMassSmartObjectUserFragmentDeinitializer::ConfigureQueries()
 	EntityQuery.AddSubsystemRequirement<USmartObjectSubsystem>(EMassFragmentAccess::ReadWrite);
 }
 
-void UMassSmartObjectUserFragmentDeinitializer::Execute(UMassEntitySubsystem& EntitySubsystem, FMassExecutionContext& Context)
+void UMassSmartObjectUserFragmentDeinitializer::Execute(FMassEntityManager& EntityManager, FMassExecutionContext& Context)
 {
-	EntityQuery.ForEachEntityChunk(EntitySubsystem, Context, [this, World = EntitySubsystem.GetWorld()](FMassExecutionContext& Context)
+	EntityQuery.ForEachEntityChunk(EntityManager, Context, [this, World = EntityManager.GetWorld()](FMassExecutionContext& Context)
 		{
 			USmartObjectSubsystem& SmartObjectSubsystem = Context.GetMutableSubsystemChecked<USmartObjectSubsystem>(World);
 			const int32 NumEntities = Context.GetNumEntities();
diff --git a/Engine/Plugins/Runtime/MassGameplay/Source/MassSmartObjects/Public/MassSmartObjectBehaviorDefinition.h b/Engine/Plugins/Runtime/MassGameplay/Source/MassSmartObjects/Public/MassSmartObjectBehaviorDefinition.h
--- a/Engine/Plugins/Runtime/MassGameplay/Source/MassSmartObjects/Public/MassSmartObjectBehaviorDefinition.h
+++ b/Engine/Plugins/Runtime/MassGameplay/Source/MassSmartObjects/Public/MassSmartObjectBehaviorDefinition.h
@@ -7,7 +7,7 @@
 #include "SmartObjectDefinition.h"
 #include "MassSmartObjectBehaviorDefinition.generated.h"
 
-class UMassEntitySubsystem;
+struct FMassEntityManager;
 class USmartObjectSubsystem;
 struct FMassExecutionContext;
 struct FTransformFragment;
diff --git a/Engine/Plugins/Runtime/MassGameplay/Source/MassSmartObjects/Public/MassSmartObjectHandler.h b/Engine/Plugins/Runtime/MassGameplay/Source/MassSmartObjects/Public/MassSmartObjectHandler.h
--- a/Engine/Plugins/Runtime/MassGameplay/Source/MassSmartObjects/Public/MassSmartObjectHandler.h
+++ b/Engine/Plugins/Runtime/MassGameplay/Source/MassSmartObjects/Public/MassSmartObjectHandler.h
@@ -5,7 +5,7 @@
 #include "MassSmartObjectRequest.h"
 #include "MassSmartObjectTypes.h"
 
-class UMassEntitySubsystem;
+struct FMassEntityManager;
 class UMassSignalSubsystem;
 class USmartObjectSubsystem;
 struct FMassExecutionContext;
@@ -30,8 +30,8 @@ struct MASSSMARTOBJECTS_API FMassSmartObjectHandler
 	 * @param InSmartObjectSubsystem is the smart object subsystem
 	 * @param InSignalSubsystem is the mass signal subsystem to use to send signal to affected entities
 	 */
-	FMassSmartObjectHandler(UMassEntitySubsystem& InEntitySubsystem, FMassExecutionContext& InExecutionContext, USmartObjectSubsystem& InSmartObjectSubsystem, UMassSignalSubsystem& InSignalSubsystem)
-		: EntitySubsystem(InEntitySubsystem)
+	FMassSmartObjectHandler(FMassEntityManager& InEntityManager, FMassExecutionContext& InExecutionContext, USmartObjectSubsystem& InSmartObjectSubsystem, UMassSignalSubsystem& InSignalSubsystem)
+		: EntityManager(InEntityManager)
 		, ExecutionContext(InExecutionContext)
 		, SmartObjectSubsystem(InSmartObjectSubsystem)
 		, SignalSubsystem(InSignalSubsystem)
@@ -118,7 +118,7 @@ struct MASSSMARTOBJECTS_API FMassSmartObjectHandler
 	void ReleaseSmartObject(const FMassEntityHandle Entity, FMassSmartObjectUserFragment& User, const FSmartObjectClaimHandle ClaimHandle) const;
 
 private:
-	UMassEntitySubsystem& EntitySubsystem;
+	FMassEntityManager& EntityManager;
 	FMassExecutionContext& ExecutionContext;
 	USmartObjectSubsystem& SmartObjectSubsystem;
 	UMassSignalSubsystem& SignalSubsystem;
diff --git a/Engine/Plugins/Runtime/MassGameplay/Source/MassSmartObjects/Public/MassSmartObjectProcessor.h b/Engine/Plugins/Runtime/MassGameplay/Source/MassSmartObjects/Public/MassSmartObjectProcessor.h
--- a/Engine/Plugins/Runtime/MassGameplay/Source/MassSmartObjects/Public/MassSmartObjectProcessor.h
+++ b/Engine/Plugins/Runtime/MassGameplay/Source/MassSmartObjects/Public/MassSmartObjectProcessor.h
@@ -20,7 +20,7 @@ public:
 
 protected:
 	virtual void ConfigureQueries() override;
-	virtual void Execute(UMassEntitySubsystem& EntitySubsystem, FMassExecutionContext& Context) override;
+	virtual void Execute(FMassEntityManager& EntityManager, FMassExecutionContext& Context) override;
 
 	/** Extents used to perform the spatial query in the octree for world location queries. */
 	UPROPERTY(EditDefaultsOnly, Category = SmartObject, config)
@@ -42,7 +42,7 @@ public:
 	UMassSmartObjectTimedBehaviorProcessor();
 
 protected:
-	virtual void Execute(UMassEntitySubsystem& EntitySubsystem, FMassExecutionContext& Context) override;
+	virtual void Execute(FMassEntityManager& EntityManager, FMassExecutionContext& Context) override;
 	virtual void ConfigureQueries() override;
 
 	FMassEntityQuery EntityQuery;
@@ -58,7 +58,7 @@ class MASSSMARTOBJECTS_API UMassSmartObjectUserFragmentDeinitializer : public UM
 
 protected:
 	virtual void ConfigureQueries() override;
-	virtual void Execute(UMassEntitySubsystem& EntitySubsystem, FMassExecutionContext& Context) override;
+	virtual void Execute(FMassEntityManager& EntityManager, FMassExecutionContext& Context) override;
 
 	FMassEntityQuery EntityQuery;
 };
\ No newline at end of file
diff --git a/Engine/Plugins/Runtime/MassGameplay/Source/MassSpawner/Private/MassEntityTemplate.cpp b/Engine/Plugins/Runtime/MassGameplay/Source/MassSpawner/Private/MassEntityTemplate.cpp
--- a/Engine/Plugins/Runtime/MassGameplay/Source/MassSpawner/Private/MassEntityTemplate.cpp
+++ b/Engine/Plugins/Runtime/MassGameplay/Source/MassSpawner/Private/MassEntityTemplate.cpp
@@ -23,16 +23,16 @@ void FMassEntityTemplate::SetArchetype(const FMassArchetypeHandle& InArchetype)
 	Archetype = InArchetype;
 }
 
-FString FMassEntityTemplate::DebugGetDescription(UMassEntitySubsystem* EntitySubsystem) const
+FString FMassEntityTemplate::DebugGetDescription(FMassEntityManager* EntityManager) const
 { 
 	FStringOutputDevice Ar;
 #if WITH_MASSGAMEPLAY_DEBUG
 	Ar.SetAutoEmitLineTerminator(true);
 
-	if (EntitySubsystem)
+	if (EntityManager)
 	{
 		Ar += TEXT("Archetype details:\n");
-		Ar += DebugGetArchetypeDescription(*EntitySubsystem);
+		Ar += DebugGetArchetypeDescription(*EntityManager);
 	}
 	else
 	{
@@ -44,7 +44,7 @@ FString FMassEntityTemplate::DebugGetDescription(UMassEntitySubsystem* EntitySub
 	return MoveTemp(Ar);
 }
 
-FString FMassEntityTemplate::DebugGetArchetypeDescription(UMassEntitySubsystem& EntitySubsystem) const
+FString FMassEntityTemplate::DebugGetArchetypeDescription(FMassEntityManager& EntityManager) const
 {
 	FStringOutputDevice OutDescription;
 #if WITH_MASSGAMEPLAY_DEBUG
diff --git a/Engine/Plugins/Runtime/MassGameplay/Source/MassSpawner/Private/MassEntityTemplateRegistry.cpp b/Engine/Plugins/Runtime/MassGameplay/Source/MassSpawner/Private/MassEntityTemplateRegistry.cpp
--- a/Engine/Plugins/Runtime/MassGameplay/Source/MassSpawner/Private/MassEntityTemplateRegistry.cpp
+++ b/Engine/Plugins/Runtime/MassGameplay/Source/MassSpawner/Private/MassEntityTemplateRegistry.cpp
@@ -10,6 +10,8 @@
 #include "MassEntityTypes.h"
 #include "MassEntityTraitBase.h"
 #include "Logging/MessageLog.h"
+#include "MassEntityUtils.h"
+
 
 #define LOCTEXT_NAMESPACE "Mass"
 
@@ -136,7 +138,7 @@ bool UMassEntityTemplateRegistry::BuildTemplateImpl(const FStructToTemplateBuild
 		InitializeEntityTemplate(OutTemplate);
 
 		UE_VLOG(this, LogMassSpawner, Log, TEXT("Created entity template for %s:\n%s"), *GetNameSafe(StructInstance.GetScriptStruct())
-			, *OutTemplate.DebugGetDescription(UWorld::GetSubsystem<UMassEntitySubsystem>(World)));
+			, *OutTemplate.DebugGetDescription(UE::Mass::Utils::GetEntityManager(World)));
 
 		return true;
 	}
@@ -149,14 +151,14 @@ void UMassEntityTemplateRegistry::InitializeEntityTemplate(FMassEntityTemplate&
 	check(!OutTemplate.IsEmpty());
 
 	UWorld* World = GetWorld();
+	check(World);
 	// find or create template
-	UMassEntitySubsystem* EntitySys = UWorld::GetSubsystem<UMassEntitySubsystem>(World);
-	check(EntitySys);
+	FMassEntityManager& EntityManager = UE::Mass::Utils::GetEntityManagerChecked(*World);
 
 	// Sort anything there is to sort for later comparison purposes
 	OutTemplate.Sort();
 
-	const FMassArchetypeHandle ArchetypeHandle = EntitySys->CreateArchetype(OutTemplate.GetCompositionDescriptor(), FName(OutTemplate.GetTemplateName()));
+	const FMassArchetypeHandle ArchetypeHandle = EntityManager.CreateArchetype(OutTemplate.GetCompositionDescriptor(), FName(OutTemplate.GetTemplateName()));
 	OutTemplate.SetArchetype(ArchetypeHandle);
 }
 
diff --git a/Engine/Plugins/Runtime/MassGameplay/Source/MassSpawner/Private/MassSpawnLocationProcessor.cpp b/Engine/Plugins/Runtime/MassGameplay/Source/MassSpawner/Private/MassSpawnLocationProcessor.cpp
--- a/Engine/Plugins/Runtime/MassGameplay/Source/MassSpawner/Private/MassSpawnLocationProcessor.cpp
+++ b/Engine/Plugins/Runtime/MassGameplay/Source/MassSpawner/Private/MassSpawnLocationProcessor.cpp
@@ -21,7 +21,7 @@ void UMassSpawnLocationProcessor::ConfigureQueries()
 	EntityQuery.AddRequirement<FTransformFragment>(EMassFragmentAccess::ReadWrite);
 }
 
-void UMassSpawnLocationProcessor::Execute(UMassEntitySubsystem& EntitySubsystem, FMassExecutionContext& Context)
+void UMassSpawnLocationProcessor::Execute(FMassEntityManager& EntityManager, FMassExecutionContext& Context)
 {
 	if (!ensure(Context.ValidateAuxDataType<FMassTransformsSpawnData>()))
 	{
@@ -29,7 +29,7 @@ void UMassSpawnLocationProcessor::Execute(UMassEntitySubsystem& EntitySubsystem,
 		return;
 	}
 
-	const UWorld* World = EntitySubsystem.GetWorld();
+	const UWorld* World = EntityManager.GetWorld();
 
 	check(World);
 
@@ -48,7 +48,7 @@ void UMassSpawnLocationProcessor::Execute(UMassEntitySubsystem& EntitySubsystem,
 		}
 
 		int32 NumRequiredSpawnTransforms = 0;
-		EntityQuery.ForEachEntityChunk(EntitySubsystem, Context, [&NumRequiredSpawnTransforms](const FMassExecutionContext& Context)
+		EntityQuery.ForEachEntityChunk(EntityManager, Context, [&NumRequiredSpawnTransforms](const FMassExecutionContext& Context)
 			{
 				NumRequiredSpawnTransforms += Context.GetNumEntities();
 			});
@@ -67,7 +67,7 @@ void UMassSpawnLocationProcessor::Execute(UMassEntitySubsystem& EntitySubsystem,
 			}
 		}
 
-		EntityQuery.ForEachEntityChunk(EntitySubsystem, Context, [&Transforms, this](FMassExecutionContext& Context)
+		EntityQuery.ForEachEntityChunk(EntityManager, Context, [&Transforms, this](FMassExecutionContext& Context)
 			{
 				const TArrayView<FTransformFragment> LocationList = Context.GetMutableFragmentView<FTransformFragment>();
 				const int32 NumEntities = Context.GetNumEntities();
diff --git a/Engine/Plugins/Runtime/MassGameplay/Source/MassSpawner/Private/MassSpawner.cpp b/Engine/Plugins/Runtime/MassGameplay/Source/MassSpawner/Private/MassSpawner.cpp
--- a/Engine/Plugins/Runtime/MassGameplay/Source/MassSpawner/Private/MassSpawner.cpp
+++ b/Engine/Plugins/Runtime/MassGameplay/Source/MassSpawner/Private/MassSpawner.cpp
@@ -17,6 +17,7 @@
 #include "Engine/AssetManager.h"
 #include "MassSpawnLocationProcessor.h"
 #include "MassExecutor.h"
+#include "MassEntityUtils.h"
 #if WITH_EDITOR
 #include "Engine/Texture2D.h"
 #endif
@@ -470,12 +471,10 @@ void AMassSpawner::SpawnGeneratedEntities(TConstArrayView<FMassEntitySpawnDataGe
 		}
 	}
 
-	if (Processors.Num() > 0)
+	if (Processors.Num() > 0 && World)
 	{
-		UMassEntitySubsystem* EntitySystem = UWorld::GetSubsystem<UMassEntitySubsystem>(GetWorld());
-		check(EntitySystem);
-
-		FMassProcessingContext ProcessingContext(*EntitySystem, /*TimeDelta=*/0.0f);
+		FMassEntityManager& EntityManager = UE::Mass::Utils::GetEntityManagerChecked(*World);
+		FMassProcessingContext ProcessingContext(EntityManager, /*TimeDelta=*/0.0f);
 		UE::Mass::Executor::RunProcessorsView(Processors, ProcessingContext);
 	}
 
