Separated "entity management" and USubsystem stuff of UMassEntityManager. Now UMassEntityManager hosts a shared ref of newly created FMassEntityManager, and users are free to create their own FMassEnt

diff --git a/Engine/Plugins/Runtime/MassGameplay/Source/MassSpawner/Private/MassSpawnerSubsystem.cpp b/Engine/Plugins/Runtime/MassGameplay/Source/MassSpawner/Private/MassSpawnerSubsystem.cpp
--- a/Engine/Plugins/Runtime/MassGameplay/Source/MassSpawner/Private/MassSpawnerSubsystem.cpp
+++ b/Engine/Plugins/Runtime/MassGameplay/Source/MassSpawner/Private/MassSpawnerSubsystem.cpp
@@ -27,9 +27,9 @@ void UMassSpawnerSubsystem::Initialize(FSubsystemCollectionBase& Collection)
 	Collection.InitializeDependency<UMassSimulationSubsystem>();
 
 	UWorld* World = GetWorld();
+	check(World);
 	SimulationSystem = UWorld::GetSubsystem<UMassSimulationSubsystem>(World);
-	EntitySystem = UWorld::GetSubsystem<UMassEntitySubsystem>(World);
-	check(EntitySystem);
+	EntityManager = UE::Mass::Utils::GetEntityManagerChecked(*World).AsShared();
 }
 
 void UMassSpawnerSubsystem::PostInitialize()
@@ -39,7 +39,7 @@ void UMassSpawnerSubsystem::PostInitialize()
 
 void UMassSpawnerSubsystem::SpawnEntities(const FMassEntityTemplate& EntityTemplate, const uint32 NumberToSpawn, TArray<FMassEntityHandle>& OutEntities)
 {
-	check(EntitySystem);
+	check(EntityManager);
 	check(EntityTemplate.IsValid());
 
 	if (NumberToSpawn == 0)
@@ -107,7 +107,7 @@ void UMassSpawnerSubsystem::DestroyEntities(const FMassEntityTemplateID Template
 		return;
 	}
 
-	check(EntitySystem);
+	check(EntityManager);
 	check(SimulationSystem);
 	checkf(!SimulationSystem->GetPhaseManager().IsDuringMassProcessing()
 		, TEXT("%s called while MassEntity processing in progress. This is unsupported and dangerous!"), ANSI_TO_TCHAR(__FUNCTION__));
@@ -117,11 +117,11 @@ void UMassSpawnerSubsystem::DestroyEntities(const FMassEntityTemplateID Template
 
 
 	TArray<FMassArchetypeEntityCollection> EntityCollections;
-	UE::Mass::Utils::CreateEntityCollections(*EntitySystem, Entities, FMassArchetypeEntityCollection::NoDuplicates, EntityCollections);
+	UE::Mass::Utils::CreateEntityCollections(*EntityManager.Get(), Entities, FMassArchetypeEntityCollection::NoDuplicates, EntityCollections);
 
 	for (const FMassArchetypeEntityCollection& Collection : EntityCollections)
 	{
-		EntitySystem->BatchDestroyEntityChunks(Collection);
+		EntityManager->BatchDestroyEntityChunks(Collection);
 	}
 }
 
@@ -151,9 +151,9 @@ UMassProcessor* UMassSpawnerSubsystem::GetSpawnDataInitializer(TSubclassOf<UMass
 
 void UMassSpawnerSubsystem::DoSpawning(const FMassEntityTemplate& EntityTemplate, const int32 NumToSpawn, FConstStructView SpawnData, TSubclassOf<UMassProcessor> InitializerClass, TArray<FMassEntityHandle>& OutEntities)
 {
-	check(EntitySystem);
+	check(EntityManager);
 	check(EntityTemplate.GetArchetype().IsValid());
-	UE_VLOG(this, LogMassSpawner, Log, TEXT("Spawning with EntityTemplate:\n%s"), *EntityTemplate.DebugGetDescription(EntitySystem));
+	UE_VLOG(this, LogMassSpawner, Log, TEXT("Spawning with EntityTemplate:\n%s"), *EntityTemplate.DebugGetDescription(EntityManager.Get()));
 
 	if (NumToSpawn <= 0)
 	{
@@ -170,16 +170,16 @@ void UMassSpawnerSubsystem::DoSpawning(const FMassEntityTemplate& EntityTemplate
 	// 4. "OnEntitiesCreated" notifies will be sent out once the CreationContext gets destroyed (via its destructor).
 
 	TArray<FMassEntityHandle> SpawnedEntities;
-	TSharedRef<UMassEntitySubsystem::FEntityCreationContext> CreationContext = EntitySystem->BatchCreateEntities(EntityTemplate.GetArchetype(), EntityTemplate.GetSharedFragmentValues(), NumToSpawn, SpawnedEntities);
+	TSharedRef<FMassEntityManager::FEntityCreationContext> CreationContext = EntityManager->BatchCreateEntities(EntityTemplate.GetArchetype(), EntityTemplate.GetSharedFragmentValues(), NumToSpawn, SpawnedEntities);
 
 	TConstArrayView<FInstancedStruct> FragmentInstances = EntityTemplate.GetInitialFragmentValues();
-	EntitySystem->BatchSetEntityFragmentsValues(CreationContext->GetEntityCollection(), FragmentInstances);
+	EntityManager->BatchSetEntityFragmentsValues(CreationContext->GetEntityCollection(), FragmentInstances);
 	
 	UMassProcessor* SpawnDataInitializer = SpawnData.IsValid() ? GetSpawnDataInitializer(InitializerClass) : nullptr;
 
 	if (SpawnDataInitializer)
 	{
-		FMassProcessingContext ProcessingContext(*EntitySystem, /*TimeDelta=*/0.0f);
+		FMassProcessingContext ProcessingContext(EntityManager, /*TimeDelta=*/0.0f);
 		ProcessingContext.AuxData = SpawnData;
 		UE::Mass::Executor::RunProcessorsView(MakeArrayView(&SpawnDataInitializer, 1), ProcessingContext, &CreationContext->GetEntityCollection());
 	}
diff --git a/Engine/Plugins/Runtime/MassGameplay/Source/MassSpawner/Public/MassEntityTemplate.h b/Engine/Plugins/Runtime/MassGameplay/Source/MassSpawner/Public/MassEntityTemplate.h
--- a/Engine/Plugins/Runtime/MassGameplay/Source/MassSpawner/Public/MassEntityTemplate.h
+++ b/Engine/Plugins/Runtime/MassGameplay/Source/MassSpawner/Public/MassEntityTemplate.h
@@ -200,8 +200,8 @@ struct MASSSPAWNER_API FMassEntityTemplate
 		}
 	}
 
-	FString DebugGetDescription(UMassEntitySubsystem* EntitySubsystem = nullptr) const;
-	FString DebugGetArchetypeDescription(UMassEntitySubsystem& EntitySubsystem) const;
+	FString DebugGetDescription(FMassEntityManager* EntityManager = nullptr) const;
+	FString DebugGetArchetypeDescription(FMassEntityManager& EntityManager) const;
 
 	template<typename T>
 	bool HasFragment() const
diff --git a/Engine/Plugins/Runtime/MassGameplay/Source/MassSpawner/Public/MassEntityTemplateRegistry.h b/Engine/Plugins/Runtime/MassGameplay/Source/MassSpawner/Public/MassEntityTemplateRegistry.h
--- a/Engine/Plugins/Runtime/MassGameplay/Source/MassSpawner/Public/MassEntityTemplateRegistry.h
+++ b/Engine/Plugins/Runtime/MassGameplay/Source/MassSpawner/Public/MassEntityTemplateRegistry.h
@@ -11,7 +11,6 @@
 
 
 class UWorld;
-class UMassEntitySubsystem;
 class UMassEntityTraitBase;
 
 enum class EFragmentInitialization : uint8
diff --git a/Engine/Plugins/Runtime/MassGameplay/Source/MassSpawner/Public/MassTranslator.h b/Engine/Plugins/Runtime/MassGameplay/Source/MassSpawner/Public/MassTranslator.h
--- a/Engine/Plugins/Runtime/MassGameplay/Source/MassSpawner/Public/MassTranslator.h
+++ b/Engine/Plugins/Runtime/MassGameplay/Source/MassSpawner/Public/MassTranslator.h
@@ -7,8 +7,6 @@
 #include "MassTranslator.generated.h"
 
 
-class UMassEntitySubsystem;
-
 UENUM()
 enum class EMassTranslationDirection : uint8
 {
