Fix CIS issue reported as "undeclared identifier 'MassActors'"

diff --git a/Engine/Plugins/Runtime/MassGameplay/Source/MassActors/Private/MassActorSpawnerSubsystem.cpp b/Engine/Plugins/Runtime/MassGameplay/Source/MassActors/Private/MassActorSpawnerSubsystem.cpp
--- a/Engine/Plugins/Runtime/MassGameplay/Source/MassActors/Private/MassActorSpawnerSubsystem.cpp
+++ b/Engine/Plugins/Runtime/MassGameplay/Source/MassActors/Private/MassActorSpawnerSubsystem.cpp
@@ -9,6 +9,7 @@
 #include "MassActorPoolableInterface.h"
 #include "MassSimulationSubsystem.h"
 #include "Logging/LogScopedVerbosityOverride.h"
+#include "ProfilingDebugging/CsvProfiler.h"
 
 CSV_DEFINE_CATEGORY(MassActors, true);
 
