[MassReplication] replication processor is required to run on game thread since it might need to create UObjects (i.e. SpawnActor) when synchronizing clients and viewers.
This will allow those objects to be properly GC'ed (i.e. not marked as 'async')

diff --git a/Engine/Plugins/Runtime/MassGameplay/Source/MassReplication/Private/MassReplicationProcessor.cpp b/Engine/Plugins/Runtime/MassGameplay/Source/MassReplication/Private/MassReplicationProcessor.cpp
--- a/Engine/Plugins/Runtime/MassGameplay/Source/MassReplication/Private/MassReplicationProcessor.cpp
+++ b/Engine/Plugins/Runtime/MassGameplay/Source/MassReplication/Private/MassReplicationProcessor.cpp
@@ -28,6 +28,10 @@ UMassReplicationProcessor::UMassReplicationProcessor()
 #endif // UE_ALLOW_DEBUG_REPLICATION_BUBBLES_STANDALONE
 
 	ProcessingPhase = EMassProcessingPhase::PostPhysics;
+
+	// Processor might need to create UObjects when synchronizing clients and viewers
+	// (e.g. SpawnActor from UMassReplicationSubsystem::SynchronizeClientsAndViewers())
+	bRequiresGameThreadExecution = true;
 }
 
 void UMassReplicationProcessor::ConfigureQueries()
