Made UMassSignalProcessorBase::OnSignalReceived thread safe.

Note: using spin lock rather than a critical section since the contention will be low and in such cases spin lock is a bit more efficient

Also, added MT access detector to check if anything is trying to add signals to process while they're being processed.

diff --git a/Engine/Plugins/Runtime/MassGameplay/Source/MassSignals/Private/MassSignalProcessorBase.cpp b/Engine/Plugins/Runtime/MassGameplay/Source/MassSignals/Private/MassSignalProcessorBase.cpp
--- a/Engine/Plugins/Runtime/MassGameplay/Source/MassSignals/Private/MassSignalProcessorBase.cpp
+++ b/Engine/Plugins/Runtime/MassGameplay/Source/MassSignals/Private/MassSignalProcessorBase.cpp
@@ -4,6 +4,7 @@
 #include "MassSignalSubsystem.h"
 #include "MassArchetypeTypes.h"
 #include "Engine/World.h"
+#include "Misc/ScopeLock.h"
 
 
 UMassSignalProcessorBase::UMassSignalProcessorBase(const FObjectInitializer& ObjectInitializer)
@@ -35,6 +36,7 @@ void UMassSignalProcessorBase::Execute(FMassEntityManager& EntityManager, FMassE
 {
 	QUICK_SCOPE_CYCLE_COUNTER(SignalEntities);
 
+	UE_MT_ACQUIRE_WRITE_ACCESS(ReceivedSignalAccessDetector);
 	// Frame buffer handling
 	const int32 ProcessingFrameBufferIndex = CurrentFrameBufferIndex;
 	CurrentFrameBufferIndex = (CurrentFrameBufferIndex + 1) % 2;
@@ -42,6 +44,8 @@ void UMassSignalProcessorBase::Execute(FMassEntityManager& EntityManager, FMassE
 	TArray<FEntitySignalRange>& ReceivedSignalRanges = ProcessingFrameBuffer.ReceivedSignalRanges;
 	TArray<FMassEntityHandle>& SignaledEntities = ProcessingFrameBuffer.SignaledEntities;
 
+	UE_MT_RELEASE_WRITE_ACCESS(ReceivedSignalAccessDetector);
+
 	if (ReceivedSignalRanges.IsEmpty())
 	{
 		return;
@@ -132,6 +136,10 @@ void UMassSignalProcessorBase::Execute(FMassEntityManager& EntityManager, FMassE
 
 void UMassSignalProcessorBase::OnSignalReceived(FName SignalName, TConstArrayView<FMassEntityHandle> Entities)
 {
+	UE::TScopeLock<UE::FSpinLock> ScopeLock(ReceivedSignalLock);
+
+	UE_MT_SCOPED_WRITE_ACCESS(ReceivedSignalAccessDetector);
+
 	FFrameReceivedSignals& CurrentFrameBuffer = FrameReceivedSignals[CurrentFrameBufferIndex];
 
 	FEntitySignalRange& Range = CurrentFrameBuffer.ReceivedSignalRanges.AddDefaulted_GetRef();
diff --git a/Engine/Plugins/Runtime/MassGameplay/Source/MassSignals/Public/MassSignalProcessorBase.h b/Engine/Plugins/Runtime/MassGameplay/Source/MassSignals/Public/MassSignalProcessorBase.h
--- a/Engine/Plugins/Runtime/MassGameplay/Source/MassSignals/Public/MassSignalProcessorBase.h
+++ b/Engine/Plugins/Runtime/MassGameplay/Source/MassSignals/Public/MassSignalProcessorBase.h
@@ -6,6 +6,7 @@
 #include "MassEntityManager.h"
 #include "MassSignalTypes.h"
 #include "Containers/StaticArray.h"
+#include "Misc/SpinLock.h"
 #include "MassSignalProcessorBase.generated.h"
 
 class UMassSignalSubsystem;
@@ -89,5 +90,8 @@ private:
 
 	/** List of all the registered signal names*/
 	TArray<FName> RegisteredSignals;
+
+	UE_MT_DECLARE_RW_ACCESS_DETECTOR(ReceivedSignalAccessDetector);
+	UE::FSpinLock ReceivedSignalLock;
 };
 
