Renamed GET_MASS_CONFIG_VALUE macro in MassSimulationSettings to GET_MASSSIMULATION_CONFIG_VALUE to avoid confusion with the MassEntitySettings' one

diff --git a/Engine/Plugins/Runtime/MassGameplay/Source/MassActors/Private/MassActorSpawnerSubsystem.cpp b/Engine/Plugins/Runtime/MassGameplay/Source/MassActors/Private/MassActorSpawnerSubsystem.cpp
--- a/Engine/Plugins/Runtime/MassGameplay/Source/MassActors/Private/MassActorSpawnerSubsystem.cpp
+++ b/Engine/Plugins/Runtime/MassGameplay/Source/MassActors/Private/MassActorSpawnerSubsystem.cpp
@@ -375,7 +375,7 @@ void UMassActorSpawnerSubsystem::OnPrePhysicsPhaseStarted(const float DeltaSecon
 	// 
 	// Note: MassRepresentationProcessor relies on actor spawns being processed before it runs so it can confirm the
 	// spawn and clean up the previous representation  
-	ProcessPendingSpawningRequest(GET_MASS_CONFIG_VALUE(DesiredActorSpawningTimeSlicePerTick));
+	ProcessPendingSpawningRequest(GET_MASSSIMULATION_CONFIG_VALUE(DesiredActorSpawningTimeSlicePerTick));
 
 	CSV_CUSTOM_STAT(MassActors, NumSpawned, NumActorSpawned, ECsvCustomStatOp::Accumulate);
 	CSV_CUSTOM_STAT(MassActors, NumPooled, NumActorPooled, ECsvCustomStatOp::Accumulate);
@@ -388,7 +388,7 @@ void UMassActorSpawnerSubsystem::OnPrePhysicsPhaseFinished(const float DeltaSeco
 	// Note: MassRepresentationProcessor relies on actor destruction processing after it runs so it can clean up
 	// unwanted actor representations that it has replaced for this frame. It also relies on this running before physics
 	// so unwanted representations don't interfere with new physics enabled actors 
-	ProcessPendingDestruction(GET_MASS_CONFIG_VALUE(DesiredActorDestructionTimeSlicePerTick));
+	ProcessPendingDestruction(GET_MASSSIMULATION_CONFIG_VALUE(DesiredActorDestructionTimeSlicePerTick));
 }
 
 void UMassActorSpawnerSubsystem::AddReferencedObjects(UObject* InThis, FReferenceCollector& Collector)
diff --git a/Engine/Plugins/Runtime/MassGameplay/Source/MassRepresentation/Private/MassRepresentationSubsystem.cpp b/Engine/Plugins/Runtime/MassGameplay/Source/MassRepresentation/Private/MassRepresentationSubsystem.cpp
--- a/Engine/Plugins/Runtime/MassGameplay/Source/MassRepresentation/Private/MassRepresentationSubsystem.cpp
+++ b/Engine/Plugins/Runtime/MassGameplay/Source/MassRepresentation/Private/MassRepresentationSubsystem.cpp
@@ -335,8 +335,8 @@ void UMassRepresentationSubsystem::Initialize(FSubsystemCollectionBase& Collecti
 		AgentSystem->GetOnMassAgentComponentEntityDetaching().AddUObject(this, &UMassRepresentationSubsystem::OnMassAgentComponentEntityDetaching);
 	}
 
-	RetryMovedDistanceSq = FMath::Square(GET_MASS_CONFIG_VALUE(DesiredActorFailedSpawningRetryMoveDistance));
-	RetryTimeInterval = GET_MASS_CONFIG_VALUE(DesiredActorFailedSpawningRetryTimeInterval);
+	RetryMovedDistanceSq = FMath::Square(GET_MASSSIMULATION_CONFIG_VALUE(DesiredActorFailedSpawningRetryMoveDistance));
+	RetryTimeInterval = GET_MASSSIMULATION_CONFIG_VALUE(DesiredActorFailedSpawningRetryTimeInterval);
 }
 
 void UMassRepresentationSubsystem::Deinitialize()
diff --git a/Engine/Plugins/Runtime/MassGameplay/Source/MassSimulation/Private/MassSimulationSubsystem.cpp b/Engine/Plugins/Runtime/MassGameplay/Source/MassSimulation/Private/MassSimulationSubsystem.cpp
--- a/Engine/Plugins/Runtime/MassGameplay/Source/MassSimulation/Private/MassSimulationSubsystem.cpp
+++ b/Engine/Plugins/Runtime/MassGameplay/Source/MassSimulation/Private/MassSimulationSubsystem.cpp
@@ -137,7 +137,7 @@ void UMassSimulationSubsystem::OnProcessingPhaseStarted(const float DeltaSeconds
 				check(EntityManager);
 				if (UE::MassSimulation::bDoEntityCompaction)
 				{
-					EntityManager->DoEntityCompaction(GET_MASS_CONFIG_VALUE(DesiredEntityCompactionTimeSlicePerTick));
+					EntityManager->DoEntityCompaction(GET_MASSSIMULATION_CONFIG_VALUE(DesiredEntityCompactionTimeSlicePerTick));
 				}
 			}
 			break;
diff --git a/Engine/Plugins/Runtime/MassGameplay/Source/MassSimulation/Public/MassSimulationSettings.h b/Engine/Plugins/Runtime/MassGameplay/Source/MassSimulation/Public/MassSimulationSettings.h
--- a/Engine/Plugins/Runtime/MassGameplay/Source/MassSimulation/Public/MassSimulationSettings.h
+++ b/Engine/Plugins/Runtime/MassGameplay/Source/MassSimulation/Public/MassSimulationSettings.h
@@ -5,7 +5,7 @@
 #include "MassSettings.h"
 #include "MassSimulationSettings.generated.h"
 
-#define GET_MASS_CONFIG_VALUE(a) (GetMutableDefault<UMassSimulationSettings>()->a)
+#define GET_MASSSIMULATION_CONFIG_VALUE(a) (GetMutableDefault<UMassSimulationSettings>()->a)
 
 /**
  * Implements the settings for MassSimulation
