Removed FMassProcessingPhase.IsDuringMassProcessing and FMassProcessingPhaseManager.IsDuringMassProcessing in favour of relying on FMassEntityManager.IsProcessing.
Added a convenience UMassSimulationSubsystem.IsDuringMassProcessing passing the call over to hosted FMassEntityManager.
Left FMassProcessingPhase.bIsDuringMassProcessing behind for debugging purposes.

diff --git a/Engine/Plugins/Runtime/MassEntity/Source/MassEntity/Public/MassEntityManager.h b/Engine/Plugins/Runtime/MassEntity/Source/MassEntity/Public/MassEntityManager.h
--- a/Engine/Plugins/Runtime/MassEntity/Source/MassEntity/Public/MassEntityManager.h
+++ b/Engine/Plugins/Runtime/MassEntity/Source/MassEntity/Public/MassEntityManager.h
@@ -321,7 +321,14 @@ public:
 	FMassExecutionContext CreateExecutionContext(const float DeltaSeconds);
 
 	FScopedProcessing NewProcessingScope() { return FScopedProcessing(ProcessingScopeCount); }
+
+	/** 
+	 * Indicates whether there are processors out there performing operations on this instance of MassEntityManager. 
+	 * Used to ensure that mutating operations (like entity destruction) are not performed while processors are running, 
+	 * which rely on the assumption that the data layout doesn't change during calculations. 
+	 */
 	bool IsProcessing() const { return ProcessingScopeCount > 0; }
+
 	FMassCommandBuffer& Defer() const { return *DeferredCommandBuffer.Get(); }
 	/** 
 	 * @param InCommandBuffer if not set then the default command buffer will be flushed. If set and there's already 
diff --git a/Engine/Plugins/Runtime/MassEntity/Source/MassEntity/Public/MassProcessingPhaseManager.h b/Engine/Plugins/Runtime/MassEntity/Source/MassEntity/Public/MassProcessingPhaseManager.h
--- a/Engine/Plugins/Runtime/MassEntity/Source/MassEntity/Public/MassProcessingPhaseManager.h
+++ b/Engine/Plugins/Runtime/MassEntity/Source/MassEntity/Public/MassProcessingPhaseManager.h
@@ -69,8 +69,6 @@ protected:
 	void ConfigureForSingleThreadMode() { bRunInParallelMode = false; }
 
 public:
-	bool IsDuringMassProcessing() const { return bIsDuringMassProcessing; }
-
 	void Initialize(FMassProcessingPhaseManager& InPhaseManager, const EMassProcessingPhase InPhase, const ETickingGroup InTickGroup, UMassCompositeProcessor& InPhaseProcessor);
 
 protected:
@@ -150,14 +148,6 @@ public:
 	void Stop();
 	bool IsRunning() const { return EntityManager.IsValid(); }
 
-	/** 
-	 *  returns true when called while any of the ProcessingPhases is actively executing its processors. Used to 
-	 *  determine whether it's safe to do entity-related operations like adding fragments.
-	 *  Note that the function will return false while the OnPhaseStart or OnPhaseEnd are being broadcast,
-	 *  the value returned will be `true` only when the entity subsystem is actively engaged 
-	 */
-	bool IsDuringMassProcessing() const { return CurrentPhase != EMassProcessingPhase::MAX && ProcessingPhases[int(CurrentPhase)].IsDuringMassProcessing(); }
-
 	FString GetName() const;
 
 protected:
diff --git a/Engine/Plugins/Runtime/MassGameplay/Source/MassActors/Private/MassAgentSubsystem.cpp b/Engine/Plugins/Runtime/MassGameplay/Source/MassActors/Private/MassAgentSubsystem.cpp
--- a/Engine/Plugins/Runtime/MassGameplay/Source/MassActors/Private/MassAgentSubsystem.cpp
+++ b/Engine/Plugins/Runtime/MassGameplay/Source/MassActors/Private/MassAgentSubsystem.cpp
@@ -182,8 +182,7 @@ void UMassAgentSubsystem::UnregisterAgentComponent(UMassAgentComponent& AgentCom
 				const FMassEntityHandle Entity = AgentComp.GetEntityHandle();
 
 			    // remove fragments that have been added for the puppet agent
-			    if (SimulationSystem->GetPhaseManager().IsDuringMassProcessing()
-					|| EntityManager->IsProcessing())
+			    if (EntityManager->IsProcessing())
 			    {
 				    // need to request via command buffer since we can't move entities while processing is happening
 					FMassArchetypeCompositionDescriptor Composition = AgentComp.GetPuppetSpecificAddition();
@@ -223,8 +222,7 @@ void UMassAgentSubsystem::UnregisterAgentComponent(UMassAgentComponent& AgentCom
 			AgentComp.ClearEntityHandle();
 
 			// Destroy the entity
-			if (SimulationSystem->GetPhaseManager().IsDuringMassProcessing()
-				|| EntityManager->IsProcessing())
+			if (EntityManager->IsProcessing())
 			{
 				// need to request via command buffer since we can't move entities while processing is happening
 				EntityManager->Defer().DestroyEntity(Entity);
@@ -303,7 +301,7 @@ void UMassAgentSubsystem::HandlePendingInitialization()
 	check(EntityManager);
 	check(SpawnerSystem);
 	check(SimulationSystem);
-	check(SimulationSystem->GetPhaseManager().IsDuringMassProcessing() == false);
+	check(SimulationSystem->IsDuringMassProcessing() == false);
 
 	for (TTuple<FMassEntityTemplateID, FMassAgentInitializationQueue>& Data : PendingAgentEntities)
 	{
diff --git a/Engine/Plugins/Runtime/MassGameplay/Source/MassSimulation/Private/MassSimulationSubsystem.cpp b/Engine/Plugins/Runtime/MassGameplay/Source/MassSimulation/Private/MassSimulationSubsystem.cpp
--- a/Engine/Plugins/Runtime/MassGameplay/Source/MassSimulation/Private/MassSimulationSubsystem.cpp
+++ b/Engine/Plugins/Runtime/MassGameplay/Source/MassSimulation/Private/MassSimulationSubsystem.cpp
@@ -66,6 +66,11 @@ FMassProcessingPhase::FOnPhaseEvent& UMassSimulationSubsystem::GetOnProcessingPh
 	return PhaseManager.GetOnPhaseEnd(Phase);
 }
 
+bool UMassSimulationSubsystem::IsDuringMassProcessing() const 
+{
+	return EntityManager.IsValid() && EntityManager->IsProcessing();
+}
+
 void UMassSimulationSubsystem::Initialize(FSubsystemCollectionBase& Collection)
 {
 	Super::Initialize(Collection);
diff --git a/Engine/Plugins/Runtime/MassGameplay/Source/MassSimulation/Public/MassSimulationSubsystem.h b/Engine/Plugins/Runtime/MassGameplay/Source/MassSimulation/Public/MassSimulationSubsystem.h
--- a/Engine/Plugins/Runtime/MassGameplay/Source/MassSimulation/Public/MassSimulationSubsystem.h
+++ b/Engine/Plugins/Runtime/MassGameplay/Source/MassSimulation/Public/MassSimulationSubsystem.h
@@ -30,6 +30,9 @@ public:
 
 	bool IsSimulationStarted() const { return bSimulationStarted; }
 
+	/** @return whether hosted EntityManager is currently, actively being used for processing purposes. Equivalent to calling FMassEntityManager.IsProcessing() */
+	bool IsDuringMassProcessing() const;
+
 protected:
 	virtual void Initialize(FSubsystemCollectionBase& Collection) override;
 	virtual void PostInitialize() override;
diff --git a/Engine/Plugins/Runtime/MassGameplay/Source/MassSpawner/Private/MassSpawnerSubsystem.cpp b/Engine/Plugins/Runtime/MassGameplay/Source/MassSpawner/Private/MassSpawnerSubsystem.cpp
--- a/Engine/Plugins/Runtime/MassGameplay/Source/MassSpawner/Private/MassSpawnerSubsystem.cpp
+++ b/Engine/Plugins/Runtime/MassGameplay/Source/MassSpawner/Private/MassSpawnerSubsystem.cpp
@@ -114,7 +114,7 @@ void UMassSpawnerSubsystem::DestroyEntities(const FMassEntityTemplateID Template
 
 	check(EntityManager);
 	check(SimulationSystem);
-	checkf(!SimulationSystem->GetPhaseManager().IsDuringMassProcessing()
+	checkf(!SimulationSystem->IsDuringMassProcessing()
 		, TEXT("%s called while MassEntity processing in progress. This is unsupported and dangerous!"), ANSI_TO_TCHAR(__FUNCTION__));
 
 	UWorld* World = GetWorld();
