Fixed MassEntityConfig validation always showing "error"-styled notification after a single failed notification, regardless of the following results.

diff --git a/Engine/Plugins/Runtime/MassGameplay/Source/MassSpawner/MassSpawner.Build.cs b/Engine/Plugins/Runtime/MassGameplay/Source/MassSpawner/MassSpawner.Build.cs
--- a/Engine/Plugins/Runtime/MassGameplay/Source/MassSpawner/MassSpawner.Build.cs
+++ b/Engine/Plugins/Runtime/MassGameplay/Source/MassSpawner/MassSpawner.Build.cs
@@ -30,7 +30,12 @@ namespace UnrealBuildTool.Rules
 
 			if (Target.bBuildEditor == true)
 			{
-				PrivateDependencyModuleNames.Add("UnrealEd");
+				PrivateDependencyModuleNames.AddRange(
+					new string[] {
+						"UnrealEd",
+						"Slate"
+					}
+				);
 			}
 		}
 	}
diff --git a/Engine/Plugins/Runtime/MassGameplay/Source/MassSpawner/Private/MassEntityConfigAsset.cpp b/Engine/Plugins/Runtime/MassGameplay/Source/MassSpawner/Private/MassEntityConfigAsset.cpp
--- a/Engine/Plugins/Runtime/MassGameplay/Source/MassSpawner/Private/MassEntityConfigAsset.cpp
+++ b/Engine/Plugins/Runtime/MassGameplay/Source/MassSpawner/Private/MassEntityConfigAsset.cpp
@@ -10,6 +10,8 @@
 #include "Engine/World.h"
 #if WITH_EDITOR
 #include "Editor.h"
+#include "Framework/Notifications/NotificationManager.h"
+#include "Widgets/Notifications/SNotificationList.h"
 #endif // WITH_EDITOR
 
 #define LOCTEXT_NAMESPACE "Mass"
@@ -51,9 +53,20 @@ void UMassEntityConfigAsset::ValidateEntityConfig()
 	{
 		if (Config.ValidateEntityTemplate(*EditorWorld, *this))
 		{
+			const FText InfoText = LOCTEXT("MassEntityConfigAssetNoErrorsDetected", "There were no error detected during validation of the EntityConfigAsset");
+			
 			FMessageLog EditorInfo("LogMass");
-			EditorInfo.Info(LOCTEXT("MassEntityConfigAssetNoErrorsDetected", "There were no error detected during validation of the EntityConfigAsset"));
-			EditorInfo.Notify(LOCTEXT("MassEntityConfigAssetNoErrorsDetected", "There were no error detected during validation of the EntityConfigAsset"), EMessageSeverity::Info, true /*bForce*/);
+			EditorInfo.Info(InfoText);
+
+			FNotificationInfo Info(InfoText);
+			Info.bFireAndForget = true;
+			Info.bUseThrobber = false;
+			Info.FadeOutDuration = 0.5f;
+			Info.ExpireDuration = 5.0f;
+			if (TSharedPtr<SNotificationItem> Notification = FSlateNotificationManager::Get().AddNotification(Info))
+			{
+				Notification->SetCompletionState(SNotificationItem::CS_Success);
+			}
 		}
 	}
 }
