Fix usage for Short type name and replace with recommended Long Type name.

diff --git a/Engine/Plugins/Runtime/MassGameplay/Source/MassRepresentation/Public/MassRepresentationFragments.h b/Engine/Plugins/Runtime/MassGameplay/Source/MassRepresentation/Public/MassRepresentationFragments.h
--- a/Engine/Plugins/Runtime/MassGameplay/Source/MassRepresentation/Public/MassRepresentationFragments.h
+++ b/Engine/Plugins/Runtime/MassGameplay/Source/MassRepresentation/Public/MassRepresentationFragments.h
@@ -178,7 +178,7 @@ struct FMassVisualizationLODParameters : public FMassSharedFragment
 	float DistanceToFrustumHysteresis = 0.0f;
 
 	/** Filter these settings with specified tag */
-	UPROPERTY(EditAnywhere, Category = "Mass|LOD", meta = (BaseStruct = "MassTag"))
+	UPROPERTY(EditAnywhere, Category = "Mass|LOD", meta = (BaseStruct = "/Script/MassEntity.MassTag"))
 	TObjectPtr<UScriptStruct> FilterTag = nullptr;
 };
 
