MASS updates to support the Editor Framework's data storage.

The following changes were made:
- Update FMassEntityHandle so it can be cast to/from a plain 64-bit integer. This helps to keep the implementation details for the data storage private and only expose a handle.
- Fixed an issue that caused a crash when retrieving the entity handle associated with an actor while the actor was no longer reachable.

diff --git a/Engine/Plugins/Runtime/MassEntity/Source/MassEntity/Public/MassEntityTypes.h b/Engine/Plugins/Runtime/MassEntity/Source/MassEntity/Public/MassEntityTypes.h
--- a/Engine/Plugins/Runtime/MassEntity/Source/MassEntity/Public/MassEntityTypes.h
+++ b/Engine/Plugins/Runtime/MassEntity/Source/MassEntity/Public/MassEntityTypes.h
@@ -50,7 +50,7 @@ struct FMassSharedFragment
 // A handle to a lightweight entity.  An entity is used in conjunction with the FMassEntityManager
 // for the current world and can contain lightweight fragments.
 USTRUCT()
-struct FMassEntityHandle
+struct alignas(8) FMassEntityHandle
 {
 	GENERATED_BODY()
 
@@ -96,6 +96,16 @@ struct FMassEntityHandle
 		Index = SerialNumber = 0;
 	}
 
+	/** Allows the entity handle to be shared anonymously. */
+	uint64 AsNumber() const { return *reinterpret_cast<const uint64*>(this); } // Relying on the fact that this struct only stores 2 integers and is aligned correctly.
+	/** Reconstruct the entity handle from an anonymously shared integer. */
+	static FMassEntityHandle FromNumber(uint64 Value) 
+	{ 
+		FMassEntityHandle Result;
+		*reinterpret_cast<uint64_t*>(&Result) = Value;
+		return Result;
+	}
+
 	friend uint32 GetTypeHash(const FMassEntityHandle Entity)
 	{
 		return HashCombine(Entity.Index, Entity.SerialNumber);
@@ -107,6 +117,9 @@ struct FMassEntityHandle
 	}
 };
 
+static_assert(sizeof(FMassEntityHandle) == sizeof(uint64), "Expected FMassEntityHandle to be convertable to a 64-bit integer value, so size needs to be 8 bytes.");
+static_assert(alignof(FMassEntityHandle) == sizeof(uint64), "Expected FMassEntityHandle to be convertable to a 64-bit integer value, so alignment needs to be 8 bytes.");
+
 DECLARE_STRUCTTYPEBITSET_EXPORTED(MASSENTITY_API, FMassFragmentBitSet, FMassFragment);
 DECLARE_STRUCTTYPEBITSET_EXPORTED(MASSENTITY_API, FMassTagBitSet, FMassTag);
 DECLARE_STRUCTTYPEBITSET_EXPORTED(MASSENTITY_API, FMassChunkFragmentBitSet, FMassChunkFragment);
diff --git a/Engine/Plugins/Runtime/MassGameplay/Source/MassActors/Private/MassActorSubsystem.cpp b/Engine/Plugins/Runtime/MassGameplay/Source/MassActors/Private/MassActorSubsystem.cpp
--- a/Engine/Plugins/Runtime/MassGameplay/Source/MassActors/Private/MassActorSubsystem.cpp
+++ b/Engine/Plugins/Runtime/MassGameplay/Source/MassActors/Private/MassActorSubsystem.cpp
@@ -54,6 +54,27 @@ void FMassActorFragment::ResetNoHandleMapUpdate()
 	bIsOwnedByMass = false;
 }
 
+AActor* FMassActorFragment::GetMutable(EActorAccess Access)
+{
+	switch (Access)
+	{
+	case EActorAccess::OnlyWhenAlive:
+		return Actor.Get();
+	case EActorAccess::IncludePendingKill:
+		return Actor.Get(true);
+	case EActorAccess::IncludeUnreachable:
+		return Actor.GetEvenIfUnreachable();
+	default:
+		checkf(false, TEXT("Invalid ActorAccess value: %i."), static_cast<int32>(Access));
+		return nullptr;
+	}
+}
+
+const AActor* FMassActorFragment::Get(EActorAccess Access) const
+{
+	return const_cast<FMassActorFragment*>(this)->GetMutable(Access);
+}
+
 //----------------------------------------------------------------------//
 //  UMassActorSubsystem 
 //----------------------------------------------------------------------//
@@ -72,6 +93,7 @@ void UMassActorSubsystem::Initialize(FSubsystemCollectionBase& Collection)
 void UMassActorSubsystem::Deinitialize()
 {
 	EntityManager.Reset();
+	ActorHandleMap.Reset();
 }
 
 FMassEntityHandle UMassActorSubsystem::GetEntityHandleFromActor(const TObjectKey<const AActor> Actor)
@@ -83,15 +105,15 @@ FMassEntityHandle UMassActorSubsystem::GetEntityHandleFromActor(const TObjectKey
 		return FMassEntityManager::InvalidEntity;
 	}
 
-	check(TObjectKey<const AActor>(GetActorFromHandle(*Entity)) == Actor);
+	check(TObjectKey<const AActor>(GetActorFromHandle(*Entity, FMassActorFragment::EActorAccess::IncludeUnreachable)) == Actor);
 	return *Entity;
 }
 
-AActor* UMassActorSubsystem::GetActorFromHandle(const FMassEntityHandle Handle) const
+AActor* UMassActorSubsystem::GetActorFromHandle(const FMassEntityHandle Handle, FMassActorFragment::EActorAccess Access) const
 {
 	check(EntityManager);
 	FMassActorFragment* Data = EntityManager->GetFragmentDataPtr<FMassActorFragment>(Handle);
-	return Data != nullptr ? Data->GetMutable() : nullptr;
+	return Data != nullptr ? Data->GetMutable(Access) : nullptr;
 }
 
 void UMassActorSubsystem::SetHandleForActor(const TObjectKey<const AActor> Actor, const FMassEntityHandle Handle)
diff --git a/Engine/Plugins/Runtime/MassGameplay/Source/MassActors/Public/MassActorSubsystem.h b/Engine/Plugins/Runtime/MassGameplay/Source/MassActors/Public/MassActorSubsystem.h
--- a/Engine/Plugins/Runtime/MassGameplay/Source/MassActors/Public/MassActorSubsystem.h
+++ b/Engine/Plugins/Runtime/MassGameplay/Source/MassActors/Public/MassActorSubsystem.h
@@ -22,6 +22,13 @@ struct MASSACTORS_API FMassActorFragment : public FObjectWrapperFragment
 {
 	GENERATED_BODY()
 
+	enum class EActorAccess
+	{
+		OnlyWhenAlive, // Only return an actor pointer if the actor is alive. This is the default.
+		IncludePendingKill, // Return an actor pointer even if the actor is marked for destruction.
+		IncludeUnreachable // Return an actor pointer even if the actor is unreachable. This implies it's being destroyed.
+	};
+
 	/**
 	 * Set the actor associated to a mass agent, will also keep the map back in MassActorSubsystem up to date.
 	 * @param MassAgent to associated with the actor
@@ -52,12 +59,14 @@ struct MASSACTORS_API FMassActorFragment : public FObjectWrapperFragment
 
 	/** @return none const pointer to the actor	*/
 	FORCEINLINE AActor* GetMutable() { return Actor.Get(); }
+	AActor* GetMutable(EActorAccess Access);
 
 	/** @return none const pointer to the actor	only if owned by mass */
 	FORCEINLINE AActor* GetOwnedByMassMutable() { return bIsOwnedByMass ? Actor.Get() : nullptr; }
 
 	/** @return none const pointer to the actor	only if owned by mass */
 	FORCEINLINE const AActor* Get() const { return Actor.Get(); }
+	const AActor* Get(EActorAccess Access) const;
 
 	/** @return if the actor is owned by mass */
 	FORCEINLINE bool IsOwnedByMass() const { return bIsOwnedByMass; }
@@ -91,7 +100,8 @@ public:
 	void RemoveHandleForActor(const TObjectKey<const AActor> Actor);
 
 	/** Get an actor pointer from a mass handle */
-	AActor* GetActorFromHandle(const FMassEntityHandle Handle) const;
+	AActor* GetActorFromHandle(const FMassEntityHandle Handle, 
+		FMassActorFragment::EActorAccess Access = FMassActorFragment::EActorAccess::OnlyWhenAlive) const;
 
 	/** 
 	 *  Removes the connection between Actor and the given entity. Does all the required book keeping 
