Updated UMassUpdateISMProcessor::ConfigureQueries to not require RW access to FMassRepresentationLODFragment

diff --git a/Engine/Plugins/Runtime/MassGameplay/Source/MassRepresentation/Private/MassUpdateISMProcessor.cpp b/Engine/Plugins/Runtime/MassGameplay/Source/MassRepresentation/Private/MassUpdateISMProcessor.cpp
--- a/Engine/Plugins/Runtime/MassGameplay/Source/MassRepresentation/Private/MassUpdateISMProcessor.cpp
+++ b/Engine/Plugins/Runtime/MassGameplay/Source/MassRepresentation/Private/MassUpdateISMProcessor.cpp
@@ -22,7 +22,7 @@ void UMassUpdateISMProcessor::ConfigureQueries()
 {
 	EntityQuery.AddRequirement<FTransformFragment>(EMassFragmentAccess::ReadOnly);
 	EntityQuery.AddRequirement<FMassRepresentationFragment>(EMassFragmentAccess::ReadWrite);
-	EntityQuery.AddRequirement<FMassRepresentationLODFragment>(EMassFragmentAccess::ReadWrite);
+	EntityQuery.AddRequirement<FMassRepresentationLODFragment>(EMassFragmentAccess::ReadOnly);
 	EntityQuery.AddChunkRequirement<FMassVisualizationChunkFragment>(EMassFragmentAccess::ReadWrite);
 	EntityQuery.SetChunkFilter(&FMassVisualizationChunkFragment::AreAnyEntitiesVisibleInChunk);
 	EntityQuery.AddSharedRequirement<FMassRepresentationSubsystemSharedFragment>(EMassFragmentAccess::ReadWrite);
