[MassGameplay]
* MassEntityConfigAsset - Fixed non-unity linux compile error

diff --git a/Engine/Plugins/Runtime/MassGameplay/Source/MassSpawner/Private/MassEntityConfigAsset.cpp b/Engine/Plugins/Runtime/MassGameplay/Source/MassSpawner/Private/MassEntityConfigAsset.cpp
--- a/Engine/Plugins/Runtime/MassGameplay/Source/MassSpawner/Private/MassEntityConfigAsset.cpp
+++ b/Engine/Plugins/Runtime/MassGameplay/Source/MassSpawner/Private/MassEntityConfigAsset.cpp
@@ -1,6 +1,7 @@
 // Copyright Epic Games, Inc. All Rights Reserved.
 
 #include "MassEntityConfigAsset.h"
+#include "Logging/MessageLog.h"
 #include "MassEntityTraitBase.h"
 #include "MassSpawnerTypes.h"
 #include "MassSpawnerSubsystem.h"
