Fix deprecation

#rnx

diff --git a/Engine/Plugins/Runtime/MassGameplay/Source/MassRepresentation/Private/MassRepresentationSubsystem.cpp b/Engine/Plugins/Runtime/MassGameplay/Source/MassRepresentation/Private/MassRepresentationSubsystem.cpp
--- a/Engine/Plugins/Runtime/MassGameplay/Source/MassRepresentation/Private/MassRepresentationSubsystem.cpp
+++ b/Engine/Plugins/Runtime/MassGameplay/Source/MassRepresentation/Private/MassRepresentationSubsystem.cpp
@@ -155,7 +155,10 @@ bool UMassRepresentationSubsystem::IsCollisionLoaded(const FName TargetGrid, con
 	QuerySource.bSpatialQuery = true;
 	QuerySource.Location = Transform.GetLocation();
 	QuerySource.Rotation = Transform.Rotator();
-	QuerySource.TargetGrid = TargetGrid;
+	if (!TargetGrid.IsNone())
+	{
+		QuerySource.TargetGrids.Add(TargetGrid);
+	}
 	QuerySource.bUseGridLoadingRange = false;
 	QuerySource.Radius = 1.f; // 1cm should be enough to know if grid is loaded at specific area
 	QuerySource.bDataLayersOnly = false;
