Enable a subset of subsystems from MASS in the editor.

A subset of MASS subsystems is enabled for the EditorStorage-type world used in the editor to store data for the Typed Elements.

diff --git a/Engine/Plugins/Runtime/MassEntity/Source/MassEntity/Private/MassEntitySubsystem.cpp b/Engine/Plugins/Runtime/MassEntity/Source/MassEntity/Private/MassEntitySubsystem.cpp
--- a/Engine/Plugins/Runtime/MassEntity/Source/MassEntity/Private/MassEntitySubsystem.cpp
+++ b/Engine/Plugins/Runtime/MassEntity/Source/MassEntity/Private/MassEntitySubsystem.cpp
@@ -40,6 +40,11 @@ void UMassEntitySubsystem::Deinitialize()
 	EntityManager.Reset();
 }
 
+bool UMassEntitySubsystem::DoesSupportWorldType(const EWorldType::Type WorldType) const
+{
+	return WorldType == EWorldType::EditorStorage || Super::DoesSupportWorldType(WorldType);
+}
+
 #if WITH_MASSENTITY_DEBUG
 //////////////////////////////////////////////////////////////////////
 // Debug commands
diff --git a/Engine/Plugins/Runtime/MassEntity/Source/MassEntity/Public/MassEntitySubsystem.h b/Engine/Plugins/Runtime/MassEntity/Source/MassEntity/Public/MassEntitySubsystem.h
--- a/Engine/Plugins/Runtime/MassEntity/Source/MassEntity/Public/MassEntitySubsystem.h
+++ b/Engine/Plugins/Runtime/MassEntity/Source/MassEntity/Public/MassEntitySubsystem.h
@@ -23,6 +23,11 @@ class MASSENTITY_API UMassEntitySubsystem : public UWorldSubsystem
 	virtual void Deinitialize() override;
 	//~End of USubsystem interface
 
+protected:
+	//~USubsystem interface
+	virtual bool DoesSupportWorldType(const EWorldType::Type WorldType) const override;
+	//~End of USubsystem interface
+
 public:
 	UMassEntitySubsystem();
 
diff --git a/Engine/Plugins/Runtime/MassEntity/Source/MassEntityDebugger/Private/SMassDebugger.cpp b/Engine/Plugins/Runtime/MassEntity/Source/MassEntityDebugger/Private/SMassDebugger.cpp
--- a/Engine/Plugins/Runtime/MassEntity/Source/MassEntityDebugger/Private/SMassDebugger.cpp
+++ b/Engine/Plugins/Runtime/MassEntity/Source/MassEntityDebugger/Private/SMassDebugger.cpp
@@ -38,7 +38,11 @@ namespace UE::Mass::Debugger::Private
 
 	bool IsSupportedWorldType(const EWorldType::Type WorldType)
 	{
-		return WorldType == EWorldType::Game || WorldType == EWorldType::Editor || WorldType == EWorldType::PIE;
+		return 
+			WorldType == EWorldType::Game || 
+			WorldType == EWorldType::Editor || 
+			WorldType == EWorldType::PIE || 
+			WorldType == EWorldType::EditorStorage;
 	}
 }
 
diff --git a/Engine/Plugins/Runtime/MassGameplay/Source/MassActors/Private/MassActorSubsystem.cpp b/Engine/Plugins/Runtime/MassGameplay/Source/MassActors/Private/MassActorSubsystem.cpp
--- a/Engine/Plugins/Runtime/MassGameplay/Source/MassActors/Private/MassActorSubsystem.cpp
+++ b/Engine/Plugins/Runtime/MassGameplay/Source/MassActors/Private/MassActorSubsystem.cpp
@@ -96,6 +96,11 @@ void UMassActorSubsystem::Deinitialize()
 	ActorHandleMap.Reset();
 }
 
+bool UMassActorSubsystem::DoesSupportWorldType(const EWorldType::Type WorldType) const
+{
+	return WorldType == EWorldType::Type::EditorStorage || Super::DoesSupportWorldType(WorldType);
+}
+
 FMassEntityHandle UMassActorSubsystem::GetEntityHandleFromActor(const TObjectKey<const AActor> Actor)
 {
 	UE_MT_SCOPED_READ_ACCESS(ActorHandleMapDetector);
diff --git a/Engine/Plugins/Runtime/MassGameplay/Source/MassActors/Public/MassActorSubsystem.h b/Engine/Plugins/Runtime/MassGameplay/Source/MassActors/Public/MassActorSubsystem.h
--- a/Engine/Plugins/Runtime/MassGameplay/Source/MassActors/Public/MassActorSubsystem.h
+++ b/Engine/Plugins/Runtime/MassGameplay/Source/MassActors/Public/MassActorSubsystem.h
@@ -113,6 +113,7 @@ protected:
 	// USubsystem BEGIN
 	virtual void Initialize(FSubsystemCollectionBase& Collection) override;
 	virtual void Deinitialize() override;
+	virtual bool DoesSupportWorldType(const EWorldType::Type WorldType) const override;
 	// USubsystem END
 	
 	TMap<TObjectKey<const AActor>, FMassEntityHandle> ActorHandleMap;
diff --git a/Engine/Plugins/Runtime/MassGameplay/Source/MassGameplayDebug/Private/MassDebuggerSubsystem.cpp b/Engine/Plugins/Runtime/MassGameplay/Source/MassGameplayDebug/Private/MassDebuggerSubsystem.cpp
--- a/Engine/Plugins/Runtime/MassGameplay/Source/MassGameplayDebug/Private/MassDebuggerSubsystem.cpp
+++ b/Engine/Plugins/Runtime/MassGameplay/Source/MassGameplayDebug/Private/MassDebuggerSubsystem.cpp
@@ -32,6 +32,15 @@ void UMassDebuggerSubsystem::Deinitialize()
 #endif // WITH_MASSENTITY_DEBUG
 }
 
+bool UMassDebuggerSubsystem::DoesSupportWorldType(const EWorldType::Type WorldType) const
+{
+	return
+#if WITH_MASSENTITY_DEBUG
+		WorldType == EWorldType::EditorStorage ||
+#endif
+		Super::DoesSupportWorldType(WorldType);
+}
+
 #if WITH_MASSENTITY_DEBUG
 void UMassDebuggerSubsystem::OnEntitySelected(const FMassEntityManager& EntityManager, const FMassEntityHandle EntityHandle)
 {
diff --git a/Engine/Plugins/Runtime/MassGameplay/Source/MassGameplayDebug/Public/MassDebuggerSubsystem.h b/Engine/Plugins/Runtime/MassGameplay/Source/MassGameplayDebug/Public/MassDebuggerSubsystem.h
--- a/Engine/Plugins/Runtime/MassGameplay/Source/MassGameplayDebug/Public/MassDebuggerSubsystem.h
+++ b/Engine/Plugins/Runtime/MassGameplay/Source/MassGameplayDebug/Public/MassDebuggerSubsystem.h
@@ -47,6 +47,7 @@ protected:
 	// USubsystem BEGIN
 	virtual void Initialize(FSubsystemCollectionBase& Collection) override;
 	virtual void Deinitialize() override;
+	virtual bool DoesSupportWorldType(const EWorldType::Type WorldType) const override;
 	// USubsystem END
 	
 	void OnProcessingPhaseStarted(const float DeltaSeconds);
diff --git a/Engine/Plugins/Runtime/MassGameplay/Source/MassSimulation/Private/MassSimulationSubsystem.cpp b/Engine/Plugins/Runtime/MassGameplay/Source/MassSimulation/Private/MassSimulationSubsystem.cpp
--- a/Engine/Plugins/Runtime/MassGameplay/Source/MassSimulation/Private/MassSimulationSubsystem.cpp
+++ b/Engine/Plugins/Runtime/MassGameplay/Source/MassSimulation/Private/MassSimulationSubsystem.cpp
@@ -56,6 +56,11 @@ void UMassSimulationSubsystem::BeginDestroy()
 	Super::BeginDestroy();
 }
 
+bool UMassSimulationSubsystem::DoesSupportWorldType(const EWorldType::Type WorldType) const
+{
+	return WorldType == EWorldType::EditorStorage || Super::DoesSupportWorldType(WorldType);
+}
+
 FMassProcessingPhase::FOnPhaseEvent& UMassSimulationSubsystem::GetOnProcessingPhaseStarted(const EMassProcessingPhase Phase)
 {
 	return PhaseManager.GetOnPhaseStart(Phase);
@@ -185,13 +190,18 @@ void UMassSimulationSubsystem::OnProcessingPhaseStarted(const float DeltaSeconds
 #if WITH_EDITOR
 void UMassSimulationSubsystem::OnPieBegin(const bool bIsSimulation)
 {
-	// called so that we're not processing phases for the editor world while PIE/SIE is running
-	StopSimulation();
+	UWorld* World = GetWorld();
+	if (!World || (World && !World->IsStorageWorld()))
+	{
+		// called so that we're not processing phases for the editor world while PIE/SIE is running
+		StopSimulation();
+	}
 }
 
 void UMassSimulationSubsystem::OnPieEnded(const bool bIsSimulation)
 {
-	if (UWorld * World = GetWorld())
+	UWorld* World = GetWorld();
+	if (World && !bSimulationStarted)
 	{
 		// Resume processing phases in the editor world.
 		StartSimulation(*World);
diff --git a/Engine/Plugins/Runtime/MassGameplay/Source/MassSimulation/Public/MassSimulationSubsystem.h b/Engine/Plugins/Runtime/MassGameplay/Source/MassSimulation/Public/MassSimulationSubsystem.h
--- a/Engine/Plugins/Runtime/MassGameplay/Source/MassSimulation/Public/MassSimulationSubsystem.h
+++ b/Engine/Plugins/Runtime/MassGameplay/Source/MassSimulation/Public/MassSimulationSubsystem.h
@@ -39,6 +39,7 @@ protected:
 	virtual void OnWorldBeginPlay(UWorld& InWorld) override;
 	virtual void Deinitialize() override;
 	virtual void BeginDestroy() override;
+	virtual bool DoesSupportWorldType(const EWorldType::Type WorldType) const override;
 
 	void RebuildTickPipeline();
 
diff --git a/Engine/Plugins/Runtime/ZoneGraph/Source/ZoneGraph/Private/ZoneGraphSubsystem.cpp b/Engine/Plugins/Runtime/ZoneGraph/Source/ZoneGraph/Private/ZoneGraphSubsystem.cpp
--- a/Engine/Plugins/Runtime/ZoneGraph/Source/ZoneGraph/Private/ZoneGraphSubsystem.cpp
+++ b/Engine/Plugins/Runtime/ZoneGraph/Source/ZoneGraph/Private/ZoneGraphSubsystem.cpp
@@ -491,3 +491,8 @@ TConstArrayView<FZoneGraphTagInfo> UZoneGraphSubsystem::GetTagInfos() const
 	}
 	return TConstArrayView<FZoneGraphTagInfo>();
 }
+
+bool UZoneGraphSubsystem::DoesSupportWorldType(const EWorldType::Type WorldType) const
+{
+	return WorldType == EWorldType::EditorStorage || Super::DoesSupportWorldType(WorldType);
+}
\ No newline at end of file
diff --git a/Engine/Plugins/Runtime/ZoneGraph/Source/ZoneGraph/Public/ZoneGraphSubsystem.h b/Engine/Plugins/Runtime/ZoneGraph/Source/ZoneGraph/Public/ZoneGraphSubsystem.h
--- a/Engine/Plugins/Runtime/ZoneGraph/Source/ZoneGraph/Public/ZoneGraphSubsystem.h
+++ b/Engine/Plugins/Runtime/ZoneGraph/Source/ZoneGraph/Public/ZoneGraphSubsystem.h
@@ -130,6 +130,7 @@ protected:
 	virtual void Deinitialize() override;
 	virtual void Tick(float DeltaTime) override;
 	virtual TStatId GetStatId() const override;
+	virtual bool DoesSupportWorldType(const EWorldType::Type WorldType) const override;
 
 	void RemoveRegisteredDataItem(const int32 Index);
 	void UnregisterStaleZoneGraphDataInstances();
diff --git a/Engine/Plugins/Runtime/ZoneGraphAnnotations/Source/ZoneGraphAnnotations/Private/ZoneGraphAnnotationSubsystem.cpp b/Engine/Plugins/Runtime/ZoneGraphAnnotations/Source/ZoneGraphAnnotations/Private/ZoneGraphAnnotationSubsystem.cpp
--- a/Engine/Plugins/Runtime/ZoneGraphAnnotations/Source/ZoneGraphAnnotations/Private/ZoneGraphAnnotationSubsystem.cpp
+++ b/Engine/Plugins/Runtime/ZoneGraphAnnotations/Source/ZoneGraphAnnotations/Private/ZoneGraphAnnotationSubsystem.cpp
@@ -197,3 +197,8 @@ TStatId UZoneGraphAnnotationSubsystem::GetStatId() const
 {
 	RETURN_QUICK_DECLARE_CYCLE_STAT(UZoneGraphAnnotationSubsystem, STATGROUP_Tickables);
 }
+
+bool UZoneGraphAnnotationSubsystem::DoesSupportWorldType(const EWorldType::Type WorldType) const
+{
+	return WorldType == EWorldType::EditorStorage || Super::DoesSupportWorldType(WorldType);
+}
diff --git a/Engine/Plugins/Runtime/ZoneGraphAnnotations/Source/ZoneGraphAnnotations/Public/ZoneGraphAnnotationSubsystem.h b/Engine/Plugins/Runtime/ZoneGraphAnnotations/Source/ZoneGraphAnnotations/Public/ZoneGraphAnnotationSubsystem.h
--- a/Engine/Plugins/Runtime/ZoneGraphAnnotations/Source/ZoneGraphAnnotations/Public/ZoneGraphAnnotationSubsystem.h
+++ b/Engine/Plugins/Runtime/ZoneGraphAnnotations/Source/ZoneGraphAnnotations/Public/ZoneGraphAnnotationSubsystem.h
@@ -120,6 +120,8 @@ protected:
 	virtual void Tick(float DeltaTime) override;
 	virtual TStatId GetStatId() const override;
 
+	virtual bool DoesSupportWorldType(const EWorldType::Type WorldType) const override;
+
 	/** Array of registered components. */
 	UPROPERTY(Transient)
 	TArray<FRegisteredZoneGraphAnnotation> RegisteredComponents;
