Changed the behavior of UMassEntityTemplateRegistry::InitializeEntityTemplate in regards to empty-composition entity templates. We now allow them and create an "empty" (i.e. fragment-less) archetype.

diff --git a/Engine/Plugins/Runtime/MassGameplay/Source/MassSpawner/Private/MassEntityTemplateRegistry.cpp b/Engine/Plugins/Runtime/MassGameplay/Source/MassSpawner/Private/MassEntityTemplateRegistry.cpp
--- a/Engine/Plugins/Runtime/MassGameplay/Source/MassSpawner/Private/MassEntityTemplateRegistry.cpp
+++ b/Engine/Plugins/Runtime/MassGameplay/Source/MassSpawner/Private/MassEntityTemplateRegistry.cpp
@@ -145,25 +145,18 @@ bool UMassEntityTemplateRegistry::BuildTemplateImpl(const FStructToTemplateBuild
 	return false;
 }
 
-void UMassEntityTemplateRegistry::InitializeEntityTemplate(FMassEntityTemplate& OutTemplate) const
+void UMassEntityTemplateRegistry::InitializeEntityTemplate(FMassEntityTemplate& InOutTemplate) const
 {
-	if (OutTemplate.IsEmpty() == true)
-	{
-		// nothing to initialize, return.
-		// this will leave the ArchetypeHandle invalid
-		return;
-	}
-
 	UWorld* World = GetWorld();
 	check(World);
 	// find or create template
 	FMassEntityManager& EntityManager = UE::Mass::Utils::GetEntityManagerChecked(*World);
 
 	// Sort anything there is to sort for later comparison purposes
-	OutTemplate.Sort();
+	InOutTemplate.Sort();
 
-	const FMassArchetypeHandle ArchetypeHandle = EntityManager.CreateArchetype(OutTemplate.GetCompositionDescriptor(), FName(OutTemplate.GetTemplateName()));
-	OutTemplate.SetArchetype(ArchetypeHandle);
+	const FMassArchetypeHandle ArchetypeHandle = EntityManager.CreateArchetype(InOutTemplate.GetCompositionDescriptor(), FName(InOutTemplate.GetTemplateName()));
+	InOutTemplate.SetArchetype(ArchetypeHandle);
 }
 
 void UMassEntityTemplateRegistry::DebugReset()
diff --git a/Engine/Plugins/Runtime/MassGameplay/Source/MassSpawner/Public/MassEntityTemplate.h b/Engine/Plugins/Runtime/MassGameplay/Source/MassSpawner/Public/MassEntityTemplate.h
--- a/Engine/Plugins/Runtime/MassGameplay/Source/MassSpawner/Public/MassEntityTemplate.h
+++ b/Engine/Plugins/Runtime/MassGameplay/Source/MassSpawner/Public/MassEntityTemplate.h
@@ -79,7 +79,7 @@ struct MASSSPAWNER_API FMassEntityTemplate
 	TConstArrayView<FObjectFragmentInitializerFunction> GetObjectFragmentInitializers() const { return ObjectInitializers; }
 	TArray<FObjectFragmentInitializerFunction>& GetMutableObjectFragmentInitializers() { return ObjectInitializers; }
 
-	bool IsValid() const { return Archetype.IsValid() && (Composition.IsEmpty() == false); }
+	bool IsValid() const { return Archetype.IsValid(); }
 	bool IsEmpty() const { return Composition.IsEmpty(); }
 
 	void SetTemplateID(FMassEntityTemplateID InTemplateID) { TemplateID = InTemplateID; }
diff --git a/Engine/Plugins/Runtime/MassGameplay/Source/MassSpawner/Public/MassEntityTemplateRegistry.h b/Engine/Plugins/Runtime/MassGameplay/Source/MassSpawner/Public/MassEntityTemplateRegistry.h
--- a/Engine/Plugins/Runtime/MassGameplay/Source/MassSpawner/Public/MassEntityTemplateRegistry.h
+++ b/Engine/Plugins/Runtime/MassGameplay/Source/MassSpawner/Public/MassEntityTemplateRegistry.h
@@ -222,7 +222,7 @@ public:
 
 	FMassEntityTemplate& CreateTemplate(const uint32 HashLookup, FMassEntityTemplateID TemplateID);
 	void DestroyTemplate(const uint32 HashLookup, FMassEntityTemplateID TemplateID);
-	void InitializeEntityTemplate(FMassEntityTemplate& OutTemplate) const;
+	void InitializeEntityTemplate(FMassEntityTemplate& InOutTemplate) const;
 
 protected:
 	/** @return true if a template has been built, false otherwise */
