Fixed compile errors when building QAGame with orphaned headers on

diff --git a/Engine/Plugins/AI/MassAI/Source/MassAIBehavior/Public/MassZoneGraphAnnotationTypes.h b/Engine/Plugins/AI/MassAI/Source/MassAIBehavior/Public/MassZoneGraphAnnotationTypes.h
--- a/Engine/Plugins/AI/MassAI/Source/MassAIBehavior/Public/MassZoneGraphAnnotationTypes.h
+++ b/Engine/Plugins/AI/MassAI/Source/MassAIBehavior/Public/MassZoneGraphAnnotationTypes.h
@@ -2,6 +2,8 @@
 
 #pragma once
 
+#include "UObject/NameTypes.h"
+
 namespace UE::Mass::Signals
 {
 	const FName AnnotationTagsChanged = FName(TEXT("AnnotationTagsChanged"));
diff --git a/Engine/Plugins/AI/MassAI/Source/MassNavigation/Public/MassNavigationUtils.h b/Engine/Plugins/AI/MassAI/Source/MassNavigation/Public/MassNavigationUtils.h
--- a/Engine/Plugins/AI/MassAI/Source/MassNavigation/Public/MassNavigationUtils.h
+++ b/Engine/Plugins/AI/MassAI/Source/MassNavigation/Public/MassNavigationUtils.h
@@ -3,6 +3,7 @@
 #pragma once
 
 #include "CoreTypes.h"
+#include "Math/Quat.h"
 #include "Math/UnrealMathUtility.h"
 
 namespace UE::MassNavigation
diff --git a/Engine/Plugins/Compositing/Composure/Source/ComposureLayersEditor/Private/ComposureEditorStyle.h b/Engine/Plugins/Compositing/Composure/Source/ComposureLayersEditor/Private/ComposureEditorStyle.h
--- a/Engine/Plugins/Compositing/Composure/Source/ComposureLayersEditor/Private/ComposureEditorStyle.h
+++ b/Engine/Plugins/Compositing/Composure/Source/ComposureLayersEditor/Private/ComposureEditorStyle.h
@@ -4,6 +4,7 @@
 
 #include "Styling/SlateStyle.h"
 #include "Styling/SlateStyleRegistry.h"
+#include "Styling/SlateTypes.h"
 
 #define IMAGE_BRUSH(RelativePath, ...) FSlateImageBrush(RootToContentDir(RelativePath, TEXT(".png")), __VA_ARGS__)
 #define BOX_BRUSH(RelativePath, ...) FSlateBoxBrush(RootToContentDir(RelativePath, TEXT(".png")), __VA_ARGS__)
diff --git a/Engine/Plugins/Editor/ConsoleVariablesEditor/Source/ConsoleVariablesEditor/Private/Views/Widgets/SConsoleVariablesEditorTooltipWidget.h b/Engine/Plugins/Editor/ConsoleVariablesEditor/Source/ConsoleVariablesEditor/Private/Views/Widgets/SConsoleVariablesEditorTooltipWidget.h
--- a/Engine/Plugins/Editor/ConsoleVariablesEditor/Source/ConsoleVariablesEditor/Private/Views/Widgets/SConsoleVariablesEditorTooltipWidget.h
+++ b/Engine/Plugins/Editor/ConsoleVariablesEditor/Source/ConsoleVariablesEditor/Private/Views/Widgets/SConsoleVariablesEditorTooltipWidget.h
@@ -4,6 +4,7 @@
 
 #include "Styling/AppStyle.h"
 #include "Widgets/Layout/SBorder.h"
+#include "Widgets/Layout/SBox.h"
 #include "Widgets/SBoxPanel.h"
 #include "Widgets/SCompoundWidget.h"
 #include "Widgets/SToolTip.h"
diff --git a/Engine/Plugins/Importers/USDImporter/Source/GeometryCacheUSD/Public/GeometryCacheTrackUSDTypes.h b/Engine/Plugins/Importers/USDImporter/Source/GeometryCacheUSD/Public/GeometryCacheTrackUSDTypes.h
--- a/Engine/Plugins/Importers/USDImporter/Source/GeometryCacheUSD/Public/GeometryCacheTrackUSDTypes.h
+++ b/Engine/Plugins/Importers/USDImporter/Source/GeometryCacheUSD/Public/GeometryCacheTrackUSDTypes.h
@@ -3,6 +3,7 @@
 #pragma once
 
 #include "Templates/Function.h"
+#include "UObject/WeakObjectPtr.h"
 
 class UGeometryCacheTrackUsd;
 struct FGeometryCacheMeshData;
diff --git a/Engine/Plugins/Importers/USDImporter/Source/USDStageEditor/Private/USDDragDropOp.h b/Engine/Plugins/Importers/USDImporter/Source/USDStageEditor/Private/USDDragDropOp.h
--- a/Engine/Plugins/Importers/USDImporter/Source/USDStageEditor/Private/USDDragDropOp.h
+++ b/Engine/Plugins/Importers/USDImporter/Source/USDStageEditor/Private/USDDragDropOp.h
@@ -3,6 +3,7 @@
 #pragma once
 
 #include "DragAndDrop/DecoratedDragDropOp.h"
+#include "Widgets/IUSDTreeViewItem.h"
 
 enum class EUsdDragDropOpType
 {
diff --git a/Engine/Plugins/Media/PixelCapture/Source/PixelCapture/Public/PixelCaptureBufferFormat.h b/Engine/Plugins/Media/PixelCapture/Source/PixelCapture/Public/PixelCaptureBufferFormat.h
--- a/Engine/Plugins/Media/PixelCapture/Source/PixelCapture/Public/PixelCaptureBufferFormat.h
+++ b/Engine/Plugins/Media/PixelCapture/Source/PixelCapture/Public/PixelCaptureBufferFormat.h
@@ -2,6 +2,8 @@
 
 #pragma once
 
+#include "CoreTypes.h"
+
 /**
  * This is a list of all internal buffer formats supported by Pixel Capture. To implement
  * your own buffer format value you must start your values at FORMAT_USER.
diff --git a/Engine/Plugins/Media/PixelStreaming/Source/PixelStreaming/Private/FrameBufferRHI.h b/Engine/Plugins/Media/PixelStreaming/Source/PixelStreaming/Private/FrameBufferRHI.h
--- a/Engine/Plugins/Media/PixelStreaming/Source/PixelStreaming/Private/FrameBufferRHI.h
+++ b/Engine/Plugins/Media/PixelStreaming/Source/PixelStreaming/Private/FrameBufferRHI.h
@@ -2,7 +2,7 @@
 
 #pragma once
 
-#include "RHIResources.h"
+#include "RHI.h"
 #include "WebRTCIncludes.h"
 
 namespace UE::PixelStreaming
diff --git a/Engine/Plugins/Media/PixelStreaming/Source/PixelStreaming/Private/ThreadSafeMap.h b/Engine/Plugins/Media/PixelStreaming/Source/PixelStreaming/Private/ThreadSafeMap.h
--- a/Engine/Plugins/Media/PixelStreaming/Source/PixelStreaming/Private/ThreadSafeMap.h
+++ b/Engine/Plugins/Media/PixelStreaming/Source/PixelStreaming/Private/ThreadSafeMap.h
@@ -2,6 +2,10 @@
 
 #pragma once
 
+#include "Containers/Map.h"
+#include "HAL/CriticalSection.h"
+#include "Misc/ScopeLock.h"
+
 namespace UE::PixelStreaming
 {
 	template <typename KeyType, typename ValueType>
diff --git a/Engine/Plugins/Media/PixelStreaming/Source/PixelStreaming/Public/IPixelStreamingAudioInput.h b/Engine/Plugins/Media/PixelStreaming/Source/PixelStreaming/Public/IPixelStreamingAudioInput.h
--- a/Engine/Plugins/Media/PixelStreaming/Source/PixelStreaming/Public/IPixelStreamingAudioInput.h
+++ b/Engine/Plugins/Media/PixelStreaming/Source/PixelStreaming/Public/IPixelStreamingAudioInput.h
@@ -2,6 +2,8 @@
 
 #pragma once
 
+#include "CoreTypes.h"
+
 class PIXELSTREAMING_API IPixelStreamingAudioInput
 {
 public:
diff --git a/Engine/Plugins/Media/PixelStreaming/Source/PixelStreaming/Public/PixelStreamingPlayerConfig.h b/Engine/Plugins/Media/PixelStreaming/Source/PixelStreaming/Public/PixelStreamingPlayerConfig.h
--- a/Engine/Plugins/Media/PixelStreaming/Source/PixelStreaming/Public/PixelStreamingPlayerConfig.h
+++ b/Engine/Plugins/Media/PixelStreaming/Source/PixelStreaming/Public/PixelStreamingPlayerConfig.h
@@ -2,6 +2,8 @@
 
 #pragma once
 
+#include "CoreTypes.h"
+
 struct PIXELSTREAMING_API FPixelStreamingPlayerConfig
 {
 	bool SupportsDataChannel = false;
diff --git a/Engine/Plugins/Media/PixelStreaming/Source/PixelStreamingEditor/Public/PixelStreamingEditorUtils.h b/Engine/Plugins/Media/PixelStreaming/Source/PixelStreamingEditor/Public/PixelStreamingEditorUtils.h
--- a/Engine/Plugins/Media/PixelStreaming/Source/PixelStreamingEditor/Public/PixelStreamingEditorUtils.h
+++ b/Engine/Plugins/Media/PixelStreaming/Source/PixelStreamingEditor/Public/PixelStreamingEditorUtils.h
@@ -2,6 +2,8 @@
 
 #pragma once
 
+#include "Containers/UnrealString.h"
+
 namespace UE::EditorPixelStreaming
 {
 	enum class PIXELSTREAMINGEDITOR_API EStreamTypes : uint8
diff --git a/Engine/Plugins/Runtime/MassGameplay/Source/MassCommon/Public/RandomSequence.h b/Engine/Plugins/Runtime/MassGameplay/Source/MassCommon/Public/RandomSequence.h
--- a/Engine/Plugins/Runtime/MassGameplay/Source/MassCommon/Public/RandomSequence.h
+++ b/Engine/Plugins/Runtime/MassGameplay/Source/MassCommon/Public/RandomSequence.h
@@ -1,4 +1,4 @@
-﻿// Copyright Epic Games, Inc. All Rights Reserved.
+// Copyright Epic Games, Inc. All Rights Reserved.
 
 #pragma once
 
@@ -30,9 +30,9 @@ namespace UE::RandomSequence
 	 */
 	FORCEINLINE uint32 FibocciHash(const int32 SeqIndex)
 	{
-		constexpr uint32_t K = 2654435769u; // 2^32 / phi (golden ratio)
+		constexpr uint32 K = 2654435769u; // 2^32 / phi (golden ratio)
 		// Offset the sequence by 1, so that index 0 is not always 0.
-		return (uint32_t)(SeqIndex + 1) * K;
+		return (uint32)(SeqIndex + 1) * K;
 	}
 	
 	/**
@@ -42,7 +42,7 @@ namespace UE::RandomSequence
 	FORCEINLINE float FRand(const int32 SeqIndex)
 	{
 		float Result;
-		*(uint32_t*)&Result = 0x3F800000U | (FibocciHash(SeqIndex) >> 9);
+		*(uint32*)&Result = 0x3F800000U | (FibocciHash(SeqIndex) >> 9);
 		return Result - 1.0f; 
 	}
 
diff --git a/Engine/Plugins/Runtime/MassGameplay/Source/MassMovementEditor/Public/MassMovementPropertyUtils.h b/Engine/Plugins/Runtime/MassGameplay/Source/MassMovementEditor/Public/MassMovementPropertyUtils.h
--- a/Engine/Plugins/Runtime/MassGameplay/Source/MassMovementEditor/Public/MassMovementPropertyUtils.h
+++ b/Engine/Plugins/Runtime/MassGameplay/Source/MassMovementEditor/Public/MassMovementPropertyUtils.h
@@ -2,8 +2,7 @@
 #pragma once
 #include "CoreMinimal.h"
 #include "Editor.h"
-
-class IPropertyHandle;
+#include "PropertyHandle.h"
 
 #define LOCTEXT_NAMESPACE "MassMovementEditor"
 
diff --git a/Engine/Plugins/Runtime/ZoneGraph/Source/ZoneGraph/Public/ZoneGraphObjectCRC32.h b/Engine/Plugins/Runtime/ZoneGraph/Source/ZoneGraph/Public/ZoneGraphObjectCRC32.h
--- a/Engine/Plugins/Runtime/ZoneGraph/Source/ZoneGraph/Public/ZoneGraphObjectCRC32.h
+++ b/Engine/Plugins/Runtime/ZoneGraph/Source/ZoneGraph/Public/ZoneGraphObjectCRC32.h
@@ -7,6 +7,8 @@
 
 #if WITH_EDITORONLY_DATA
 
+#include "UObject/UnrealType.h"
+
 /**
  * Archive based object hashing to be used with ZoneGraph data calculation.
  * If a property is has "IncludeInHash" meta tag, any of it's child properties will be included in the hash.
diff --git a/Engine/Source/Runtime/AVEncoder/Private/Decoders/vdecmpeg4/M4VlcData.h b/Engine/Source/Runtime/AVEncoder/Private/Decoders/vdecmpeg4/M4VlcData.h
--- a/Engine/Source/Runtime/AVEncoder/Private/Decoders/vdecmpeg4/M4VlcData.h
+++ b/Engine/Source/Runtime/AVEncoder/Private/Decoders/vdecmpeg4/M4VlcData.h
@@ -1,6 +1,8 @@
 // Copyright Epic Games, Inc. All Rights Reserved.
 #pragma once
 
+#include "Decoders/vdecmpeg4/M4VlcDecoder.h"
+
 namespace vdecmpeg4
 {
 
