Extracted the core of UMassActorSubsystem into a new FMassActorManager struct to allow support the logic's reuse in other circumstances (like a dedicated editor subsystem)

diff --git a/Engine/Plugins/Runtime/MassGameplay/Source/MassActors/Private/MassActorSubsystem.cpp b/Engine/Plugins/Runtime/MassGameplay/Source/MassActors/Private/MassActorSubsystem.cpp
--- a/Engine/Plugins/Runtime/MassGameplay/Source/MassActors/Private/MassActorSubsystem.cpp
+++ b/Engine/Plugins/Runtime/MassGameplay/Source/MassActors/Private/MassActorSubsystem.cpp
@@ -86,14 +86,13 @@ void UMassActorSubsystem::Initialize(FSubsystemCollectionBase& Collection)
 	
 	if (UMassEntitySubsystem* EntitySubsystem = UWorld::GetSubsystem<UMassEntitySubsystem>(GetWorld()))
 	{
-		EntityManager = EntitySubsystem->GetMutableEntityManager().AsShared();
+		ActorManager = MakeShareable(new FMassActorManager(EntitySubsystem->GetMutableEntityManager().AsShared()));
 	}
 }
 
 void UMassActorSubsystem::Deinitialize()
 {
-	EntityManager.Reset();
-	ActorHandleMap.Reset();
+	ActorManager.Reset();
 }
 
 bool UMassActorSubsystem::DoesSupportWorldType(const EWorldType::Type WorldType) const
@@ -101,7 +100,17 @@ bool UMassActorSubsystem::DoesSupportWorldType(const EWorldType::Type WorldType)
 	return WorldType == EWorldType::Type::EditorStorage || Super::DoesSupportWorldType(WorldType);
 }
 
-FMassEntityHandle UMassActorSubsystem::GetEntityHandleFromActor(const TObjectKey<const AActor> Actor)
+//----------------------------------------------------------------------//
+//  FMassActorManager
+//----------------------------------------------------------------------//
+FMassActorManager::FMassActorManager(const TSharedPtr<FMassEntityManager>& InEntityManager, UObject* InOwner)
+	: EntityManager(InEntityManager)
+	, Owner(InOwner)
+{
+
+}
+
+FMassEntityHandle FMassActorManager::GetEntityHandleFromActor(const TObjectKey<const AActor> Actor)
 {
 	UE_MT_SCOPED_READ_ACCESS(ActorHandleMapDetector);
 	FMassEntityHandle* Entity = ActorHandleMap.Find(Actor);
@@ -114,26 +123,26 @@ FMassEntityHandle UMassActorSubsystem::GetEntityHandleFromActor(const TObjectKey
 	return *Entity;
 }
 
-AActor* UMassActorSubsystem::GetActorFromHandle(const FMassEntityHandle Handle, FMassActorFragment::EActorAccess Access) const
+AActor* FMassActorManager::GetActorFromHandle(const FMassEntityHandle Handle, FMassActorFragment::EActorAccess Access) const
 {
 	check(EntityManager);
 	FMassActorFragment* Data = EntityManager->GetFragmentDataPtr<FMassActorFragment>(Handle);
 	return Data != nullptr ? Data->GetMutable(Access) : nullptr;
 }
 
-void UMassActorSubsystem::SetHandleForActor(const TObjectKey<const AActor> Actor, const FMassEntityHandle Handle)
+void FMassActorManager::SetHandleForActor(const TObjectKey<const AActor> Actor, const FMassEntityHandle Handle)
 {
 	UE_MT_SCOPED_WRITE_ACCESS(ActorHandleMapDetector);
 	ActorHandleMap.Add(Actor, Handle);
 }
 
-void UMassActorSubsystem::RemoveHandleForActor(const TObjectKey<const AActor> Actor)
+void FMassActorManager::RemoveHandleForActor(const TObjectKey<const AActor> Actor)
 {
 	UE_MT_SCOPED_WRITE_ACCESS(ActorHandleMapDetector);
 	ActorHandleMap.Remove(Actor);
 }
 
-void UMassActorSubsystem::DisconnectActor(const TObjectKey<const AActor> Actor, const FMassEntityHandle Handle)
+void FMassActorManager::DisconnectActor(const TObjectKey<const AActor> Actor, const FMassEntityHandle Handle)
 {
 	if (Handle.IsValid() == false)
 	{
@@ -163,8 +172,8 @@ void UMassActorSubsystem::DisconnectActor(const TObjectKey<const AActor> Actor,
 	else
 	{
 		// unexpected mismatch. Add back and notify.
-		UE_VLOG_UELOG(this, LogMass, Warning, TEXT("%s: Trying to disconnect actor %s while the Handle given doesn't match the system\'s records")
+		UE_VLOG_UELOG(Owner.Get(), LogMass, Warning, TEXT("%s: Trying to disconnect actor %s while the Handle given doesn't match the system\'s records")
 			, ANSI_TO_TCHAR(__FUNCTION__), *AActor::GetDebugName(Actor.ResolveObjectPtr()));
 		SetHandleForActor(Actor, Handle);
 	}
-}
\ No newline at end of file
+}
diff --git a/Engine/Plugins/Runtime/MassGameplay/Source/MassActors/Public/MassActorSubsystem.h b/Engine/Plugins/Runtime/MassGameplay/Source/MassActors/Public/MassActorSubsystem.h
--- a/Engine/Plugins/Runtime/MassGameplay/Source/MassActors/Public/MassActorSubsystem.h
+++ b/Engine/Plugins/Runtime/MassGameplay/Source/MassActors/Public/MassActorSubsystem.h
@@ -81,6 +81,41 @@ private:
 	bool bIsOwnedByMass = false;
 };
 
+struct MASSACTORS_API FMassActorManager : public TSharedFromThis<FMassActorManager>
+{
+public:
+	explicit FMassActorManager(const TSharedPtr<FMassEntityManager>& EntityManager, UObject* InOwner = nullptr);
+
+	/** Get mass handle from an actor */
+	FMassEntityHandle GetEntityHandleFromActor(const TObjectKey<const AActor> Actor);
+
+	/** Set the mass handle associated to an actor */
+	void SetHandleForActor(const TObjectKey<const AActor> Actor, const FMassEntityHandle Handle);
+
+	/** Remove entry associated to an actor */
+	void RemoveHandleForActor(const TObjectKey<const AActor> Actor);
+
+	/** Get an actor pointer from a mass handle */
+	AActor* GetActorFromHandle(const FMassEntityHandle Handle, 
+		FMassActorFragment::EActorAccess Access = FMassActorFragment::EActorAccess::OnlyWhenAlive) const;
+
+	/** 
+	 *  Removes the connection between Actor and the given entity. Does all the required book keeping 
+	 *  (as opposed to straight up RemoveHandleForActor call). If the Handle doesn't match Actor no action is taken.
+	 */
+	void DisconnectActor(const TObjectKey<const AActor> Actor, const FMassEntityHandle Handle);
+
+protected:
+
+	TMap<TObjectKey<const AActor>, FMassEntityHandle> ActorHandleMap;
+	UE_MT_DECLARE_RW_ACCESS_DETECTOR(ActorHandleMapDetector);
+	
+	TSharedPtr<FMassEntityManager> EntityManager;
+
+	/** Points at an UObject hosting this instance of the FMassActorManager. It's fine for this to be null. */
+	TWeakObjectPtr<UObject> Owner;
+};
+
 /**
  * A subsystem managing communication between Actors and Mass
  */
@@ -91,23 +126,23 @@ class MASSACTORS_API UMassActorSubsystem : public UWorldSubsystem
 
 public:
 	/** Get mass handle from an actor */
-	FMassEntityHandle GetEntityHandleFromActor(const TObjectKey<const AActor> Actor);
+	inline FMassEntityHandle GetEntityHandleFromActor(const TObjectKey<const AActor> Actor);
 
 	/** Set the mass handle associated to an actor */
-	void SetHandleForActor(const TObjectKey<const AActor> Actor, const FMassEntityHandle Handle);
+	inline void SetHandleForActor(const TObjectKey<const AActor> Actor, const FMassEntityHandle Handle);
 
 	/** Remove entry associated to an actor */
-	void RemoveHandleForActor(const TObjectKey<const AActor> Actor);
+	inline void RemoveHandleForActor(const TObjectKey<const AActor> Actor);
 
 	/** Get an actor pointer from a mass handle */
-	AActor* GetActorFromHandle(const FMassEntityHandle Handle, 
+	inline AActor* GetActorFromHandle(const FMassEntityHandle Handle,
 		FMassActorFragment::EActorAccess Access = FMassActorFragment::EActorAccess::OnlyWhenAlive) const;
 
 	/** 
 	 *  Removes the connection between Actor and the given entity. Does all the required book keeping 
 	 *  (as opposed to straight up RemoveHandleForActor call). If the Handle doesn't match Actor no action is taken.
 	 */
-	void DisconnectActor(const TObjectKey<const AActor> Actor, const FMassEntityHandle Handle);
+	inline void DisconnectActor(const TObjectKey<const AActor> Actor, const FMassEntityHandle Handle);
 
 protected:
 	// USubsystem BEGIN
@@ -116,10 +151,7 @@ protected:
 	virtual bool DoesSupportWorldType(const EWorldType::Type WorldType) const override;
 	// USubsystem END
 	
-	TMap<TObjectKey<const AActor>, FMassEntityHandle> ActorHandleMap;
-	UE_MT_DECLARE_RW_ACCESS_DETECTOR(ActorHandleMapDetector);
-	
-	TSharedPtr<FMassEntityManager> EntityManager;
+	TSharedPtr<FMassActorManager> ActorManager;
 };
 
 template<>
@@ -130,3 +162,33 @@ struct TMassExternalSubsystemTraits<UMassActorSubsystem> final
 		GameThreadOnly = false
 	};
 };
+
+/**
+ * UMassActorSubsystem's inlines
+ */
+FMassEntityHandle UMassActorSubsystem::GetEntityHandleFromActor(const TObjectKey<const AActor> Actor)
+{
+	return ActorManager->GetEntityHandleFromActor(Actor);
+}
+
+void UMassActorSubsystem::SetHandleForActor(const TObjectKey<const AActor> Actor, const FMassEntityHandle Handle)
+{
+	ActorManager->SetHandleForActor(Actor, Handle);
+}
+
+void UMassActorSubsystem::RemoveHandleForActor(const TObjectKey<const AActor> Actor)
+{
+	ActorManager->RemoveHandleForActor(Actor);
+}
+
+AActor* UMassActorSubsystem::GetActorFromHandle(const FMassEntityHandle Handle,
+	FMassActorFragment::EActorAccess Access) const
+{
+	return ActorManager->GetActorFromHandle(Handle, Access);
+}
+
+void UMassActorSubsystem::DisconnectActor(const TObjectKey<const AActor> Actor, const FMassEntityHandle Handle) 
+{ 
+	return ActorManager->DisconnectActor(Actor, Handle); 
+}
+
