Refactored UMassEntityTemplateRegistry into a regular struct, FMassEntityTemplateRegistry.

diff --git a/Engine/Plugins/Runtime/MassGameplay/Source/MassSpawner/Private/MassEntityConfigAsset.cpp b/Engine/Plugins/Runtime/MassGameplay/Source/MassSpawner/Private/MassEntityConfigAsset.cpp
--- a/Engine/Plugins/Runtime/MassGameplay/Source/MassSpawner/Private/MassEntityConfigAsset.cpp
+++ b/Engine/Plugins/Runtime/MassGameplay/Source/MassSpawner/Private/MassEntityConfigAsset.cpp
@@ -85,7 +85,7 @@ const FMassEntityTemplate& FMassEntityConfig::GetOrCreateEntityTemplate(const UW
 
 	UMassSpawnerSubsystem* SpawnerSystem = UWorld::GetSubsystem<UMassSpawnerSubsystem>(&World);
 	check(SpawnerSystem);
-	UMassEntityTemplateRegistry& TemplateRegistry = SpawnerSystem->GetTemplateRegistryInstance();
+	FMassEntityTemplateRegistry& TemplateRegistry = SpawnerSystem->GetMutableTemplateRegistryInstance();
 
 	// Build new template
 	// TODO: Add methods to FMassEntityTemplateBuildContext to indicate dependency vs setup.
@@ -118,8 +118,7 @@ void FMassEntityConfig::DestroyEntityTemplate(const UWorld& World, const UObject
 
 	UMassSpawnerSubsystem* SpawnerSystem = UWorld::GetSubsystem<UMassSpawnerSubsystem>(&World);
 	check(SpawnerSystem);
-	UMassEntityTemplateRegistry& TemplateRegistry = SpawnerSystem->GetTemplateRegistryInstance();
-
+	FMassEntityTemplateRegistry& TemplateRegistry = SpawnerSystem->GetMutableTemplateRegistryInstance();
 
 	for (const UMassEntityTraitBase* Trait : CombinedTraits)
 	{
@@ -146,7 +145,7 @@ const FMassEntityTemplate* FMassEntityConfig::GetEntityTemplateInternal(const UW
 {
 	UMassSpawnerSubsystem* SpawnerSystem = UWorld::GetSubsystem<UMassSpawnerSubsystem>(&World);
 	check(SpawnerSystem);
-	const UMassEntityTemplateRegistry& TemplateRegistry = SpawnerSystem->GetTemplateRegistryInstance();
+	const FMassEntityTemplateRegistry& TemplateRegistry = SpawnerSystem->GetTemplateRegistryInstance();
 
 	// Combine all the features into one array
 	// @todo this is an inefficient way assuming given template is expected to have already been created. Figure out a way to cache it.
diff --git a/Engine/Plugins/Runtime/MassGameplay/Source/MassSpawner/Private/MassEntityTemplateRegistry.cpp b/Engine/Plugins/Runtime/MassGameplay/Source/MassSpawner/Private/MassEntityTemplateRegistry.cpp
--- a/Engine/Plugins/Runtime/MassGameplay/Source/MassSpawner/Private/MassEntityTemplateRegistry.cpp
+++ b/Engine/Plugins/Runtime/MassGameplay/Source/MassSpawner/Private/MassEntityTemplateRegistry.cpp
@@ -38,7 +38,7 @@ namespace FTemplateRegistryHelpers
 		UMassSpawnerSubsystem* SpawnerSystem = UWorld::GetSubsystem<UMassSpawnerSubsystem>(InWorld);
 		if (ensure(SpawnerSystem))
 		{
-			UMassEntityTemplateRegistry& Registry = SpawnerSystem->GetTemplateRegistryInstance();
+			FMassEntityTemplateRegistry& Registry = SpawnerSystem->GetMutableTemplateRegistryInstance();
 			Registry.DebugReset();
 		}
 	}
@@ -52,34 +52,27 @@ namespace FTemplateRegistryHelpers
 }
 
 //----------------------------------------------------------------------//
-// UMassEntityTemplateRegistry 
+// FMassEntityTemplateRegistry 
 //----------------------------------------------------------------------//
-TMap<const UScriptStruct*, UMassEntityTemplateRegistry::FStructToTemplateBuilderDelegate> UMassEntityTemplateRegistry::StructBasedBuilders;
+TMap<const UScriptStruct*, FMassEntityTemplateRegistry::FStructToTemplateBuilderDelegate> FMassEntityTemplateRegistry::StructBasedBuilders;
 
-void UMassEntityTemplateRegistry::BeginDestroy()
+FMassEntityTemplateRegistry::FMassEntityTemplateRegistry(UObject* InOwner)
+	: Owner(InOwner)
 {
-	// force release of memory owned by individual templates (especially the hosted InstancedScriptStructs).
-	TemplateIDToTemplateMap.Reset();
-
-	Super::BeginDestroy();
 }
 
-UWorld* UMassEntityTemplateRegistry::GetWorld() const 
+UWorld* FMassEntityTemplateRegistry::GetWorld() const 
 {
-	const UObject* Outer = GetOuter();
-	return Outer ? Outer->GetWorld() : nullptr;
+	return Owner.IsValid() ? Owner->GetWorld() : nullptr;
 }
 
-UMassEntityTemplateRegistry::FStructToTemplateBuilderDelegate& UMassEntityTemplateRegistry::FindOrAdd(const UScriptStruct& DataType)
+FMassEntityTemplateRegistry::FStructToTemplateBuilderDelegate& FMassEntityTemplateRegistry::FindOrAdd(const UScriptStruct& DataType)
 {
 	return StructBasedBuilders.FindOrAdd(&DataType);
 }
 
-const FMassEntityTemplate* UMassEntityTemplateRegistry::FindOrBuildStructTemplate(const FConstStructView StructInstance)
+const FMassEntityTemplate* FMassEntityTemplateRegistry::FindOrBuildStructTemplate(const FConstStructView StructInstance)
 {
-	// thou shall not call this function on CDO
-	check(HasAnyFlags(RF_ClassDefaultObject) == false);
-
 	const UScriptStruct* Type = StructInstance.GetScriptStruct();
 	check(Type);
 	// 1. Check if we already have the template stored.
@@ -122,13 +115,13 @@ const FMassEntityTemplate* UMassEntityTemplateRegistry::FindOrBuildStructTemplat
 
 		BuildTemplateImpl(*Builder, StructInstance, *NewTemplate);
 	}
-	UE_CVLOG_UELOG(Builder == nullptr, this, LogMassSpawner, Warning, TEXT("Attempting to build a MassAgentTemplate for struct type %s while template builder has not been registered for this type")
+	UE_CVLOG_UELOG(Builder == nullptr, Owner.Get(), LogMassSpawner, Warning, TEXT("Attempting to build a MassAgentTemplate for struct type %s while template builder has not been registered for this type")
 		, *GetNameSafe(Type));
 
 	return NewTemplate;
 }
 
-bool UMassEntityTemplateRegistry::BuildTemplateImpl(const FStructToTemplateBuilderDelegate& Builder, const FConstStructView StructInstance, FMassEntityTemplate& OutTemplate)
+bool FMassEntityTemplateRegistry::BuildTemplateImpl(const FStructToTemplateBuilderDelegate& Builder, const FConstStructView StructInstance, FMassEntityTemplate& OutTemplate)
 {
 	UWorld* World = GetWorld();
 	FMassEntityTemplateBuildContext Context(OutTemplate);
@@ -137,7 +130,7 @@ bool UMassEntityTemplateRegistry::BuildTemplateImpl(const FStructToTemplateBuild
 	{
 		InitializeEntityTemplate(OutTemplate);
 
-		UE_VLOG(this, LogMassSpawner, Log, TEXT("Created entity template for %s:\n%s"), *GetNameSafe(StructInstance.GetScriptStruct())
+		UE_VLOG(Owner.Get(), LogMassSpawner, Log, TEXT("Created entity template for %s:\n%s"), *GetNameSafe(StructInstance.GetScriptStruct())
 			, *OutTemplate.DebugGetDescription(UE::Mass::Utils::GetEntityManager(World)));
 
 		return true;
@@ -145,7 +138,7 @@ bool UMassEntityTemplateRegistry::BuildTemplateImpl(const FStructToTemplateBuild
 	return false;
 }
 
-void UMassEntityTemplateRegistry::InitializeEntityTemplate(FMassEntityTemplate& InOutTemplate) const
+void FMassEntityTemplateRegistry::InitializeEntityTemplate(FMassEntityTemplate& InOutTemplate) const
 {
 	UWorld* World = GetWorld();
 	check(World);
@@ -159,7 +152,7 @@ void UMassEntityTemplateRegistry::InitializeEntityTemplate(FMassEntityTemplate&
 	InOutTemplate.SetArchetype(ArchetypeHandle);
 }
 
-void UMassEntityTemplateRegistry::DebugReset()
+void FMassEntityTemplateRegistry::DebugReset()
 {
 #if WITH_MASSGAMEPLAY_DEBUG
 	LookupTemplateIDMap.Reset();
@@ -167,17 +160,17 @@ void UMassEntityTemplateRegistry::DebugReset()
 #endif // WITH_MASSGAMEPLAY_DEBUG
 }
 
-const FMassEntityTemplate* UMassEntityTemplateRegistry::FindTemplateFromTemplateID(FMassEntityTemplateID TemplateID) const 
+const FMassEntityTemplate* FMassEntityTemplateRegistry::FindTemplateFromTemplateID(FMassEntityTemplateID TemplateID) const 
 {
 	return TemplateIDToTemplateMap.Find(TemplateID);
 }
 
-FMassEntityTemplate* UMassEntityTemplateRegistry::FindMutableTemplateFromTemplateID(FMassEntityTemplateID TemplateID) 
+FMassEntityTemplate* FMassEntityTemplateRegistry::FindMutableTemplateFromTemplateID(FMassEntityTemplateID TemplateID) 
 {
 	return TemplateIDToTemplateMap.Find(TemplateID);
 }
 
-FMassEntityTemplate& UMassEntityTemplateRegistry::CreateTemplate(const uint32 HashLookup, FMassEntityTemplateID TemplateID)
+FMassEntityTemplate& FMassEntityTemplateRegistry::CreateTemplate(const uint32 HashLookup, FMassEntityTemplateID TemplateID)
 {
 	checkSlow(!LookupTemplateIDMap.Contains(HashLookup));
 	LookupTemplateIDMap.Add(HashLookup, TemplateID);
@@ -187,7 +180,7 @@ FMassEntityTemplate& UMassEntityTemplateRegistry::CreateTemplate(const uint32 Ha
 	return NewTemplate;
 }
 
-void UMassEntityTemplateRegistry::DestroyTemplate(const uint32 HashLookup, FMassEntityTemplateID TemplateID)
+void FMassEntityTemplateRegistry::DestroyTemplate(const uint32 HashLookup, FMassEntityTemplateID TemplateID)
 {
 	LookupTemplateIDMap.Remove(HashLookup);
 	TemplateIDToTemplateMap.Remove(TemplateID);
diff --git a/Engine/Plugins/Runtime/MassGameplay/Source/MassSpawner/Private/MassSpawnerSubsystem.cpp b/Engine/Plugins/Runtime/MassGameplay/Source/MassSpawner/Private/MassSpawnerSubsystem.cpp
--- a/Engine/Plugins/Runtime/MassGameplay/Source/MassSpawner/Private/MassSpawnerSubsystem.cpp
+++ b/Engine/Plugins/Runtime/MassGameplay/Source/MassSpawner/Private/MassSpawnerSubsystem.cpp
@@ -19,16 +19,22 @@
 //----------------------------------------------------------------------//
 //  UMassSpawnerSubsystem
 //----------------------------------------------------------------------//
-void UMassSpawnerSubsystem::Initialize(FSubsystemCollectionBase& Collection) 
+UMassSpawnerSubsystem::UMassSpawnerSubsystem()
+	: TemplateRegistryInstance(this)
+{
+
+}
+
+void UMassSpawnerSubsystem::Initialize(FSubsystemCollectionBase& Collection)
 {	
 	Super::Initialize(Collection);
 
-	// making sure UMassSimulationSubsystem gets created before the MassSpawnerSubsystem
+	// making sure UMassSimulationSubsystem gets created before the MassSpawnerSubsystem, since UMassSimulationSubsystem
+	// is where the EntityManager gets created for the runtime MassGameplay simulation
 	Collection.InitializeDependency<UMassSimulationSubsystem>();
 
 	UWorld* World = GetWorld();
 	check(World);
-	SimulationSystem = UWorld::GetSubsystem<UMassSimulationSubsystem>(World);
 	EntityManager = UE::Mass::Utils::GetEntityManagerChecked(*World).AsShared();
 }
 
@@ -37,11 +43,6 @@ void UMassSpawnerSubsystem::Deinitialize()
 	EntityManager.Reset();
 }
 
-void UMassSpawnerSubsystem::PostInitialize()
-{
-	TemplateRegistryInstance = NewObject<UMassEntityTemplateRegistry>(this);
-}
-
 void UMassSpawnerSubsystem::SpawnEntities(const FMassEntityTemplate& EntityTemplate, const uint32 NumberToSpawn, TArray<FMassEntityHandle>& OutEntities)
 {
 	check(EntityManager);
@@ -60,12 +61,7 @@ void UMassSpawnerSubsystem::SpawnEntities(FMassEntityTemplateID TemplateID, cons
 {
 	check(TemplateID.IsValid());
 
-	if (!ensureMsgf(TemplateRegistryInstance, TEXT("UMassSpawnerSubsystem didn\'t get its OnPostWorldInit call yet!")))
-	{
-		return;
-	}
-
-	const FMassEntityTemplate* EntityTemplate = TemplateRegistryInstance->FindTemplateFromTemplateID(TemplateID);
+	const FMassEntityTemplate* EntityTemplate = TemplateRegistryInstance.FindTemplateFromTemplateID(TemplateID);
 	checkf(EntityTemplate && EntityTemplate->IsValid(), TEXT("SpawnEntities: TemplateID must have been registered!"));
 
 	DoSpawning(*EntityTemplate, NumberToSpawn, SpawnData, InitializerClass, OutEntities);
@@ -75,12 +71,7 @@ void UMassSpawnerSubsystem::SpawnFromConfig(FStructView Config, const int32 NumT
 {
 	check(Config.IsValid());
 
-	if (!ensureMsgf(TemplateRegistryInstance, TEXT("UMassSpawnerSubsystem didn\'t get its OnPostWorldInit call yet!")))
-	{
-		return;
-	}
-	
-	const FMassEntityTemplate* EntityTemplate = TemplateRegistryInstance->FindOrBuildStructTemplate(Config);
+	const FMassEntityTemplate* EntityTemplate = TemplateRegistryInstance.FindOrBuildStructTemplate(Config);
 	checkf(EntityTemplate && EntityTemplate->IsValid(), TEXT("SpawnFromConfig: TemplateID must have been registered!"));
 
 	TArray<FMassEntityHandle> Entities;
@@ -89,16 +80,11 @@ void UMassSpawnerSubsystem::SpawnFromConfig(FStructView Config, const int32 NumT
 
 void UMassSpawnerSubsystem::RegisterCollection(TArrayView<FInstancedStruct> Collection)
 {
-	if (!ensureMsgf(TemplateRegistryInstance, TEXT("UMassSpawnerSubsystem didn't get its OnPostWorldInit called yet!")))
-	{
-		return;
-	}
-
 	for (FInstancedStruct& Entry : Collection)
 	{
 		if (Entry.IsValid())
 		{
-			TemplateRegistryInstance->FindOrBuildStructTemplate(Entry);
+			TemplateRegistryInstance.FindOrBuildStructTemplate(Entry);
 		}
 	}
 }
@@ -107,14 +93,8 @@ void UMassSpawnerSubsystem::DestroyEntities(const FMassEntityTemplateID Template
 {
 	TRACE_CPUPROFILER_EVENT_SCOPE_STR("MassSpawnerSubsystem_DestroyEntities")
 
-	if (!ensureMsgf(TemplateRegistryInstance, TEXT("UMassSpawnerSubsystem didn\'t get its OnPostWorldInit call yet!")))
-	{
-		return;
-	}
-
 	check(EntityManager);
-	check(SimulationSystem);
-	checkf(!SimulationSystem->IsDuringMassProcessing()
+	checkf(!EntityManager->IsProcessing()
 		, TEXT("%s called while MassEntity processing in progress. This is unsupported and dangerous!"), ANSI_TO_TCHAR(__FUNCTION__));
 
 	UWorld* World = GetWorld();
@@ -195,11 +175,5 @@ void UMassSpawnerSubsystem::DoSpawning(const FMassEntityTemplate& EntityTemplate
 const FMassEntityTemplate* UMassSpawnerSubsystem::GetMassEntityTemplate(FMassEntityTemplateID TemplateID) const
 {
 	check(TemplateID.IsValid());
-
-	if (!ensureMsgf(TemplateRegistryInstance, TEXT("UMassSpawnerSubsystem didn\'t get its OnPostWorldInit call yet!")))
-	{
-		return nullptr;
-	}
-
-	return TemplateRegistryInstance->FindTemplateFromTemplateID(TemplateID);
+	return TemplateRegistryInstance.FindTemplateFromTemplateID(TemplateID);
 }
diff --git a/Engine/Plugins/Runtime/MassGameplay/Source/MassSpawner/Public/MassEntityTemplateRegistry.h b/Engine/Plugins/Runtime/MassGameplay/Source/MassSpawner/Public/MassEntityTemplateRegistry.h
--- a/Engine/Plugins/Runtime/MassGameplay/Source/MassSpawner/Public/MassEntityTemplateRegistry.h
+++ b/Engine/Plugins/Runtime/MassGameplay/Source/MassSpawner/Public/MassEntityTemplateRegistry.h
@@ -198,17 +198,14 @@ protected:
 
 /** @todo document 
  */
-UCLASS()
-class MASSSPAWNER_API UMassEntityTemplateRegistry : public UObject
+struct MASSSPAWNER_API FMassEntityTemplateRegistry
 {
-	GENERATED_BODY()
-public:
 	// @todo consider TFunction instead
 	DECLARE_DELEGATE_ThreeParams(FStructToTemplateBuilderDelegate, const UWorld* /*World*/, const FConstStructView /*InStructInstance*/, FMassEntityTemplateBuildContext& /*BuildContext*/);
 
-	UMassEntityTemplateRegistry() = default;
-	virtual void BeginDestroy() override;
-	virtual UWorld* GetWorld() const override;
+	explicit FMassEntityTemplateRegistry(UObject* InOwner = nullptr);
+	
+	UWorld* GetWorld() const;
 
 	static FStructToTemplateBuilderDelegate& FindOrAdd(const UScriptStruct& DataType);
 
@@ -238,7 +235,14 @@ protected:
 	 */
 	TMap<uint32, FMassEntityTemplateID> LookupTemplateIDMap;
 
-	UPROPERTY(Transient)
 	TMap<FMassEntityTemplateID, FMassEntityTemplate> TemplateIDToTemplateMap;
+
+	TWeakObjectPtr<UObject> Owner;
 };
 
+
+UCLASS(deprecated, meta = (DeprecationMessage = "UMassEntityTemplateRegistry is deprecated starting UE5.2. Use FMassEntityTemplateRegistry instead"))
+class MASSSPAWNER_API UDEPRECATED_MassEntityTemplateRegistry : public UObject
+{
+	GENERATED_BODY()
+};
diff --git a/Engine/Plugins/Runtime/MassGameplay/Source/MassSpawner/Public/MassSpawnerSubsystem.h b/Engine/Plugins/Runtime/MassGameplay/Source/MassSpawner/Public/MassSpawnerSubsystem.h
--- a/Engine/Plugins/Runtime/MassGameplay/Source/MassSpawner/Public/MassSpawnerSubsystem.h
+++ b/Engine/Plugins/Runtime/MassGameplay/Source/MassSpawner/Public/MassSpawnerSubsystem.h
@@ -5,11 +5,11 @@
 #include "MassEntityManager.h"
 #include "InstancedStruct.h"
 #include "Subsystems/WorldSubsystem.h"
+#include "MassEntityTemplateRegistry.h"
 #include "MassSpawnerSubsystem.generated.h"
 
 struct FMassEntityManager;
 struct FMassEntityTemplate;
-class UMassEntityTemplateRegistry;
 struct FInstancedStruct;
 struct FStructView;
 struct FMassEntityTemplateID;
@@ -21,6 +21,8 @@ class MASSSPAWNER_API UMassSpawnerSubsystem : public UWorldSubsystem
 	GENERATED_BODY()
 
 public:
+	UMassSpawnerSubsystem();
+
 	/** Spawns entities of the kind described by the given EntityTemplate. The spawned entities are fully initialized
 	 *  meaning the EntityTemplate.InitializationPipeline gets run for all spawned entities.
 	 *  @param EntityTemplate template to use for spawning entities
@@ -35,7 +37,8 @@ public:
 
 	void DestroyEntities(const FMassEntityTemplateID TemplateID, TConstArrayView<FMassEntityHandle> Entities);
 
-	UMassEntityTemplateRegistry& GetTemplateRegistryInstance() const { check(TemplateRegistryInstance); return *TemplateRegistryInstance; }
+	const FMassEntityTemplateRegistry& GetTemplateRegistryInstance() const { return TemplateRegistryInstance; }
+	FMassEntityTemplateRegistry& GetMutableTemplateRegistryInstance() { return TemplateRegistryInstance; }
 
 	void RegisterCollection(TArrayView<FInstancedStruct> Collection);
 
@@ -44,7 +47,6 @@ public:
 protected:
 	virtual void Initialize(FSubsystemCollectionBase& Collection) override;
 	virtual void Deinitialize() override;
-	virtual void PostInitialize() override;
 
 	void DoSpawning(const FMassEntityTemplate& EntityTemplate, const int32 NumToSpawn, FConstStructView SpawnData, TSubclassOf<UMassProcessor> InitializerClass, TArray<FMassEntityHandle>& OutEntities);
 
@@ -55,10 +57,6 @@ protected:
 
 	TSharedPtr<FMassEntityManager> EntityManager;
 
-	UPROPERTY()
-	TObjectPtr<UMassEntityTemplateRegistry> TemplateRegistryInstance;
-
-	UPROPERTY()
-	TObjectPtr<UMassSimulationSubsystem> SimulationSystem;
+	FMassEntityTemplateRegistry TemplateRegistryInstance;
 };
 
