Simplified Mass spawning and entity template API a bit by cutting out stuff we don't use anywhere.

diff --git a/Engine/Plugins/Runtime/MassGameplay/Source/MassSpawner/Private/MassEntityConfigAsset.cpp b/Engine/Plugins/Runtime/MassGameplay/Source/MassSpawner/Private/MassEntityConfigAsset.cpp
--- a/Engine/Plugins/Runtime/MassGameplay/Source/MassSpawner/Private/MassEntityConfigAsset.cpp
+++ b/Engine/Plugins/Runtime/MassGameplay/Source/MassSpawner/Private/MassEntityConfigAsset.cpp
@@ -75,10 +75,9 @@ void UMassEntityConfigAsset::ValidateEntityConfig()
 
 const FMassEntityTemplate& FMassEntityConfig::GetOrCreateEntityTemplate(const UWorld& World, const UObject& ConfigOwner) const
 {
-	uint32 Hash;
 	FMassEntityTemplateID TemplateID;
 	TArray<UMassEntityTraitBase*> CombinedTraits;
-	if (const FMassEntityTemplate* ExistingTemplate = GetEntityTemplateInternal(World, ConfigOwner, Hash, TemplateID, CombinedTraits))
+	if (const FMassEntityTemplate* ExistingTemplate = GetEntityTemplateInternal(World, ConfigOwner, TemplateID, CombinedTraits))
 	{
 		return *ExistingTemplate;
 	}
@@ -91,7 +90,7 @@ const FMassEntityTemplate& FMassEntityConfig::GetOrCreateEntityTemplate(const UW
 	// TODO: Add methods to FMassEntityTemplateBuildContext to indicate dependency vs setup.
 	// Dependency should add a fragment with default values (which later can be overridden),
 	// while setup would override values and should be run just once.
-	FMassEntityTemplate& Template = TemplateRegistry.CreateTemplate(Hash, TemplateID);
+	FMassEntityTemplate& Template = TemplateRegistry.CreateTemplate(TemplateID);
 	FMassEntityTemplateBuildContext BuildContext(Template);
 
 	BuildContext.BuildFromTraits(CombinedTraits, World);
@@ -107,10 +106,9 @@ const FMassEntityTemplate& FMassEntityConfig::GetOrCreateEntityTemplate(const UW
 
 void FMassEntityConfig::DestroyEntityTemplate(const UWorld& World, const UObject& ConfigOwner) const
 {
-	uint32 Hash;
 	FMassEntityTemplateID TemplateID;
 	TArray<UMassEntityTraitBase*> CombinedTraits;
-	const FMassEntityTemplate* Template = GetEntityTemplateInternal(World, ConfigOwner, Hash, TemplateID, CombinedTraits);
+	const FMassEntityTemplate* Template = GetEntityTemplateInternal(World, ConfigOwner, TemplateID, CombinedTraits);
 	if (Template == nullptr)
 	{
 		return;
@@ -128,20 +126,19 @@ void FMassEntityConfig::DestroyEntityTemplate(const UWorld& World, const UObject
 
 	// TODO - The templates are not being torn down completely, resulting in traits that leave data in various subsystems. (Representation system)
 	
-	TemplateRegistry.DestroyTemplate(Hash, TemplateID);
+	TemplateRegistry.DestroyTemplate(TemplateID);
 }
 
 const FMassEntityTemplate& FMassEntityConfig::GetEntityTemplateChecked(const UWorld& World, const UObject& ConfigOwner) const
 {
-	uint32 Hash;
 	FMassEntityTemplateID TemplateID;
 	TArray<UMassEntityTraitBase*> CombinedTraits;
-	const FMassEntityTemplate* ExistingTemplate = GetEntityTemplateInternal(World, ConfigOwner, Hash, TemplateID, CombinedTraits);
+	const FMassEntityTemplate* ExistingTemplate = GetEntityTemplateInternal(World, ConfigOwner, TemplateID, CombinedTraits);
 	check(ExistingTemplate);
 	return *ExistingTemplate;
 }
 
-const FMassEntityTemplate* FMassEntityConfig::GetEntityTemplateInternal(const UWorld& World, const UObject& ConfigOwner, uint32& HashOut, FMassEntityTemplateID& TemplateIDOut, TArray<UMassEntityTraitBase*>& CombinedTraitsOut) const
+const FMassEntityTemplate* FMassEntityConfig::GetEntityTemplateInternal(const UWorld& World, const UObject& ConfigOwner, FMassEntityTemplateID& TemplateIDOut, TArray<UMassEntityTraitBase*>& CombinedTraitsOut) const
 {
 	UMassSpawnerSubsystem* SpawnerSystem = UWorld::GetSubsystem<UMassSpawnerSubsystem>(&World);
 	check(SpawnerSystem);
@@ -156,8 +153,8 @@ const FMassEntityTemplate* FMassEntityConfig::GetEntityTemplateInternal(const UW
 
 	// Return existing template if found.
 	// TODO: cache the hash.
-	HashOut = UE::MassSpawner::HashTraits(CombinedTraitsOut);
-	TemplateIDOut = FMassEntityTemplateID(HashOut, EMassEntityTemplateIDType::ScriptStruct); // TODO: add proper ID type
+	const uint32 HashOut = UE::MassSpawner::HashTraits(CombinedTraitsOut);
+	TemplateIDOut = FMassEntityTemplateID(HashOut);
 	return TemplateRegistry.FindTemplateFromTemplateID(TemplateIDOut);
 }
 
diff --git a/Engine/Plugins/Runtime/MassGameplay/Source/MassSpawner/Private/MassEntityTemplate.cpp b/Engine/Plugins/Runtime/MassGameplay/Source/MassSpawner/Private/MassEntityTemplate.cpp
--- a/Engine/Plugins/Runtime/MassGameplay/Source/MassSpawner/Private/MassEntityTemplate.cpp
+++ b/Engine/Plugins/Runtime/MassGameplay/Source/MassSpawner/Private/MassEntityTemplate.cpp
@@ -4,14 +4,13 @@
 #include "VisualLogger/VisualLoggerTypes.h"
 #include "MassDebugger.h"
 
-DEFINE_ENUM_TO_STRING(EMassEntityTemplateIDType, "/Script/MassSpawner");
 
 //----------------------------------------------------------------------//
 //  FMassEntityTemplateID
 //----------------------------------------------------------------------//
 FString FMassEntityTemplateID::ToString() const
 {
-	return FString::Printf(TEXT("[%s %d]"), *EnumToString(Type), Hash);
+	return FString::Printf(TEXT("[%d]"), Hash);
 }
 
 //----------------------------------------------------------------------//
diff --git a/Engine/Plugins/Runtime/MassGameplay/Source/MassSpawner/Private/MassEntityTemplateRegistry.cpp b/Engine/Plugins/Runtime/MassGameplay/Source/MassSpawner/Private/MassEntityTemplateRegistry.cpp
--- a/Engine/Plugins/Runtime/MassGameplay/Source/MassSpawner/Private/MassEntityTemplateRegistry.cpp
+++ b/Engine/Plugins/Runtime/MassGameplay/Source/MassSpawner/Private/MassEntityTemplateRegistry.cpp
@@ -71,56 +71,6 @@ FMassEntityTemplateRegistry::FStructToTemplateBuilderDelegate& FMassEntityTempla
 	return StructBasedBuilders.FindOrAdd(&DataType);
 }
 
-const FMassEntityTemplate* FMassEntityTemplateRegistry::FindOrBuildStructTemplate(const FConstStructView StructInstance)
-{
-	const UScriptStruct* Type = StructInstance.GetScriptStruct();
-	check(Type);
-	// 1. Check if we already have the template stored.
-	// 2. If not, 
-	//	a. build it
-	//	b. store it
-
-	const uint32 StructInstanceHash = FTemplateRegistryHelpers::CalcHash(StructInstance);
-	const uint32 HashLookup = HashCombine(GetTypeHash(Type->GetFName()), StructInstanceHash);
-
-	FMassEntityTemplateID* TemplateID = LookupTemplateIDMap.Find(HashLookup);
-
-	if (TemplateID != nullptr)
-	{
-		if (const FMassEntityTemplate* TemplateFound = TemplateIDToTemplateMap.Find(*TemplateID))
-		{
-			return TemplateFound;
-		}
-	}
-
-	// this means we don't have an entry for given struct. Let's see if we know how to make one
-	FMassEntityTemplate* NewTemplate = nullptr;
-	FStructToTemplateBuilderDelegate* Builder = StructBasedBuilders.Find(StructInstance.GetScriptStruct());
-	if (Builder)
-	{
-		if (TemplateID == nullptr)
-		{
-			// TODO consider removing the need for strings here
-			// Use the class name string for the hash here so the hash can be deterministic between client and server
-			const uint32 NameStringHash = GetTypeHash(Type->GetName());
-			const uint32 Hash = HashCombine(NameStringHash, StructInstanceHash);
-
-			TemplateID = &LookupTemplateIDMap.Add(HashLookup, FMassEntityTemplateID(Hash, EMassEntityTemplateIDType::ScriptStruct));
-		}
-
-		NewTemplate = &TemplateIDToTemplateMap.Add(*TemplateID);
-
-		check(NewTemplate);
-		NewTemplate->SetTemplateID(*TemplateID);
-
-		BuildTemplateImpl(*Builder, StructInstance, *NewTemplate);
-	}
-	UE_CVLOG_UELOG(Builder == nullptr, Owner.Get(), LogMassSpawner, Warning, TEXT("Attempting to build a MassAgentTemplate for struct type %s while template builder has not been registered for this type")
-		, *GetNameSafe(Type));
-
-	return NewTemplate;
-}
-
 bool FMassEntityTemplateRegistry::BuildTemplateImpl(const FStructToTemplateBuilderDelegate& Builder, const FConstStructView StructInstance, FMassEntityTemplate& OutTemplate)
 {
 	UWorld* World = GetWorld();
@@ -155,7 +105,6 @@ void FMassEntityTemplateRegistry::InitializeEntityTemplate(FMassEntityTemplate&
 void FMassEntityTemplateRegistry::DebugReset()
 {
 #if WITH_MASSGAMEPLAY_DEBUG
-	LookupTemplateIDMap.Reset();
 	TemplateIDToTemplateMap.Reset();
 #endif // WITH_MASSGAMEPLAY_DEBUG
 }
@@ -170,19 +119,16 @@ FMassEntityTemplate* FMassEntityTemplateRegistry::FindMutableTemplateFromTemplat
 	return TemplateIDToTemplateMap.Find(TemplateID);
 }
 
-FMassEntityTemplate& FMassEntityTemplateRegistry::CreateTemplate(const uint32 HashLookup, FMassEntityTemplateID TemplateID)
+FMassEntityTemplate& FMassEntityTemplateRegistry::CreateTemplate(FMassEntityTemplateID TemplateID)
 {
-	checkSlow(!LookupTemplateIDMap.Contains(HashLookup));
-	LookupTemplateIDMap.Add(HashLookup, TemplateID);
 	checkSlow(!TemplateIDToTemplateMap.Contains(TemplateID));
 	FMassEntityTemplate& NewTemplate = TemplateIDToTemplateMap.Add(TemplateID);
 	NewTemplate.SetTemplateID(TemplateID);
 	return NewTemplate;
 }
 
-void FMassEntityTemplateRegistry::DestroyTemplate(const uint32 HashLookup, FMassEntityTemplateID TemplateID)
+void FMassEntityTemplateRegistry::DestroyTemplate(FMassEntityTemplateID TemplateID)
 {
-	LookupTemplateIDMap.Remove(HashLookup);
 	TemplateIDToTemplateMap.Remove(TemplateID);
 }
 
diff --git a/Engine/Plugins/Runtime/MassGameplay/Source/MassSpawner/Private/MassSpawnerSubsystem.cpp b/Engine/Plugins/Runtime/MassGameplay/Source/MassSpawner/Private/MassSpawnerSubsystem.cpp
--- a/Engine/Plugins/Runtime/MassGameplay/Source/MassSpawner/Private/MassSpawnerSubsystem.cpp
+++ b/Engine/Plugins/Runtime/MassGameplay/Source/MassSpawner/Private/MassSpawnerSubsystem.cpp
@@ -67,28 +67,6 @@ void UMassSpawnerSubsystem::SpawnEntities(FMassEntityTemplateID TemplateID, cons
 	DoSpawning(*EntityTemplate, NumberToSpawn, SpawnData, InitializerClass, OutEntities);
 }
 
-void UMassSpawnerSubsystem::SpawnFromConfig(FStructView Config, const int32 NumToSpawn, FConstStructView SpawnData, TSubclassOf<UMassProcessor> InitializerClass)
-{
-	check(Config.IsValid());
-
-	const FMassEntityTemplate* EntityTemplate = TemplateRegistryInstance.FindOrBuildStructTemplate(Config);
-	checkf(EntityTemplate && EntityTemplate->IsValid(), TEXT("SpawnFromConfig: TemplateID must have been registered!"));
-
-	TArray<FMassEntityHandle> Entities;
-	DoSpawning(*EntityTemplate, NumToSpawn, SpawnData, InitializerClass, Entities);
-}
-
-void UMassSpawnerSubsystem::RegisterCollection(TArrayView<FInstancedStruct> Collection)
-{
-	for (FInstancedStruct& Entry : Collection)
-	{
-		if (Entry.IsValid())
-		{
-			TemplateRegistryInstance.FindOrBuildStructTemplate(Entry);
-		}
-	}
-}
-
 void UMassSpawnerSubsystem::DestroyEntities(const FMassEntityTemplateID TemplateID, TConstArrayView<FMassEntityHandle> Entities)
 {
 	TRACE_CPUPROFILER_EVENT_SCOPE_STR("MassSpawnerSubsystem_DestroyEntities")
diff --git a/Engine/Plugins/Runtime/MassGameplay/Source/MassSpawner/Public/MassEntityConfigAsset.h b/Engine/Plugins/Runtime/MassGameplay/Source/MassSpawner/Public/MassEntityConfigAsset.h
--- a/Engine/Plugins/Runtime/MassGameplay/Source/MassSpawner/Public/MassEntityConfigAsset.h
+++ b/Engine/Plugins/Runtime/MassGameplay/Source/MassSpawner/Public/MassEntityConfigAsset.h
@@ -64,7 +64,7 @@ protected:
 	TArray<TObjectPtr<UMassEntityTraitBase>> Traits;
 
 private:
-	const FMassEntityTemplate* GetEntityTemplateInternal(const UWorld& World, const UObject& ConfigOwner, uint32& HashOut, FMassEntityTemplateID& TemplateIDOut, TArray<UMassEntityTraitBase*>& CombinedTraitsOut) const;
+	const FMassEntityTemplate* GetEntityTemplateInternal(const UWorld& World, const UObject& ConfigOwner, FMassEntityTemplateID& TemplateIDOut, TArray<UMassEntityTraitBase*>& CombinedTraitsOut) const;
 };
 
 /**
diff --git a/Engine/Plugins/Runtime/MassGameplay/Source/MassSpawner/Public/MassEntityTemplate.h b/Engine/Plugins/Runtime/MassGameplay/Source/MassSpawner/Public/MassEntityTemplate.h
--- a/Engine/Plugins/Runtime/MassGameplay/Source/MassSpawner/Public/MassEntityTemplate.h
+++ b/Engine/Plugins/Runtime/MassGameplay/Source/MassSpawner/Public/MassEntityTemplate.h
@@ -12,25 +12,14 @@
 
 struct FMassEntityView;
 
-
-UENUM()
-enum class EMassEntityTemplateIDType : uint8
-{
-	None,
-	ScriptStruct,
-	Class,
-	Instance
-};
-
 //ID of the template an entity is using
 USTRUCT()
 struct MASSSPAWNER_API FMassEntityTemplateID
 {
 	GENERATED_BODY()
 
-	FMassEntityTemplateID(uint32 InHash, EMassEntityTemplateIDType InType)
-	: Hash(InHash)
-	, Type(InType)
+	explicit FMassEntityTemplateID(uint32 InHash)
+		: Hash(InHash)
 	{}
 
 	FMassEntityTemplateID() = default;
@@ -38,29 +27,23 @@ struct MASSSPAWNER_API FMassEntityTemplateID
 	uint32 GetHash() const { return Hash; }
 	void SetHash(uint32 InHash) { Hash = InHash; }
 
-	EMassEntityTemplateIDType GetType() const { return Type; }
-	void SetType(EMassEntityTemplateIDType InType) { Type = InType; }
-
 	bool operator==(const FMassEntityTemplateID& Other) const
 	{
-		return (Hash == Other.Hash) && (Type == Other.Type);
+		return (Hash == Other.Hash);
 	}
 
 	friend uint32 GetTypeHash(const FMassEntityTemplateID& TemplateID)
 	{
-		return HashCombine(TemplateID.Hash, (uint32)TemplateID.Type);
+		return TemplateID.Hash;
 	}
 
-	bool IsValid() { return Type != EMassEntityTemplateIDType::None; }
+	bool IsValid() { return Hash != 0; }
 
 	FString ToString() const;
 
 protected:
 	UPROPERTY()
 	uint32 Hash = 0;
-
-	UPROPERTY()
-	EMassEntityTemplateIDType Type = EMassEntityTemplateIDType::None;
 };
 
 /** @todo document	*/
diff --git a/Engine/Plugins/Runtime/MassGameplay/Source/MassSpawner/Public/MassEntityTemplateRegistry.h b/Engine/Plugins/Runtime/MassGameplay/Source/MassSpawner/Public/MassEntityTemplateRegistry.h
--- a/Engine/Plugins/Runtime/MassGameplay/Source/MassSpawner/Public/MassEntityTemplateRegistry.h
+++ b/Engine/Plugins/Runtime/MassGameplay/Source/MassSpawner/Public/MassEntityTemplateRegistry.h
@@ -209,16 +209,14 @@ struct MASSSPAWNER_API FMassEntityTemplateRegistry
 
 	static FStructToTemplateBuilderDelegate& FindOrAdd(const UScriptStruct& DataType);
 
-	const FMassEntityTemplate* FindOrBuildStructTemplate(const FConstStructView StructInstance);
-
 	/** Removes all the cached template instances */
 	void DebugReset();
 
 	const FMassEntityTemplate* FindTemplateFromTemplateID(FMassEntityTemplateID TemplateID) const;
 	FMassEntityTemplate* FindMutableTemplateFromTemplateID(FMassEntityTemplateID TemplateID);
 
-	FMassEntityTemplate& CreateTemplate(const uint32 HashLookup, FMassEntityTemplateID TemplateID);
-	void DestroyTemplate(const uint32 HashLookup, FMassEntityTemplateID TemplateID);
+	FMassEntityTemplate& CreateTemplate(FMassEntityTemplateID TemplateID);
+	void DestroyTemplate(FMassEntityTemplateID TemplateID);
 	void InitializeEntityTemplate(FMassEntityTemplate& InOutTemplate) const;
 
 protected:
@@ -228,13 +226,6 @@ protected:
 protected:
 	static TMap<const UScriptStruct*, FStructToTemplateBuilderDelegate> StructBasedBuilders;
 
-	/** 
-	 *  Map from a hash to a FMassEntityTemplateID
-	 *  For build from FInstancedStruct it will be the Combined hash of the UScriptStruct FName and FTemplateRegistryHelpers::CalcHash()
-	 *  This hash will not be deterministic between server and clients
-	 */
-	TMap<uint32, FMassEntityTemplateID> LookupTemplateIDMap;
-
 	TMap<FMassEntityTemplateID, FMassEntityTemplate> TemplateIDToTemplateMap;
 
 	TWeakObjectPtr<UObject> Owner;
diff --git a/Engine/Plugins/Runtime/MassGameplay/Source/MassSpawner/Public/MassSpawnerSubsystem.h b/Engine/Plugins/Runtime/MassGameplay/Source/MassSpawner/Public/MassSpawnerSubsystem.h
--- a/Engine/Plugins/Runtime/MassGameplay/Source/MassSpawner/Public/MassSpawnerSubsystem.h
+++ b/Engine/Plugins/Runtime/MassGameplay/Source/MassSpawner/Public/MassSpawnerSubsystem.h
@@ -33,15 +33,11 @@ public:
 
 	void SpawnEntities(FMassEntityTemplateID TemplateID, const uint32 NumberToSpawn, FConstStructView SpawnData, TSubclassOf<UMassProcessor> InitializerClass, TArray<FMassEntityHandle>& OutEntities);
 
-	void SpawnFromConfig(FStructView Config, const int32 NumToSpawn, FConstStructView SpawnData, TSubclassOf<UMassProcessor> InitializerClass);
-
 	void DestroyEntities(const FMassEntityTemplateID TemplateID, TConstArrayView<FMassEntityHandle> Entities);
 
 	const FMassEntityTemplateRegistry& GetTemplateRegistryInstance() const { return TemplateRegistryInstance; }
 	FMassEntityTemplateRegistry& GetMutableTemplateRegistryInstance() { return TemplateRegistryInstance; }
 
-	void RegisterCollection(TArrayView<FInstancedStruct> Collection);
-
 	const FMassEntityTemplate* GetMassEntityTemplate(FMassEntityTemplateID TemplateID) const;
 
 protected:
