Fixed an issue resulting from recent change to MassEntityTemplateID that rendered IDs representing empty templates "invalid" which in turn made "addition-less puppets" invalid, which caused all sorts of issues down the line.

#jria UE-174723

diff --git a/Engine/Plugins/Runtime/MassGameplay/Source/MassSpawner/Public/MassEntityTemplate.h b/Engine/Plugins/Runtime/MassGameplay/Source/MassSpawner/Public/MassEntityTemplate.h
--- a/Engine/Plugins/Runtime/MassGameplay/Source/MassSpawner/Public/MassEntityTemplate.h
+++ b/Engine/Plugins/Runtime/MassGameplay/Source/MassSpawner/Public/MassEntityTemplate.h
@@ -18,32 +18,53 @@ struct MASSSPAWNER_API FMassEntityTemplateID
 {
 	GENERATED_BODY()
 
+	FMassEntityTemplateID()
+		: bIsSet(false)
+	{}
 	explicit FMassEntityTemplateID(uint32 InHash)
-		: Hash(InHash)
+		: Hash(InHash), bIsSet(true)
 	{}
 
-	FMassEntityTemplateID() = default;
+	uint32 GetHash() const 
+	{
+		checkSlow(bIsSet);
+		return Hash; 
+	}
+	
+	void SetHash(uint32 InHash) 
+	{ 
+		Hash = InHash; 
+		bIsSet = true;
+	}
 
-	uint32 GetHash() const { return Hash; }
-	void SetHash(uint32 InHash) { Hash = InHash; }
+	void Invalidate(uint32 InHash)
+	{
+		// the exact value we set here doesn't really matter, but just to keep the possible states consistent we set it 
+		// to the default value;
+		Hash = 0;
+		bIsSet = false;
+	}
 
 	bool operator==(const FMassEntityTemplateID& Other) const
 	{
-		return (Hash == Other.Hash);
+		return (Hash == Other.Hash) && IsValid() == Other.IsValid();
 	}
 
 	friend uint32 GetTypeHash(const FMassEntityTemplateID& TemplateID)
 	{
-		return TemplateID.Hash;
+		return HashCombine(TemplateID.Hash, uint32(TemplateID.bIsSet));
 	}
 
-	bool IsValid() { return Hash != 0; }
+	bool IsValid() const { return bIsSet; }
 
 	FString ToString() const;
 
 protected:
 	UPROPERTY()
 	uint32 Hash = 0;
+
+	UPROPERTY()
+	uint8 bIsSet : 1;
 };
 
 /** @todo document	*/
