Added FMassEntityConfig.ConfigOwner member variable, mainly to get rid of the confusing UObject parameter from a number of FMassEntityConfig's functions

diff --git a/Engine/Plugins/Runtime/MassGameplay/Source/MassActors/Private/MassAgentComponent.cpp b/Engine/Plugins/Runtime/MassGameplay/Source/MassActors/Private/MassAgentComponent.cpp
--- a/Engine/Plugins/Runtime/MassGameplay/Source/MassActors/Private/MassAgentComponent.cpp
+++ b/Engine/Plugins/Runtime/MassGameplay/Source/MassActors/Private/MassAgentComponent.cpp
@@ -24,6 +24,7 @@
 // UMassAgentComponent
 //----------------------------------------------------------------------//
 UMassAgentComponent::UMassAgentComponent()
+	: EntityConfig(*this)
 {
 #if WITH_EDITORONLY_DATA
 	bAutoRegisterInEditorMode = true;
@@ -388,6 +389,7 @@ void UMassAgentComponent::DebugCheckStateConsistency()
 void UMassAgentComponent::SetEntityConfig(const FMassEntityConfig& InEntityConfig)
 {
 	EntityConfig = InEntityConfig;
+	EntityConfig.SetOwner(*this);
 }
 
 void UMassAgentComponent::Enable()
diff --git a/Engine/Plugins/Runtime/MassGameplay/Source/MassActors/Private/MassAgentSubsystem.cpp b/Engine/Plugins/Runtime/MassGameplay/Source/MassActors/Private/MassAgentSubsystem.cpp
--- a/Engine/Plugins/Runtime/MassGameplay/Source/MassActors/Private/MassAgentSubsystem.cpp
+++ b/Engine/Plugins/Runtime/MassGameplay/Source/MassActors/Private/MassAgentSubsystem.cpp
@@ -95,7 +95,7 @@ FMassEntityTemplateID UMassAgentSubsystem::RegisterAgentComponent(UMassAgentComp
 	check(World);
 
 	const FMassEntityConfig& EntityConfig = AgentComp.GetEntityConfig();
-	const FMassEntityTemplate& EntityTemplate = EntityConfig.GetOrCreateEntityTemplate(*World, AgentComp);
+	const FMassEntityTemplate& EntityTemplate = EntityConfig.GetOrCreateEntityTemplate(*World);
 
 #if UE_REPLICATION_COMPILE_CLIENT_CODE
 	if (AgentComp.IsNetSimulating())
@@ -135,7 +135,7 @@ void UMassAgentSubsystem::UpdateAgentComponent(const UMassAgentComponent& AgentC
 	check(World);
 
 	const FMassEntityConfig& EntityConfig = AgentComp.GetEntityConfig();
-	const FMassEntityTemplate& EntityTemplate = EntityConfig.GetOrCreateEntityTemplate(*World, AgentComp);
+	const FMassEntityTemplate& EntityTemplate = EntityConfig.GetOrCreateEntityTemplate(*World);
 
 	const FMassEntityHandle Entity = AgentComp.GetEntityHandle();
 	const FMassArchetypeHandle CurrentArchetypeHandle = EntityManager->GetArchetypeForEntity(Entity);
@@ -214,7 +214,7 @@ void UMassAgentSubsystem::UnregisterAgentComponent(UMassAgentComponent& AgentCom
 			if (ensure(World))
 			{
 				const FMassEntityConfig& EntityConfig = AgentComp.GetEntityConfig();
-				EntityTemplate = &EntityConfig.GetEntityTemplateChecked(*World, AgentComp);
+				EntityTemplate = &EntityConfig.GetEntityTemplateChecked(*World);
 			}
 
 			FMassEntityHandle Entity = AgentComp.GetEntityHandle();
diff --git a/Engine/Plugins/Runtime/MassGameplay/Source/MassSpawner/Private/MassEntityConfigAsset.cpp b/Engine/Plugins/Runtime/MassGameplay/Source/MassSpawner/Private/MassEntityConfigAsset.cpp
--- a/Engine/Plugins/Runtime/MassGameplay/Source/MassSpawner/Private/MassEntityConfigAsset.cpp
+++ b/Engine/Plugins/Runtime/MassGameplay/Source/MassSpawner/Private/MassEntityConfigAsset.cpp
@@ -41,8 +41,9 @@ namespace UE::MassSpawner
 	}
 } // UE::MassSpawner
 
-FMassEntityConfig::FMassEntityConfig(UMassEntityConfigAsset& InParent)
-	: Parent(&InParent)
+
+FMassEntityConfig::FMassEntityConfig(UObject& InOwner)
+	: ConfigOwner(&InOwner)
 {
 
 }
@@ -52,7 +53,7 @@ void UMassEntityConfigAsset::ValidateEntityConfig()
 {
 	if (UWorld* EditorWorld = GEditor->GetEditorWorldContext().World())
 	{
-		if (Config.ValidateEntityTemplate(*EditorWorld, *this))
+		if (Config.ValidateEntityTemplate(*EditorWorld))
 		{
 			const FText InfoText = LOCTEXT("MassEntityConfigAssetNoErrorsDetected", "There were no error detected during validation of the EntityConfigAsset");
 			
@@ -73,11 +74,11 @@ void UMassEntityConfigAsset::ValidateEntityConfig()
 }
 #endif // WITH_EDITOR
 
-const FMassEntityTemplate& FMassEntityConfig::GetOrCreateEntityTemplate(const UWorld& World, const UObject& ConfigOwner) const
+const FMassEntityTemplate& FMassEntityConfig::GetOrCreateEntityTemplate(const UWorld& World) const
 {
 	FMassEntityTemplateID TemplateID;
 	TArray<UMassEntityTraitBase*> CombinedTraits;
-	if (const FMassEntityTemplate* ExistingTemplate = GetEntityTemplateInternal(World, ConfigOwner, TemplateID, CombinedTraits))
+	if (const FMassEntityTemplate* ExistingTemplate = GetEntityTemplateInternal(World, TemplateID, CombinedTraits))
 	{
 		return *ExistingTemplate;
 	}
@@ -94,7 +95,7 @@ const FMassEntityTemplate& FMassEntityConfig::GetOrCreateEntityTemplate(const UW
 	FMassEntityTemplateBuildContext BuildContext(Template);
 
 	BuildContext.BuildFromTraits(CombinedTraits, World);
-	Template.SetTemplateName(ConfigOwner.GetName());
+	Template.SetTemplateName(GetNameSafe(ConfigOwner));
 
 	// It is ok to have an empty template, 
     // but be aware there will be an error if you try to create an entity with it
@@ -104,11 +105,11 @@ const FMassEntityTemplate& FMassEntityConfig::GetOrCreateEntityTemplate(const UW
 	return Template;
 }
 
-void FMassEntityConfig::DestroyEntityTemplate(const UWorld& World, const UObject& ConfigOwner) const
+void FMassEntityConfig::DestroyEntityTemplate(const UWorld& World) const
 {
 	FMassEntityTemplateID TemplateID;
 	TArray<UMassEntityTraitBase*> CombinedTraits;
-	const FMassEntityTemplate* Template = GetEntityTemplateInternal(World, ConfigOwner, TemplateID, CombinedTraits);
+	const FMassEntityTemplate* Template = GetEntityTemplateInternal(World, TemplateID, CombinedTraits);
 	if (Template == nullptr)
 	{
 		return;
@@ -129,16 +130,16 @@ void FMassEntityConfig::DestroyEntityTemplate(const UWorld& World, const UObject
 	TemplateRegistry.DestroyTemplate(TemplateID);
 }
 
-const FMassEntityTemplate& FMassEntityConfig::GetEntityTemplateChecked(const UWorld& World, const UObject& ConfigOwner) const
+const FMassEntityTemplate& FMassEntityConfig::GetEntityTemplateChecked(const UWorld& World) const
 {
 	FMassEntityTemplateID TemplateID;
 	TArray<UMassEntityTraitBase*> CombinedTraits;
-	const FMassEntityTemplate* ExistingTemplate = GetEntityTemplateInternal(World, ConfigOwner, TemplateID, CombinedTraits);
+	const FMassEntityTemplate* ExistingTemplate = GetEntityTemplateInternal(World, TemplateID, CombinedTraits);
 	check(ExistingTemplate);
 	return *ExistingTemplate;
 }
 
-const FMassEntityTemplate* FMassEntityConfig::GetEntityTemplateInternal(const UWorld& World, const UObject& ConfigOwner, FMassEntityTemplateID& TemplateIDOut, TArray<UMassEntityTraitBase*>& CombinedTraitsOut) const
+const FMassEntityTemplate* FMassEntityConfig::GetEntityTemplateInternal(const UWorld& World, FMassEntityTemplateID& TemplateIDOut, TArray<UMassEntityTraitBase*>& CombinedTraitsOut) const
 {
 	UMassSpawnerSubsystem* SpawnerSystem = UWorld::GetSubsystem<UMassSpawnerSubsystem>(&World);
 	check(SpawnerSystem);
@@ -148,8 +149,8 @@ const FMassEntityTemplate* FMassEntityConfig::GetEntityTemplateInternal(const UW
 	// @todo this is an inefficient way assuming given template is expected to have already been created. Figure out a way to cache it.
 	TArray<const UObject*> Visited;
 	CombinedTraitsOut.Reset();
-	Visited.Add(&ConfigOwner);
-	GetCombinedTraits(CombinedTraitsOut, Visited, ConfigOwner);
+	Visited.Add(ConfigOwner);
+	GetCombinedTraits(CombinedTraitsOut, Visited);
 
 	// Return existing template if found.
 	// TODO: cache the hash.
@@ -158,7 +159,7 @@ const FMassEntityTemplate* FMassEntityConfig::GetEntityTemplateInternal(const UW
 	return TemplateRegistry.FindTemplateFromTemplateID(TemplateIDOut);
 }
 
-void FMassEntityConfig::GetCombinedTraits(TArray<UMassEntityTraitBase*>& OutTraits, TArray<const UObject*>& Visited, const UObject& ConfigOwner) const
+void FMassEntityConfig::GetCombinedTraits(TArray<UMassEntityTraitBase*>& OutTraits, TArray<const UObject*>& Visited) const
 {
 	if (Parent)
 	{
@@ -171,12 +172,12 @@ void FMassEntityConfig::GetCombinedTraits(TArray<UMassEntityTraitBase*>& OutTrai
 				Path += Object->GetName();
 				Path += TEXT("/");
 			}
-			UE_VLOG(&ConfigOwner, LogMassSpawner, Error, TEXT("%s: Encountered %s as parent second time (Infinite loop). %s"), *GetNameSafe(&ConfigOwner), *GetNameSafe(Parent), *Path);
+			UE_VLOG(ConfigOwner, LogMassSpawner, Error, TEXT("%s: Encountered %s as parent second time (Infinite loop). %s"), *GetNameSafe(ConfigOwner), *GetNameSafe(Parent), *Path);
 		}
 		else
 		{
 			Visited.Add(Parent);
-			Parent->GetConfig().GetCombinedTraits(OutTraits, Visited, ConfigOwner);
+			Parent->GetConfig().GetCombinedTraits(OutTraits, Visited);
 		}
 	}
 
@@ -204,12 +205,12 @@ void FMassEntityConfig::AddTrait(UMassEntityTraitBase& Trait)
 	Traits.Add(&Trait);
 }
 
-bool FMassEntityConfig::ValidateEntityTemplate(const UWorld& World, const UObject& ConfigOwner)
+bool FMassEntityConfig::ValidateEntityTemplate(const UWorld& World)
 {
 	TArray<const UObject*> Visited;
 	TArray<UMassEntityTraitBase*> CombinedTraits;
-	Visited.Add(&ConfigOwner);
-	GetCombinedTraits(CombinedTraits, Visited, ConfigOwner);
+	Visited.Add(ConfigOwner);
+	GetCombinedTraits(CombinedTraits, Visited);
 
 	FMassEntityTemplate Template;
 	FMassEntityTemplateBuildContext BuildContext(Template);
@@ -217,4 +218,33 @@ bool FMassEntityConfig::ValidateEntityTemplate(const UWorld& World, const UObjec
 	return BuildContext.BuildFromTraits(CombinedTraits, World);
 }
 
+//-----------------------------------------------------------------------------
+// DEPRECATED
+//-----------------------------------------------------------------------------
+const FMassEntityTemplate& FMassEntityConfig::GetOrCreateEntityTemplate(const UWorld& World, const UObject& InConfigOwner) const
+{
+	return GetOrCreateEntityTemplate(World);
+}
+
+void FMassEntityConfig::DestroyEntityTemplate(const UWorld& World, const UObject& InConfigOwner) const
+{
+	DestroyEntityTemplate(World);
+}
+
+const FMassEntityTemplate& FMassEntityConfig::GetEntityTemplateChecked(const UWorld& World, const UObject& InConfigOwner) const
+{
+	return GetEntityTemplateChecked(World);
+}
+
+bool FMassEntityConfig::ValidateEntityTemplate(const UWorld& World, const UObject& InConfigOwner)
+{
+	return ValidateEntityTemplate(World);
+}
+
+void FMassEntityConfig::GetCombinedTraits(TArray<UMassEntityTraitBase*>& OutTraits, TArray<const UObject*>& Visited, const UObject& InConfigOwner) const
+{
+	return GetCombinedTraits(OutTraits, Visited);
+}
+
+
 #undef LOCTEXT_NAMESPACE 
diff --git a/Engine/Plugins/Runtime/MassGameplay/Source/MassSpawner/Public/MassEntityConfigAsset.h b/Engine/Plugins/Runtime/MassGameplay/Source/MassSpawner/Public/MassEntityConfigAsset.h
--- a/Engine/Plugins/Runtime/MassGameplay/Source/MassSpawner/Public/MassEntityConfigAsset.h
+++ b/Engine/Plugins/Runtime/MassGameplay/Source/MassSpawner/Public/MassEntityConfigAsset.h
@@ -24,24 +24,25 @@ struct MASSSPAWNER_API FMassEntityConfig
 	GENERATED_BODY()
 
 	FMassEntityConfig() = default;
-	FMassEntityConfig(UMassEntityConfigAsset& InParent);
+	FMassEntityConfig(UObject& InOwner);
 
 	/** Create entity template based on the features included in this config.
 	 *  @param World in which we are creating the template.
-	 *  @param ConfigOwner Owner of the FMassEntityConfig used for error reporting.
 	 */
-	const FMassEntityTemplate& GetOrCreateEntityTemplate(const UWorld& World, const UObject& ConfigOwner) const;
+	const FMassEntityTemplate& GetOrCreateEntityTemplate(const UWorld& World) const;
 
-	void DestroyEntityTemplate(const UWorld& World, const UObject& ConfigOwner) const;
+	void DestroyEntityTemplate(const UWorld& World) const;
 
 	/** 
 	 * Fetches the EntityTemplate for given World, fails a check if one cannot be found.
 	 */
-	const FMassEntityTemplate& GetEntityTemplateChecked(const UWorld& World, const UObject& ConfigOwner) const;
+	const FMassEntityTemplate& GetEntityTemplateChecked(const UWorld& World) const;
 
 	/** @return Parent config */
 	const UMassEntityConfigAsset* GetParent() const { return Parent; }
 
+	void SetParentAsset(const UMassEntityConfigAsset& InParent) { Parent = &InParent; }
+	
 	/** @return View to the array of features defined on this config */
 	TConstArrayView<UMassEntityTraitBase*> GetTraits() const { return Traits; }
 
@@ -49,22 +50,43 @@ struct MASSSPAWNER_API FMassEntityConfig
 	void AddTrait(UMassEntityTraitBase& Trait);
 
 	/** Validates if the entity template is well built */
+	bool ValidateEntityTemplate(const UWorld& World);
+
+	void SetOwner(UObject& InOwner) { ConfigOwner = &InOwner; }
+
+	UE_DEPRECATED(5.3, "This flavor of GetOrCreateEntityTemplate is deperecated. Use the one without the ConfigOwner parameter (now a property of the FMassEntityConfig itself)")
+	const FMassEntityTemplate& GetOrCreateEntityTemplate(const UWorld& World, const UObject& ConfigOwner) const;
+
+	UE_DEPRECATED(5.3, "This flavor of DestroyEntityTemplate is deperecated. Use the one without the ConfigOwner parameter (now a property of the FMassEntityConfig itself)")
+	void DestroyEntityTemplate(const UWorld& World, const UObject& ConfigOwner) const;
+
+	UE_DEPRECATED(5.3, "This flavor of GetEntityTemplateChecked is deperecated. Use the one without the ConfigOwner parameter (now a property of the FMassEntityConfig itself)")
+	const FMassEntityTemplate& GetEntityTemplateChecked(const UWorld& World, const UObject& ConfigOwner) const;
+
+	UE_DEPRECATED(5.3, "This flavor of ValidateEntityTemplate is deperecated. Use the one without the ConfigOwner parameter (now a property of the FMassEntityConfig itself)")
 	bool ValidateEntityTemplate(const UWorld& World, const UObject& ConfigOwner);
 
 protected:
 	/** Combines traits based on the config hierarchy and returns list of unique traits */
+	void GetCombinedTraits(TArray<UMassEntityTraitBase*>& OutTraits, TArray<const UObject*>& Visited) const;
+
+	/** Combines traits based on the config hierarchy and returns list of unique traits */
+	UE_DEPRECATED(5.3, "This flavor of GetCombinedTraits is deperecated. Use the one without the ConfigOwner parameter (now a property of the FMassEntityConfig itself)")
 	void GetCombinedTraits(TArray<UMassEntityTraitBase*>& OutTraits, TArray<const UObject*>& Visited, const UObject& ConfigOwner) const;
 
 	/** Reference to parent config asset */
 	UPROPERTY(Category = "Derived Traits", EditAnywhere)
-	TObjectPtr<UMassEntityConfigAsset> Parent = nullptr;
+	TObjectPtr<const UMassEntityConfigAsset> Parent = nullptr;
 
 	/** Array of unique traits of this config */
 	UPROPERTY(Category = "Traits", EditAnywhere, Instanced)
 	TArray<TObjectPtr<UMassEntityTraitBase>> Traits;
 
+	UPROPERTY(transient)
+	TObjectPtr<UObject> ConfigOwner = nullptr;
+
 private:
-	const FMassEntityTemplate* GetEntityTemplateInternal(const UWorld& World, const UObject& ConfigOwner, FMassEntityTemplateID& TemplateIDOut, TArray<UMassEntityTraitBase*>& CombinedTraitsOut) const;
+	const FMassEntityTemplate* GetEntityTemplateInternal(const UWorld& World, FMassEntityTemplateID& TemplateIDOut, TArray<UMassEntityTraitBase*>& CombinedTraitsOut) const;
 };
 
 /**
@@ -77,6 +99,10 @@ class MASSSPAWNER_API UMassEntityConfigAsset : public UDataAsset
 	GENERATED_BODY()
 
 public:
+	UMassEntityConfigAsset()
+		: Config(*this)
+	{}
+
 	/** @return Agent config stored in this asset */
 	const FMassEntityConfig& GetConfig() const { return Config; }
 
@@ -85,12 +111,12 @@ public:
 
 	const FMassEntityTemplate& GetOrCreateEntityTemplate(const UWorld& World) const
 	{
-		return Config.GetOrCreateEntityTemplate(World, *this);
+		return Config.GetOrCreateEntityTemplate(World);
 	}
 
 	void DestroyEntityTemplate(const UWorld& World) const
 	{
-		Config.DestroyEntityTemplate(World, *this);
+		Config.DestroyEntityTemplate(World);
 	}
 
 #if WITH_EDITOR
