Modified MassLODSubsystem to support using editor viewports as viewers in editor worlds.

diff --git a/Engine/Plugins/Runtime/MassGameplay/Source/MassLOD/Private/MassLODSubsystem.cpp b/Engine/Plugins/Runtime/MassGameplay/Source/MassLOD/Private/MassLODSubsystem.cpp
--- a/Engine/Plugins/Runtime/MassGameplay/Source/MassLOD/Private/MassLODSubsystem.cpp
+++ b/Engine/Plugins/Runtime/MassGameplay/Source/MassLOD/Private/MassLODSubsystem.cpp
@@ -6,6 +6,11 @@
 #include "WorldPartition/WorldPartitionStreamingSource.h"
 #include "Engine/World.h"
 #include "MassSimulationSubsystem.h"
+#if WITH_EDITOR
+#include "CoreGlobals.h" // GIsEditor
+#include "LevelEditorViewport.h"
+#include "Editor/EditorEngine.h"
+#endif // WITH_EDITOR
 
 namespace UE
 {
@@ -26,14 +31,18 @@ void UMassLODSubsystem::Initialize(FSubsystemCollectionBase& Collection)
 	Collection.InitializeDependency(UMassSimulationSubsystem::StaticClass());
 
 	Super::Initialize(Collection);
-	SynchronizeViewers();
 	
 	if (UWorld* World = GetWorld())
 	{
 		UMassSimulationSubsystem* SimSystem = World->GetSubsystem<UMassSimulationSubsystem>();
 		check(SimSystem);
 		SimSystem->GetOnProcessingPhaseStarted(EMassProcessingPhase::PrePhysics).AddUObject(this, &UMassLODSubsystem::OnPrePhysicsPhaseStarted);
+#if WITH_EDITOR
+		bUseEditorLevelViewports = (GIsEditor && World->WorldType == EWorldType::Editor);
+#endif // WITH_EDITOR
 	}
+
+	SynchronizeViewers();
 }
 
 void UMassLODSubsystem::OnPrePhysicsPhaseStarted(float DeltaTime)
@@ -144,6 +153,16 @@ void UMassLODSubsystem::SynchronizeViewers()
 		{
 			LocalViewerMap.Add(GetTypeHash(ViewerInfo.StreamingSourceName), ViewerInfo.Handle);
 		}
+#if WITH_EDITOR
+		else if (bUseEditorLevelViewports && ViewerInfo.EditorViewportClientIndex != INDEX_NONE
+			&& GEditor && GEditor->GetLevelViewportClients().IsValidIndex(ViewerInfo.EditorViewportClientIndex)
+			&& GEditor->GetLevelViewportClients()[ViewerInfo.EditorViewportClientIndex])
+		{
+			
+			const int32 HashValue = GetTypeHash(GEditor->GetLevelViewportClients()[ViewerInfo.EditorViewportClientIndex]);
+			LocalViewerMap.Add(HashValue, ViewerInfo.Handle);
+		}
+#endif // WITH_EDITOR
 		else
 		{
 			// Safe to remove while iterating as it is a sparse array with a free list
@@ -164,7 +183,7 @@ void UMassLODSubsystem::SynchronizeViewers()
 			if (LocalViewerMap.Remove(GetTypeHash(PlayerController->GetFName())) == 0)
 			{
 				// If not add it to the list
-				AddViewer(PlayerController);
+				AddPlayerViewer(*PlayerController);
 			}
 		}
 
@@ -173,10 +192,26 @@ void UMassLODSubsystem::SynchronizeViewers()
 		{
 			if (LocalViewerMap.Remove(GetTypeHash(StreamingSource.Name)) == 0)
 			{
-				AddViewer(nullptr, StreamingSource.Name);
+				AddStreamingSourceViewer(StreamingSource.Name);
 			}
 		}
 	}
+#if WITH_EDITOR
+	if (bUseEditorLevelViewports)
+	{
+		CA_ASSUME(GEditor);
+		for (int32 ClientIndex = 0; ClientIndex < GEditor->GetLevelViewportClients().Num(); ++ClientIndex)
+		{
+			const FLevelEditorViewportClient* LevelVC = GEditor->GetLevelViewportClients()[ClientIndex];
+			if (LevelVC && LevelVC->IsPerspective())
+			{
+				const int32 HashValue = GetTypeHash(LevelVC);
+				AddEditorViewer(HashValue, ClientIndex);
+			}
+		}
+	}
+	
+#endif // WITH_EDITOR
 
 	// Anything left in the map need to be removed from the list
 	for (TMap<uint32, FMassViewerHandle>::TIterator Itr = LocalViewerMap.CreateIterator(); Itr; ++Itr)
@@ -222,6 +257,18 @@ void UMassLODSubsystem::SynchronizeViewers()
 				//ViewerInfo.AspectRatio = MinViewInfo.AspectRatio;
 			}
 		}
+#if WITH_EDITOR
+		else if (bUseEditorLevelViewports && ViewerInfo.EditorViewportClientIndex != INDEX_NONE)
+		{
+			CA_ASSUME(GEditor);
+			const FLevelEditorViewportClient* LevelVC = GEditor->GetLevelViewportClients()[ViewerInfo.EditorViewportClientIndex];
+			checkSlow(LevelVC);
+			ViewerInfo.bEnabled = LevelVC && LevelVC->IsPerspective();
+			ViewerInfo.Location = LevelVC->GetViewLocation();
+			ViewerInfo.Rotation = LevelVC->GetViewRotation();
+			
+		}
+#endif // WITH_EDITOR
 		else
 		{
 			checkf(!ViewerInfo.StreamingSourceName.IsNone(), TEXT("Expecting to have a streamingsourcename if the playercontroller is null"));
@@ -239,37 +286,68 @@ void UMassLODSubsystem::SynchronizeViewers()
 
 void UMassLODSubsystem::AddViewer(APlayerController* PlayerController, FName StreamingSourceName /* = NAME_None*/)
 {
-	const int32 HashValue = PlayerController ? GetTypeHash(PlayerController->GetFName()) : GetTypeHash(StreamingSourceName);
+	if (PlayerController)
+	{
+		AddPlayerViewer(*PlayerController);
+	}
+	else
+	{
+		AddStreamingSourceViewer(StreamingSourceName);
+	}
+}
+
+void UMassLODSubsystem::AddPlayerViewer(APlayerController& PlayerController)
+{
+	const int32 HashValue = GetTypeHash(PlayerController.GetFName());
 
 	FMassViewerHandle& ViewerHandle = ViewerMap.FindOrAdd(HashValue, FMassViewerHandle());
 	if (ViewerHandle.IsValid())
 	{
 		// We are only interested to set the player controller if it was not already set.
-		if (PlayerController)
+		const int32 ViewerHandleIdx = GetValidViewerIdx(ViewerHandle);
+		check(ViewerHandleIdx != INDEX_NONE);
+
+		FViewerInfo& ViewerInfo = Viewers[ViewerHandleIdx];
+		check(ViewerInfo.PlayerController == nullptr);
+		ViewerInfo.PlayerController = &PlayerController;
+	}
+	else
+	{
+		// Add new viewer
+#if UE_ALLOW_DEBUG_REPLICATION_DUPLICATE_VIEWERS_PER_CONTROLLER
+		//for debugging / profiling purposes create DebugNumberViwersPerController
+		//in this case ViewerMap will only contain a hash to the most recent viewer handle created.
+		for (int Idx = 0; Idx < UE::MassLOD::DebugNumberViewersPerController; ++Idx)
+#endif //UE_ALLOW_DEBUG_REPLICATION_DUPLICATE_VIEWERS_PER_CONTROLLER
 		{
-			const int32 ViewerHandleIdx = GetValidViewerIdx(ViewerHandle);
-			check(ViewerHandleIdx != INDEX_NONE);
+			const bool bAddNew = ViewerFreeIndices.Num() == 0;
+			const int NewIdx = bAddNew ? Viewers.Num() : ViewerFreeIndices.Pop();
+			FViewerInfo& NewViewerInfo = bAddNew ? Viewers.AddDefaulted_GetRef() : Viewers[NewIdx];
+			NewViewerInfo.PlayerController = &PlayerController;
+			NewViewerInfo.Handle.Index = NewIdx;
+			NewViewerInfo.Handle.SerialNumber = GetNextViewerSerialNumber();
+			NewViewerInfo.HashValue = HashValue;
 
-			FViewerInfo& ViewerInfo = Viewers[ViewerHandleIdx];
-			check(ViewerInfo.PlayerController == nullptr);
-			ViewerInfo.PlayerController = PlayerController;
-			PlayerController->OnEndPlay.AddUniqueDynamic(this, &UMassLODSubsystem::OnPlayerControllerEndPlay);
-		}
+			ViewerHandle = NewViewerInfo.Handle;
 
-		return;
+			OnViewerAddedDelegate.Broadcast(NewViewerInfo);
+		}
 	}
 
-	// Add new viewer
-#if UE_ALLOW_DEBUG_REPLICATION_DUPLICATE_VIEWERS_PER_CONTROLLER
-	//for debugging / profiling purposes create DebugNumberViwersPerController
-	//in this case ViewerMap will only contain a hash to the most recent viewer handle created.
-	for (int Idx = 0; Idx < UE::MassLOD::DebugNumberViewersPerController; ++Idx)
-#endif //UE_ALLOW_DEBUG_REPLICATION_DUPLICATE_VIEWERS_PER_CONTROLLER
+	PlayerController.OnEndPlay.AddUniqueDynamic(this, &UMassLODSubsystem::OnPlayerControllerEndPlay);
+}
+
+void UMassLODSubsystem::AddStreamingSourceViewer(const FName StreamingSourceName)
+{
+	const int32 HashValue = GetTypeHash(StreamingSourceName);
+
+	FMassViewerHandle& ViewerHandle = ViewerMap.FindOrAdd(HashValue, FMassViewerHandle());
+	// only add new viewer if it hasn't been added yet
+	if (ViewerHandle.IsValid() == false)
 	{
 		const bool bAddNew = ViewerFreeIndices.Num() == 0;
 		const int NewIdx = bAddNew ? Viewers.Num() : ViewerFreeIndices.Pop();
 		FViewerInfo& NewViewerInfo = bAddNew ? Viewers.AddDefaulted_GetRef() : Viewers[NewIdx];
-		NewViewerInfo.PlayerController = PlayerController;
 		NewViewerInfo.StreamingSourceName = StreamingSourceName;
 		NewViewerInfo.Handle.Index = NewIdx;
 		NewViewerInfo.Handle.SerialNumber = GetNextViewerSerialNumber();
@@ -277,14 +355,32 @@ void UMassLODSubsystem::AddViewer(APlayerController* PlayerController, FName Str
 
 		ViewerHandle = NewViewerInfo.Handle;
 
-		OnViewerAddedDelegate.Broadcast(NewViewerInfo.Handle, NewViewerInfo.PlayerController, NewViewerInfo.StreamingSourceName);
+		OnViewerAddedDelegate.Broadcast(NewViewerInfo);
 	}
+}
 
-	if(PlayerController)
+#if WITH_EDITOR
+void UMassLODSubsystem::AddEditorViewer(const int32 HashValue, const int32 ClientIndex)
+{
+	FMassViewerHandle& ViewerHandle = ViewerMap.FindOrAdd(HashValue, FMassViewerHandle());
+	// only add new viewer if it hasn't been added yet
+	if (ViewerHandle.IsValid() == false)
 	{
-		PlayerController->OnEndPlay.AddUniqueDynamic(this, &UMassLODSubsystem::OnPlayerControllerEndPlay);
+		const bool bAddNew = ViewerFreeIndices.Num() == 0;
+		const int NewIdx = bAddNew ? Viewers.Num() : ViewerFreeIndices.Pop();
+		FViewerInfo& NewViewerInfo = bAddNew ? Viewers.AddDefaulted_GetRef() : Viewers[NewIdx];
+
+		NewViewerInfo.EditorViewportClientIndex = ClientIndex;
+		NewViewerInfo.Handle.Index = NewIdx;
+		NewViewerInfo.Handle.SerialNumber = GetNextViewerSerialNumber();
+		NewViewerInfo.HashValue = HashValue;
+
+		ViewerHandle = NewViewerInfo.Handle;
+
+		OnViewerAddedDelegate.Broadcast(NewViewerInfo);
 	}
 }
+#endif // WITH_EDITOR
 
 void UMassLODSubsystem::RemoveViewer(const FMassViewerHandle& ViewerHandle)
 {
@@ -317,7 +413,7 @@ void UMassLODSubsystem::RemoveViewerInternal(const FMassViewerHandle& ViewerHand
 	check(ViewerIdx != INDEX_NONE);
 	FViewerInfo& ViewerInfo = Viewers[ViewerIdx];
 
-	OnViewerRemovedDelegate.Broadcast(ViewerInfo.Handle, ViewerInfo.PlayerController, ViewerInfo.StreamingSourceName);
+	OnViewerRemovedDelegate.Broadcast(ViewerInfo);
 
 	if (ViewerInfo.PlayerController)
 	{
@@ -347,10 +443,17 @@ void FViewerInfo::Reset()
 {
 	Handle.Invalidate();
 	PlayerController = nullptr;
+#if WITH_EDITOR
+	EditorViewportClientIndex = INDEX_NONE;
+#endif // WITH_EDITOR
 	HashValue = 0;
 }
 
 bool FViewerInfo::IsLocal() const
 {
-	return (PlayerController && PlayerController->IsLocalController()) || !StreamingSourceName.IsNone();
+	return (PlayerController && PlayerController->IsLocalController()) || !StreamingSourceName.IsNone()
+#if WITH_EDITOR
+		|| EditorViewportClientIndex != INDEX_NONE
+#endif // WITH_EDITOR
+		;
 }
diff --git a/Engine/Plugins/Runtime/MassGameplay/Source/MassLOD/Public/MassLODSubsystem.h b/Engine/Plugins/Runtime/MassGameplay/Source/MassLOD/Public/MassLODSubsystem.h
--- a/Engine/Plugins/Runtime/MassGameplay/Source/MassLOD/Public/MassLODSubsystem.h
+++ b/Engine/Plugins/Runtime/MassGameplay/Source/MassLOD/Public/MassLODSubsystem.h
@@ -33,6 +33,10 @@ struct FViewerInfo
 	
 	FName StreamingSourceName;
 
+#if WITH_EDITOR
+	int8 EditorViewportClientIndex = INDEX_NONE;
+#endif // WITH_EDITOR
+
 	FMassViewerHandle Handle;
 	uint32 HashValue = 0;
 
@@ -48,7 +52,9 @@ struct FViewerInfo
 	bool IsLocal() const;
 };
 
+UE_DEPRECATED(5.3, "FOnViewerAdded is deprecated. Use UMassLODSubsystem::FOnViewerAdded instead.")
 DECLARE_MULTICAST_DELEGATE_ThreeParams(FOnViewerAdded, FMassViewerHandle ViewerHandle, APlayerController* PlayerController, FName StreamingSourceName);
+UE_DEPRECATED(5.3, "FOnViewerRemoved is deprecated. Use UMassLODSubsystem::FOnViewerRemoved instead.")
 DECLARE_MULTICAST_DELEGATE_ThreeParams(FOnViewerRemoved, FMassViewerHandle ViewerHandle, APlayerController* PlayerController, FName StreamingSourceName);
 
 /*
@@ -59,6 +65,9 @@ class MASSLOD_API UMassLODSubsystem : public UTickableWorldSubsystem
 {
 	GENERATED_BODY()
 
+	DECLARE_MULTICAST_DELEGATE_OneParam(FOnViewerAdded, const FViewerInfo& ViewerInfo);
+	DECLARE_MULTICAST_DELEGATE_OneParam(FOnViewerRemoved, const FViewerInfo& ViewerInfo);
+
 public:
 	/** Checks the validity of a viewer handle */
 	bool IsValidViewer(const FMassViewerHandle& ViewerHandle) const { return GetValidViewerIdx(ViewerHandle) != INDEX_NONE; }
@@ -98,9 +107,20 @@ protected:
 	/** Synchronizes the viewers from the engine PlayerController list */
 	void SynchronizeViewers();
 
-	/** Adds a viewer to the list and sends notification about addition */
+	UE_DEPRECATED(5.3, "AddViewer has been deprecated. Use AddPlayerViewer or AddStreamingSourceViewer instead")
 	void AddViewer(APlayerController* PlayerController, FName StreamingSourceName = NAME_None);
 
+	/** Adds the given player as a viewer to the list and sends notification about addition */
+	void AddPlayerViewer(APlayerController& PlayerController);
+
+	/** Adds the given streaming source as a viewer to the list and sends notification about addition */
+	void AddStreamingSourceViewer(const FName StreamingSourceName);
+
+#if WITH_EDITOR
+	/** Adds the editor viewport client (identified via an index) as a viewer to the list and sends notification about addition */
+	void AddEditorViewer(const int32 HashValue, const int32 ClientIndex);
+#endif // WITH_EDITOR
+
 	/** Removes a viewer to the list and send notification about removal */
 	void RemoveViewer(const FMassViewerHandle& ViewerHandle);
 
@@ -128,6 +148,10 @@ private:
 	/** Viewer serial number counter */
 	uint32 ViewerSerialNumberCounter = 0;
 
+#if WITH_EDITOR
+	bool bUseEditorLevelViewports = false;
+#endif // WITH_EDITOR
+
 	/** Free list of indices in the sparse viewer array */
 	TArray<int32> ViewerFreeIndices;
 
