Linux editor build - missing StructUtils dependency for Mass Editor modules

diff --git a/Engine/Plugins/AI/MassAI/Source/MassNavigationEditor/MassNavigationEditor.Build.cs b/Engine/Plugins/AI/MassAI/Source/MassNavigationEditor/MassNavigationEditor.Build.cs
--- a/Engine/Plugins/AI/MassAI/Source/MassNavigationEditor/MassNavigationEditor.Build.cs
+++ b/Engine/Plugins/AI/MassAI/Source/MassNavigationEditor/MassNavigationEditor.Build.cs
@@ -30,6 +30,7 @@ namespace UnrealBuildTool.Rules
 				"MassNavigation",
 				"MassZoneGraphNavigation",
 				"ZoneGraph",
+				"StructUtils"
 			}
 			);
 
diff --git a/Engine/Plugins/Runtime/MassEntity/Source/MassEntityEditor/MassEntityEditor.Build.cs b/Engine/Plugins/Runtime/MassEntity/Source/MassEntityEditor/MassEntityEditor.Build.cs
--- a/Engine/Plugins/Runtime/MassEntity/Source/MassEntityEditor/MassEntityEditor.Build.cs
+++ b/Engine/Plugins/Runtime/MassEntity/Source/MassEntityEditor/MassEntityEditor.Build.cs
@@ -28,7 +28,8 @@ namespace UnrealBuildTool.Rules
 				"DetailCustomizations",
 				"ComponentVisualizers",
 				"Projects",
-				"EditorSubsystem"
+				"EditorSubsystem",
+				"StructUtils"
 			}
 			);
 
diff --git a/Engine/Plugins/Runtime/MassGameplay/Source/MassGameplayEditor/MassGameplayEditor.Build.cs b/Engine/Plugins/Runtime/MassGameplay/Source/MassGameplayEditor/MassGameplayEditor.Build.cs
--- a/Engine/Plugins/Runtime/MassGameplay/Source/MassGameplayEditor/MassGameplayEditor.Build.cs
+++ b/Engine/Plugins/Runtime/MassGameplay/Source/MassGameplayEditor/MassGameplayEditor.Build.cs
@@ -32,7 +32,8 @@ namespace UnrealBuildTool.Rules
 				"DetailCustomizations",
 				"ComponentVisualizers",
 				"Projects",
-				"EditorSubsystem"
+				"EditorSubsystem",
+				"StructUtils"
 			}
 			);
 
