Switched Mass's "a fragment added by more than one trait" Error over to be reported as a Warning, and fixed where those errors/warnings end up shown in the editor.

diff --git a/Engine/Plugins/Runtime/MassGameplay/Source/MassSpawner/Private/MassEntityConfigAsset.cpp b/Engine/Plugins/Runtime/MassGameplay/Source/MassSpawner/Private/MassEntityConfigAsset.cpp
--- a/Engine/Plugins/Runtime/MassGameplay/Source/MassSpawner/Private/MassEntityConfigAsset.cpp
+++ b/Engine/Plugins/Runtime/MassGameplay/Source/MassSpawner/Private/MassEntityConfigAsset.cpp
@@ -55,9 +55,9 @@ void UMassEntityConfigAsset::ValidateEntityConfig()
 	{
 		if (Config.ValidateEntityTemplate(*EditorWorld))
 		{
-			const FText InfoText = LOCTEXT("MassEntityConfigAssetNoErrorsDetected", "There were no error detected during validation of the EntityConfigAsset");
+			const FText InfoText = LOCTEXT("MassEntityConfigAssetNoErrorsDetected", "There were no errors nor warnings detected during validation of the EntityConfigAsset");
 			
-			FMessageLog EditorInfo("LogMass");
+			FMessageLog EditorInfo("MassEntity");
 			EditorInfo.Info(InfoText);
 
 			FNotificationInfo Info(InfoText);
diff --git a/Engine/Plugins/Runtime/MassGameplay/Source/MassSpawner/Private/MassEntityTemplateRegistry.cpp b/Engine/Plugins/Runtime/MassGameplay/Source/MassSpawner/Private/MassEntityTemplateRegistry.cpp
--- a/Engine/Plugins/Runtime/MassGameplay/Source/MassSpawner/Private/MassEntityTemplateRegistry.cpp
+++ b/Engine/Plugins/Runtime/MassGameplay/Source/MassSpawner/Private/MassEntityTemplateRegistry.cpp
@@ -176,11 +176,11 @@ bool FMassEntityTemplateBuildContext::ValidateBuildContext(const UWorld& World)
 		{
 			if (!bHeaderOutputed)
 			{
-				UE_LOG(LogMass, Error, TEXT("Fragment(%s) was added multiple time and can only be added by one trait. Fragment was added by:"), CurrentStruct ? *CurrentStruct->GetName() : TEXT("null"));
-				UE_LOG(LogMass, Error, TEXT("\t\t%s"), CurrentTrait ? *CurrentTrait->GetClass()->GetName() : TEXT("null"));
+				UE_LOG(LogMass, Warning, TEXT("Fragment(%s) was added multiple time and should only be added by one trait. Fragment was added by:"), CurrentStruct ? *CurrentStruct->GetName() : TEXT("null"));
+				UE_LOG(LogMass, Warning, TEXT("\t\t%s"), CurrentTrait ? *CurrentTrait->GetClass()->GetName() : TEXT("null"));
 				bHeaderOutputed = true;
 			}
-			UE_LOG(LogMass, Error, TEXT("\t\t%s"), *Pair.Value->GetClass()->GetName());
+			UE_LOG(LogMass, Warning, TEXT("\t\t%s"), *Pair.Value->GetClass()->GetName());
 			bFragmentHasMultipleOwners = true;
 		}
  	}
@@ -212,10 +212,10 @@ bool FMassEntityTemplateBuildContext::ValidateBuildContext(const UWorld& World)
 #if WITH_UNREAL_DEVELOPER_TOOLS
 	if (bFragmentHasMultipleOwners || bMissingFragmentDependencies)
 	{
-		FMessageLog EditorErrors("LogMass");
+		FMessageLog EditorErrors("MassEntity");
 		if (bFragmentHasMultipleOwners)
 		{
-			EditorErrors.Error(LOCTEXT("MassEntityTraitsFragmentOwnershipError", "Some fragments are added by multiple traits and can only be added by one!"));
+			EditorErrors.Warning(LOCTEXT("MassEntityTraitsFragmentOwnershipError", "Some fragments are added by multiple traits and can only be added by one!"));
 			EditorErrors.Notify(LOCTEXT("MassEntityTraitsFragmentOwnershipError", "Some fragments are added by multiple traits and can only be added by one!"));
 		}
 		if (bMissingFragmentDependencies)
