Moved UE::MassSpawner::HashTraits out of MassEntityConfigAsset.cpp into an accesible location, and switched the parameter to TConstArrayView

This change is a part of MassEntityConfig and MassEntityTemplate changes.

diff --git a/Engine/Plugins/Runtime/MassGameplay/Source/MassSpawner/Private/MassEntityConfigAsset.cpp b/Engine/Plugins/Runtime/MassGameplay/Source/MassSpawner/Private/MassEntityConfigAsset.cpp
--- a/Engine/Plugins/Runtime/MassGameplay/Source/MassSpawner/Private/MassEntityConfigAsset.cpp
+++ b/Engine/Plugins/Runtime/MassGameplay/Source/MassSpawner/Private/MassEntityConfigAsset.cpp
@@ -17,30 +17,6 @@
 
 #define LOCTEXT_NAMESPACE "Mass"
 
-namespace UE::MassSpawner
-{
-	uint32 HashTraits(const TArray<UMassEntityTraitBase*>& CombinedTraits)
-	{
-		class FArchiveObjectCRC32AgentConfig : public FArchiveObjectCrc32
-		{
-		public:
-			virtual bool ShouldSkipProperty(const FProperty* InProperty) const override
-			{
-				check(InProperty);
-				return FArchiveObjectCrc32::ShouldSkipProperty(InProperty) || InProperty->HasAllPropertyFlags(CPF_Transient);
-			}
-		};
-
-		uint32 CRC = 0;
-		for (UMassEntityTraitBase* Trait : CombinedTraits)
-		{
-			FArchiveObjectCRC32AgentConfig Archive;
-			CRC = Archive.Crc32(Trait, CRC);
-		}
-		return CRC;
-	}
-} // UE::MassSpawner
-
 
 FMassEntityConfig::FMassEntityConfig(UObject& InOwner)
 	: ConfigOwner(&InOwner)
diff --git a/Engine/Plugins/Runtime/MassGameplay/Source/MassSpawner/Private/MassSpawnerTypes.cpp b/Engine/Plugins/Runtime/MassGameplay/Source/MassSpawner/Private/MassSpawnerTypes.cpp
--- a/Engine/Plugins/Runtime/MassGameplay/Source/MassSpawner/Private/MassSpawnerTypes.cpp
+++ b/Engine/Plugins/Runtime/MassGameplay/Source/MassSpawner/Private/MassSpawnerTypes.cpp
@@ -2,7 +2,36 @@
 
 #include "MassSpawnerTypes.h"
 #include "MassEntityConfigAsset.h"
+#include "Serialization/ArchiveObjectCrc32.h"
+#include "MassEntityTraitBase.h"
 
+namespace UE::MassSpawner
+{
+	uint32 HashTraits(TConstArrayView<UMassEntityTraitBase*> CombinedTraits)
+	{
+		class FArchiveObjectCRC32AgentConfig : public FArchiveObjectCrc32
+		{
+		public:
+			virtual bool ShouldSkipProperty(const FProperty* InProperty) const override
+			{
+				check(InProperty);
+				return FArchiveObjectCrc32::ShouldSkipProperty(InProperty) || InProperty->HasAllPropertyFlags(CPF_Transient);
+			}
+		};
+
+		uint32 CRC = 0;
+		for (UMassEntityTraitBase* Trait : CombinedTraits)
+		{
+			FArchiveObjectCRC32AgentConfig Archive;
+			CRC = Archive.Crc32(Trait, CRC);
+		}
+		return CRC;
+	}
+} // UE::MassSpawner
+
+//-----------------------------------------------------------------------------
+// FMassSpawnedEntityType
+//-----------------------------------------------------------------------------
 const UMassEntityConfigAsset* FMassSpawnedEntityType::GetEntityConfig() const
 {
 	if (EntityConfigPtr == nullptr)
diff --git a/Engine/Plugins/Runtime/MassGameplay/Source/MassSpawner/Public/MassSpawnerTypes.h b/Engine/Plugins/Runtime/MassGameplay/Source/MassSpawner/Public/MassSpawnerTypes.h
--- a/Engine/Plugins/Runtime/MassGameplay/Source/MassSpawner/Public/MassSpawnerTypes.h
+++ b/Engine/Plugins/Runtime/MassGameplay/Source/MassSpawner/Public/MassSpawnerTypes.h
@@ -13,6 +13,12 @@ class APawn;
 class UCurveFloat;
 class UMassEntityConfigAsset;
 class UMassEntitySpawnDataGeneratorBase;
+class UMassEntityTraitBase;
+
+namespace UE::MassSpawner
+{
+	MASSSPAWNER_API uint32 HashTraits(TConstArrayView<UMassEntityTraitBase*> CombinedTraits);
+}
 
 USTRUCT()
 struct FMassTransformsSpawnData
