Added an "IsEmpty" function to MassEntityConfig

diff --git a/Engine/Plugins/Runtime/MassGameplay/Source/MassSpawner/Public/MassEntityConfigAsset.h b/Engine/Plugins/Runtime/MassGameplay/Source/MassSpawner/Public/MassEntityConfigAsset.h
--- a/Engine/Plugins/Runtime/MassGameplay/Source/MassSpawner/Public/MassEntityConfigAsset.h
+++ b/Engine/Plugins/Runtime/MassGameplay/Source/MassSpawner/Public/MassEntityConfigAsset.h
@@ -66,6 +66,8 @@ struct MASSSPAWNER_API FMassEntityConfig
 	UE_DEPRECATED(5.3, "This flavor of ValidateEntityTemplate is deperecated. Use the one without the ConfigOwner parameter (now a property of the FMassEntityConfig itself)")
 	bool ValidateEntityTemplate(const UWorld& World, const UObject& ConfigOwner);
 
+	bool IsEmpty() const { return Parent == nullptr && Traits.Num() == 0; }
+
 protected:
 	/** Combines traits based on the config hierarchy and returns list of unique traits */
 	void GetCombinedTraits(TArray<UMassEntityTraitBase*>& OutTraits, TArray<const UObject*>& Visited) const;
