MassEntityConfig and MassEntityTemplate changes
- separated MassEntityTemplate's configurable data into a separate MassEntityTemplateData type, and made MassEntityTemplate represent a read-only, finalized and instantiated (i.e. tied to an archetype and template ID) template
	The main reason for the separation is to have a type (the FMassEntityTemplateData ) that can be used to:
	* procedurally build templates - for procedural spawning
	* be used as a "prefab" (i.e. has data, but is not associated with an actual archetype)
	With this change the FMassEntityTemplate is a "ready to be used" template, that's associated with an archetype and has a template ID (both of which the Data is missing)
- added reliable trait/template hash calculation
- added a factory for TemplateID creation to keep it consistent
- MassEntityTemplateRegistry now returns templates as const TSharedRefs

The change is also intended to further enforce that "finalized" templates (i.e. the ones registered with the TemplateRegistry) should not be modified.

diff --git a/Engine/Plugins/Runtime/MassGameplay/Source/MassActors/Private/MassAgentSubsystem.cpp b/Engine/Plugins/Runtime/MassGameplay/Source/MassActors/Private/MassAgentSubsystem.cpp
--- a/Engine/Plugins/Runtime/MassGameplay/Source/MassActors/Private/MassAgentSubsystem.cpp
+++ b/Engine/Plugins/Runtime/MassGameplay/Source/MassActors/Private/MassAgentSubsystem.cpp
@@ -20,11 +20,12 @@
 namespace FMassAgentSubsystemHelper
 {
 
-inline void InitializeAgentComponentFragments(const UMassAgentComponent& AgentComp, FMassEntityView& EntityView, const EMassTranslationDirection Direction, TConstArrayView<FMassEntityTemplate::FObjectFragmentInitializerFunction> ObjectFragmentInitializers)
+inline void InitializeAgentComponentFragments(const UMassAgentComponent& AgentComp, FMassEntityView& EntityView, const EMassTranslationDirection Direction
+	, TConstArrayView<FMassEntityTemplateData::FObjectFragmentInitializerFunction> ObjectFragmentInitializers)
 {
 	AActor* Owner = AgentComp.GetOwner();
 	check(Owner);
-	for (const FMassEntityTemplate::FObjectFragmentInitializerFunction& Initializer : ObjectFragmentInitializers)
+	for (const FMassEntityTemplateData::FObjectFragmentInitializerFunction& Initializer : ObjectFragmentInitializers)
 	{
 		Initializer(*Owner, EntityView, Direction);
 	}
@@ -306,8 +307,9 @@ void UMassAgentSubsystem::HandlePendingInitialization()
 	for (TTuple<FMassEntityTemplateID, FMassAgentInitializationQueue>& Data : PendingAgentEntities)
 	{
 		const FMassEntityTemplateID EntityTemplateID = Data.Get<0>();
-		const FMassEntityTemplate* EntityTemplate = SpawnerSystem->GetTemplateRegistryInstance().FindTemplateFromTemplateID(EntityTemplateID);
-		check(EntityTemplate);
+		const TSharedRef<FMassEntityTemplate>* EntityTemplatePtr = SpawnerSystem->GetTemplateRegistryInstance().FindTemplateFromTemplateID(EntityTemplateID);
+		check(EntityTemplatePtr);
+		const FMassEntityTemplate& EntityTemplate = EntityTemplatePtr->Get();
 		
 		TArray<UMassAgentComponent*>& AgentComponents = Data.Get<1>().AgentComponents;
 		const int32 NewEntityCount = AgentComponents.Num();
@@ -319,16 +321,16 @@ void UMassAgentSubsystem::HandlePendingInitialization()
 		}
 		
 		TArray<FMassEntityHandle> Entities;
-		SpawnerSystem->SpawnEntities(*EntityTemplate, NewEntityCount, Entities);
+		SpawnerSystem->SpawnEntities(EntityTemplate, NewEntityCount, Entities);
 		check(Entities.Num() == NewEntityCount);
 
-		if (EntityTemplate->GetObjectFragmentInitializers().Num())
+		if (EntityTemplate.GetObjectFragmentInitializers().Num())
 		{
-			const TConstArrayView<FMassEntityTemplate::FObjectFragmentInitializerFunction> ObjectFragmentInitializers = EntityTemplate->GetObjectFragmentInitializers();
+			const TConstArrayView<FMassEntityTemplateData::FObjectFragmentInitializerFunction> ObjectFragmentInitializers = EntityTemplate.GetObjectFragmentInitializers();
 
 			for (int AgentIndex = 0; AgentIndex < Entities.Num(); ++AgentIndex)
 			{
-				FMassEntityView EntityView(EntityTemplate->GetArchetype(), Entities[AgentIndex]);
+				FMassEntityView EntityView(EntityTemplate.GetArchetype(), Entities[AgentIndex]);
 				FMassAgentSubsystemHelper::InitializeAgentComponentFragments(*AgentComponents[AgentIndex], EntityView, EMassTranslationDirection::ActorToMass, ObjectFragmentInitializers);
 			}
 		}
@@ -344,14 +346,16 @@ void UMassAgentSubsystem::HandlePendingInitialization()
 	for (TTuple<FMassEntityTemplateID, FMassAgentInitializationQueue>& Data : PendingPuppets)
 	{
 		const FMassEntityTemplateID EntityTemplateID = Data.Get<0>();
-		const FMassEntityTemplate* EntityTemplate = SpawnerSystem->GetTemplateRegistryInstance().FindTemplateFromTemplateID(EntityTemplateID);
-		if (!ensure(EntityTemplate))
+		const TSharedRef<FMassEntityTemplate>* EntityTemplatePtr = SpawnerSystem->GetTemplateRegistryInstance().FindTemplateFromTemplateID(EntityTemplateID);
+
+		if (!ensure(EntityTemplatePtr))
 		{
 			// note that this condition is temporary, we'll be switched to a `check` once we set up characters
 			continue;
 		}
+		const FMassEntityTemplate& EntityTemplate = EntityTemplatePtr->Get();
 
-		const FMassArchetypeCompositionDescriptor TemplateDescriptor = EntityTemplate->GetCompositionDescriptor();
+		const FMassArchetypeCompositionDescriptor TemplateDescriptor = EntityTemplate.GetCompositionDescriptor();
 
 		TArray<UMassAgentComponent*>& AgentComponents = Data.Get<1>().AgentComponents;
 
@@ -367,11 +371,11 @@ void UMassAgentSubsystem::HandlePendingInitialization()
 			PuppetDescriptor = TemplateDescriptor;
 			EntityManager->AddCompositionToEntity_GetDelta(PuppetEntity, PuppetDescriptor);
 			
-			if (EntityTemplate->GetObjectFragmentInitializers().Num())
+			if (EntityTemplate.GetObjectFragmentInitializers().Num())
 			{
 				const FMassArchetypeHandle ArchetypeHandle = EntityManager->GetArchetypeForEntity(PuppetEntity);
 				FMassEntityView EntityView(ArchetypeHandle, PuppetEntity);
-				FMassAgentSubsystemHelper::InitializeAgentComponentFragments(*AgentComp, EntityView, EMassTranslationDirection::MassToActor, EntityTemplate->GetObjectFragmentInitializers());
+				FMassAgentSubsystemHelper::InitializeAgentComponentFragments(*AgentComp, EntityView, EMassTranslationDirection::MassToActor, EntityTemplate.GetObjectFragmentInitializers());
 			}
 
 			AgentComp->PuppetInitializationDone();
diff --git a/Engine/Plugins/Runtime/MassGameplay/Source/MassSpawner/Private/MassEntityConfigAsset.cpp b/Engine/Plugins/Runtime/MassGameplay/Source/MassSpawner/Private/MassEntityConfigAsset.cpp
--- a/Engine/Plugins/Runtime/MassGameplay/Source/MassSpawner/Private/MassEntityConfigAsset.cpp
+++ b/Engine/Plugins/Runtime/MassGameplay/Source/MassSpawner/Private/MassEntityConfigAsset.cpp
@@ -67,18 +67,14 @@ const FMassEntityTemplate& FMassEntityConfig::GetOrCreateEntityTemplate(const UW
 	// TODO: Add methods to FMassEntityTemplateBuildContext to indicate dependency vs setup.
 	// Dependency should add a fragment with default values (which later can be overridden),
 	// while setup would override values and should be run just once.
-	FMassEntityTemplate& Template = TemplateRegistry.CreateTemplate(TemplateID);
-	FMassEntityTemplateBuildContext BuildContext(Template);
 
-	BuildContext.BuildFromTraits(CombinedTraits, World);
-	Template.SetTemplateName(GetNameSafe(ConfigOwner));
+	FMassEntityTemplate Template;
+	FMassEntityTemplateBuildContext BuildContext(Template, TemplateID);
 
-	// It is ok to have an empty template, 
-    // but be aware there will be an error if you try to create an entity with it
-    // as there will be no archetype associated with this template...
-	TemplateRegistry.InitializeEntityTemplate(Template);
+	BuildContext.BuildFromTraits(CombinedTraits, World);
+	BuildContext.SetTemplateName(GetNameSafe(ConfigOwner));
 
-	return Template;
+	return TemplateRegistry.AddTemplate(TemplateID, MoveTemp(Template)).Get();
 }
 
 void FMassEntityConfig::DestroyEntityTemplate(const UWorld& World) const
@@ -115,7 +111,7 @@ const FMassEntityTemplate& FMassEntityConfig::GetEntityTemplateChecked(const UWo
 	return *ExistingTemplate;
 }
 
-const FMassEntityTemplate* FMassEntityConfig::GetEntityTemplateInternal(const UWorld& World, FMassEntityTemplateID& TemplateIDOut, TArray<UMassEntityTraitBase*>& CombinedTraitsOut) const
+const FMassEntityTemplate* FMassEntityConfig::GetEntityTemplateInternal(const UWorld& World, FMassEntityTemplateID& OutTemplateID, TArray<UMassEntityTraitBase*>& OutCombinedTraits) const
 {
 	UMassSpawnerSubsystem* SpawnerSystem = UWorld::GetSubsystem<UMassSpawnerSubsystem>(&World);
 	check(SpawnerSystem);
@@ -124,15 +120,15 @@ const FMassEntityTemplate* FMassEntityConfig::GetEntityTemplateInternal(const UW
 	// Combine all the features into one array
 	// @todo this is an inefficient way assuming given template is expected to have already been created. Figure out a way to cache it.
 	TArray<const UObject*> Visited;
-	CombinedTraitsOut.Reset();
+	OutCombinedTraits.Reset();
 	Visited.Add(ConfigOwner);
-	GetCombinedTraits(CombinedTraitsOut, Visited);
+	GetCombinedTraits(OutCombinedTraits, Visited);
 
 	// Return existing template if found.
 	// TODO: cache the hash.
-	const uint32 HashOut = UE::MassSpawner::HashTraits(CombinedTraitsOut);
-	TemplateIDOut = FMassEntityTemplateID(HashOut);
-	return TemplateRegistry.FindTemplateFromTemplateID(TemplateIDOut);
+	OutTemplateID = FMassEntityTemplateIDFactory::Make(OutCombinedTraits);
+	const TSharedRef<FMassEntityTemplate>* TemplateFound = TemplateRegistry.FindTemplateFromTemplateID(OutTemplateID);
+	return TemplateFound ? &TemplateFound->Get() : nullptr;
 }
 
 void FMassEntityConfig::GetCombinedTraits(TArray<UMassEntityTraitBase*>& OutTraits, TArray<const UObject*>& Visited) const
diff --git a/Engine/Plugins/Runtime/MassGameplay/Source/MassSpawner/Private/MassEntityTemplate.cpp b/Engine/Plugins/Runtime/MassGameplay/Source/MassSpawner/Private/MassEntityTemplate.cpp
--- a/Engine/Plugins/Runtime/MassGameplay/Source/MassSpawner/Private/MassEntityTemplate.cpp
+++ b/Engine/Plugins/Runtime/MassGameplay/Source/MassSpawner/Private/MassEntityTemplate.cpp
@@ -3,7 +3,8 @@
 #include "MassEntityTemplate.h"
 #include "VisualLogger/VisualLoggerTypes.h"
 #include "MassDebugger.h"
-
+#include "MassSpawnerTypes.h"
+#include "StructUtilsTypes.h"
 
 //----------------------------------------------------------------------//
 //  FMassEntityTemplateID
@@ -14,16 +15,80 @@ FString FMassEntityTemplateID::ToString() const
 }
 
 //----------------------------------------------------------------------//
-//  FMassEntityTemplate
+//  FMassEntityTemplateData
 //----------------------------------------------------------------------//
+uint32 GetTypeHash(const FMassEntityTemplateData& Template)
+{
+	const uint32 CompositionHash = Template.Composition.CalculateHash();
+	const uint32 SharedFragmentValuesHash = GetTypeHash(Template.SharedFragmentValues);
+	
+	uint32 InitialValuesHash = 0;
+	// Initial fragment values, this is not part of the archetype as it is the spawner job to set them.
+	for (const FInstancedStruct& Struct : Template.InitialFragmentValues)
+	{
+		InitialValuesHash = UE::StructUtils::GetStructCrc32(Struct, InitialValuesHash);
+	}
+
+	// These functions will be called to initialize entity's UObject-based fragments
+	// @todo this is very bad for hash creation - no way to make this consistent between client server. Might need some "initializer index"
+	uint32 InitializersHash = 0;
+	/* left here on purpose, will address in following CLs
+	for (const FMassEntityTemplateData::FObjectFragmentInitializerFunction& Initializer : Template.ObjectInitializers)
+	{
+		InitializersHash = HashCombine(InitializersHash, GetTypeHash(Initializer));
+	}*/
+
+	// @todo maybe better to have two separate hashes - one for composition and maybe shared fragments, and then the other one for the rest, and use multimap for look up
+	return HashCombine(HashCombine(CompositionHash, SharedFragmentValuesHash), HashCombine(InitialValuesHash, InitializersHash));
+}
+
+//-----------------------------------------------------------------------------
+// FMassEntityTemplate
+//-----------------------------------------------------------------------------
+TSharedRef<FMassEntityTemplate> FMassEntityTemplate::MakeFinalTemplate(FMassEntityManager& EntityManager, FMassEntityTemplateData&& TempTemplateData, FMassEntityTemplateID InTemplateID)
+{
+	return MakeShared<FMassEntityTemplate>(MoveTemp(TempTemplateData), EntityManager, InTemplateID);
+}
+
+FMassEntityTemplate::FMassEntityTemplate(const FMassEntityTemplateData& InData, FMassEntityManager& EntityManager, FMassEntityTemplateID InTemplateID)
+	: Super(InData)
+	, TemplateID(InTemplateID)
+{
+	// Sort anything there is to sort for later comparison purposes
+	Sort();
+
+	const FMassArchetypeHandle ArchetypeHandle = EntityManager.CreateArchetype(GetCompositionDescriptor(), FName(GetTemplateName()));
+	SetArchetype(ArchetypeHandle);
+}
+
+FMassEntityTemplate::FMassEntityTemplate(FMassEntityTemplateData&& InData, FMassEntityManager& EntityManager, FMassEntityTemplateID InTemplateID)
+	: Super(InData)
+	, TemplateID(InTemplateID)
+{
+	// Sort anything there is to sort for later comparison purposes
+	Sort();
+
+	const FMassArchetypeHandle ArchetypeHandle = EntityManager.CreateArchetype(GetCompositionDescriptor(), FName(GetTemplateName()));
+	SetArchetype(ArchetypeHandle);
+}
+
 void FMassEntityTemplate::SetArchetype(const FMassArchetypeHandle& InArchetype)
 {
 	check(InArchetype.IsValid());
 	Archetype = InArchetype;
 }
 
+FString FMassEntityTemplate::DebugGetArchetypeDescription(FMassEntityManager& EntityManager) const
+{
+	FStringOutputDevice OutDescription;
+#if WITH_MASSGAMEPLAY_DEBUG
+	FMassDebugger::OutputArchetypeDescription(OutDescription, Archetype);
+#endif // WITH_MASSGAMEPLAY_DEBUG
+	return MoveTemp(OutDescription);
+}
+
 FString FMassEntityTemplate::DebugGetDescription(FMassEntityManager* EntityManager) const
-{ 
+{
 	FStringOutputDevice Ar;
 #if WITH_MASSGAMEPLAY_DEBUG
 	Ar.SetAutoEmitLineTerminator(true);
@@ -43,11 +108,16 @@ FString FMassEntityTemplate::DebugGetDescription(FMassEntityManager* EntityManag
 	return MoveTemp(Ar);
 }
 
-FString FMassEntityTemplate::DebugGetArchetypeDescription(FMassEntityManager& EntityManager) const
+//-----------------------------------------------------------------------------
+// FMassEntityTemplateIDFactory
+//-----------------------------------------------------------------------------
+FMassEntityTemplateID FMassEntityTemplateIDFactory::Make(const FMassEntityTemplateData& TemplateData)
 {
-	FStringOutputDevice OutDescription;
-#if WITH_MASSGAMEPLAY_DEBUG
-	FMassDebugger::OutputArchetypeDescription(OutDescription, Archetype);
-#endif // WITH_MASSGAMEPLAY_DEBUG
-	return MoveTemp(OutDescription);
+	return FMassEntityTemplateID(GetTypeHash(TemplateData));
+}
+
+FMassEntityTemplateID FMassEntityTemplateIDFactory::Make(TConstArrayView<UMassEntityTraitBase*> Traits)
+{
+	const uint32 HashOut = UE::MassSpawner::HashTraits(Traits);
+	return FMassEntityTemplateID(HashOut);
 }
diff --git a/Engine/Plugins/Runtime/MassGameplay/Source/MassSpawner/Private/MassEntityTemplateRegistry.cpp b/Engine/Plugins/Runtime/MassGameplay/Source/MassSpawner/Private/MassEntityTemplateRegistry.cpp
--- a/Engine/Plugins/Runtime/MassGameplay/Source/MassSpawner/Private/MassEntityTemplateRegistry.cpp
+++ b/Engine/Plugins/Runtime/MassGameplay/Source/MassSpawner/Private/MassEntityTemplateRegistry.cpp
@@ -48,8 +48,8 @@ namespace FTemplateRegistryHelpers
 		TEXT("Clears all the runtime information cached by MassEntityTemplateRegistry. Will result in lazily building all entity templates again."),
 		FConsoleCommandWithWorldAndArgsDelegate::CreateStatic(&ResetEntityTemplates)
 	);
+} // namespace FTemplateRegistryHelpers
 
-}
 
 //----------------------------------------------------------------------//
 // FMassEntityTemplateRegistry 
@@ -61,6 +61,12 @@ FMassEntityTemplateRegistry::FMassEntityTemplateRegistry(UObject* InOwner)
 {
 }
 
+void FMassEntityTemplateRegistry::ShutDown()
+{
+	TemplateIDToTemplateMap.Reset();
+	EntityManager = nullptr;
+}
+
 UWorld* FMassEntityTemplateRegistry::GetWorld() const 
 {
 	return Owner.IsValid() ? Owner->GetWorld() : nullptr;
@@ -71,35 +77,15 @@ FMassEntityTemplateRegistry::FStructToTemplateBuilderDelegate& FMassEntityTempla
 	return StructBasedBuilders.FindOrAdd(&DataType);
 }
 
-bool FMassEntityTemplateRegistry::BuildTemplateImpl(const FStructToTemplateBuilderDelegate& Builder, const FConstStructView StructInstance, FMassEntityTemplate& OutTemplate)
+void FMassEntityTemplateRegistry::StoreEntityManager(const TSharedPtr<FMassEntityManager>& InEntityManager)
 {
-	UWorld* World = GetWorld();
-	FMassEntityTemplateBuildContext Context(OutTemplate);
-	Builder.Execute(World, StructInstance, Context);
-	if (ensure(!OutTemplate.IsEmpty())) // need at least one fragment to create an Archetype
+	if (EntityManager)
 	{
-		InitializeEntityTemplate(OutTemplate);
-
-		UE_VLOG(Owner.Get(), LogMassSpawner, Log, TEXT("Created entity template for %s:\n%s"), *GetNameSafe(StructInstance.GetScriptStruct())
-			, *OutTemplate.DebugGetDescription(UE::Mass::Utils::GetEntityManager(World)));
-
-		return true;
+		ensureMsgf(EntityManager == InEntityManager, TEXT("Attempting to store a different EntityManager then the previously stored one - this indicated a set up issue, attempting to use multiple EntityManager instances"));
+		return;
 	}
-	return false;
-}
 
-void FMassEntityTemplateRegistry::InitializeEntityTemplate(FMassEntityTemplate& InOutTemplate) const
-{
-	UWorld* World = GetWorld();
-	check(World);
-	// find or create template
-	FMassEntityManager& EntityManager = UE::Mass::Utils::GetEntityManagerChecked(*World);
-
-	// Sort anything there is to sort for later comparison purposes
-	InOutTemplate.Sort();
-
-	const FMassArchetypeHandle ArchetypeHandle = EntityManager.CreateArchetype(InOutTemplate.GetCompositionDescriptor(), FName(InOutTemplate.GetTemplateName()));
-	InOutTemplate.SetArchetype(ArchetypeHandle);
+	EntityManager = InEntityManager;
 }
 
 void FMassEntityTemplateRegistry::DebugReset()
@@ -109,22 +95,36 @@ void FMassEntityTemplateRegistry::DebugReset()
 #endif // WITH_MASSGAMEPLAY_DEBUG
 }
 
-const FMassEntityTemplate* FMassEntityTemplateRegistry::FindTemplateFromTemplateID(FMassEntityTemplateID TemplateID) const 
+const TSharedRef<FMassEntityTemplate>* FMassEntityTemplateRegistry::FindTemplateFromTemplateID(FMassEntityTemplateID TemplateID) const
 {
 	return TemplateIDToTemplateMap.Find(TemplateID);
 }
 
-FMassEntityTemplate* FMassEntityTemplateRegistry::FindMutableTemplateFromTemplateID(FMassEntityTemplateID TemplateID) 
+const TSharedRef<FMassEntityTemplate>& FMassEntityTemplateRegistry::AddTemplate(FMassEntityTemplateID TemplateID, FMassEntityTemplateData&& TemplateData)
 {
-	return TemplateIDToTemplateMap.Find(TemplateID);
+	check(EntityManager);
+	const TSharedRef<FMassEntityTemplate>* ExistingTemplate = FindTemplateFromTemplateID(TemplateID);
+	if (!ensureMsgf(ExistingTemplate == nullptr, TEXT("")))
+	{
+		return *ExistingTemplate;
+	}
+
+	return TemplateIDToTemplateMap.Add(TemplateID, FMassEntityTemplate::MakeFinalTemplate(*EntityManager, MoveTemp(TemplateData), TemplateID));
 }
 
-FMassEntityTemplate& FMassEntityTemplateRegistry::CreateTemplate(FMassEntityTemplateID TemplateID)
+const TSharedRef<FMassEntityTemplate>& FMassEntityTemplateRegistry::AddTemplate(FMassEntityTemplateData&& TemplateData)
 {
-	checkSlow(!TemplateIDToTemplateMap.Contains(TemplateID));
-	FMassEntityTemplate& NewTemplate = TemplateIDToTemplateMap.Add(TemplateID);
-	NewTemplate.SetTemplateID(TemplateID);
-	return NewTemplate;
+	const FMassEntityTemplateID NewTemplateID = FMassEntityTemplateIDFactory::Make(TemplateData);
+
+	check(EntityManager);
+	const TSharedRef<FMassEntityTemplate>* ExistingTemplate = FindTemplateFromTemplateID(NewTemplateID);
+
+	if (!ensureMsgf(ExistingTemplate == nullptr, TEXT("")))
+	{
+		return *ExistingTemplate;
+	}
+
+	return TemplateIDToTemplateMap.Add(NewTemplateID, FMassEntityTemplate::MakeFinalTemplate(*EntityManager, MoveTemp(TemplateData), NewTemplateID));
 }
 
 void FMassEntityTemplateRegistry::DestroyTemplate(FMassEntityTemplateID TemplateID)
@@ -230,4 +230,21 @@ bool FMassEntityTemplateBuildContext::ValidateBuildContext(const UWorld& World)
 	return !bFragmentHasMultipleOwners && !bMissingFragmentDependencies;
 }
 
+//-----------------------------------------------------------------------------
+// DEPRECATED
+//-----------------------------------------------------------------------------
+FMassEntityTemplate* FMassEntityTemplateRegistry::FindMutableTemplateFromTemplateID(FMassEntityTemplateID TemplateID)
+{
+	return nullptr;
+}
+
+FMassEntityTemplate& FMassEntityTemplateRegistry::CreateTemplate(const uint32 HashLookup, FMassEntityTemplateID TemplateID)
+{
+	static FMassEntityTemplate Dummy;
+	return Dummy;
+}
+
+void FMassEntityTemplateRegistry::InitializeEntityTemplate(FMassEntityTemplate& InOutTemplate) const
+{}
+
 #undef LOCTEXT_NAMESPACE 
diff --git a/Engine/Plugins/Runtime/MassGameplay/Source/MassSpawner/Private/MassSpawnerSubsystem.cpp b/Engine/Plugins/Runtime/MassGameplay/Source/MassSpawner/Private/MassSpawnerSubsystem.cpp
--- a/Engine/Plugins/Runtime/MassGameplay/Source/MassSpawner/Private/MassSpawnerSubsystem.cpp
+++ b/Engine/Plugins/Runtime/MassGameplay/Source/MassSpawner/Private/MassSpawnerSubsystem.cpp
@@ -36,11 +36,13 @@ void UMassSpawnerSubsystem::Initialize(FSubsystemCollectionBase& Collection)
 	UWorld* World = GetWorld();
 	check(World);
 	EntityManager = UE::Mass::Utils::GetEntityManagerChecked(*World).AsShared();
+	TemplateRegistryInstance.StoreEntityManager(EntityManager);
 }
 
 void UMassSpawnerSubsystem::Deinitialize() 
 {
 	EntityManager.Reset();
+	TemplateRegistryInstance.ShutDown();
 }
 
 void UMassSpawnerSubsystem::SpawnEntities(const FMassEntityTemplate& EntityTemplate, const uint32 NumberToSpawn, TArray<FMassEntityHandle>& OutEntities)
@@ -61,10 +63,10 @@ void UMassSpawnerSubsystem::SpawnEntities(FMassEntityTemplateID TemplateID, cons
 {
 	check(TemplateID.IsValid());
 
-	const FMassEntityTemplate* EntityTemplate = TemplateRegistryInstance.FindTemplateFromTemplateID(TemplateID);
-	checkf(EntityTemplate && EntityTemplate->IsValid(), TEXT("SpawnEntities: TemplateID must have been registered!"));
+	const TSharedRef<FMassEntityTemplate>* EntityTemplate = TemplateRegistryInstance.FindTemplateFromTemplateID(TemplateID);
+	checkf(EntityTemplate, TEXT("SpawnEntities: TemplateID must have been registered!"));
 
-	DoSpawning(*EntityTemplate, NumberToSpawn, SpawnData, InitializerClass, OutEntities);
+	DoSpawning(EntityTemplate->Get(), NumberToSpawn, SpawnData, InitializerClass, OutEntities);
 }
 
 void UMassSpawnerSubsystem::DestroyEntities(const FMassEntityTemplateID TemplateID, TConstArrayView<FMassEntityHandle> Entities)
@@ -153,5 +155,6 @@ void UMassSpawnerSubsystem::DoSpawning(const FMassEntityTemplate& EntityTemplate
 const FMassEntityTemplate* UMassSpawnerSubsystem::GetMassEntityTemplate(FMassEntityTemplateID TemplateID) const
 {
 	check(TemplateID.IsValid());
-	return TemplateRegistryInstance.FindTemplateFromTemplateID(TemplateID);
+	const TSharedRef<FMassEntityTemplate>* TemplateFound = TemplateRegistryInstance.FindTemplateFromTemplateID(TemplateID);
+	return TemplateFound ? &TemplateFound->Get() : nullptr;
 }
diff --git a/Engine/Plugins/Runtime/MassGameplay/Source/MassSpawner/Public/MassEntityTemplate.h b/Engine/Plugins/Runtime/MassGameplay/Source/MassSpawner/Public/MassEntityTemplate.h
--- a/Engine/Plugins/Runtime/MassGameplay/Source/MassSpawner/Public/MassEntityTemplate.h
+++ b/Engine/Plugins/Runtime/MassGameplay/Source/MassSpawner/Public/MassEntityTemplate.h
@@ -7,10 +7,14 @@
 #include "MassEntityManager.h"
 #include "MassCommonTypes.h"
 #include "MassTranslator.h"
+#include "Templates/SharedPointer.h"
 #include "MassEntityTemplate.generated.h"
 
 
+class UMassEntityTraitBase;
 struct FMassEntityView;
+struct FMassEntityTemplateIDFactory;
+
 
 //ID of the template an entity is using
 USTRUCT()
@@ -21,22 +25,21 @@ struct MASSSPAWNER_API FMassEntityTemplateID
 	FMassEntityTemplateID()
 		: bIsSet(false)
 	{}
+
+private:
+	friend FMassEntityTemplateIDFactory;
+	// use FMassEntityTemplateIDFactory to access this constructor flavor
 	explicit FMassEntityTemplateID(uint32 InHash)
 		: Hash(InHash), bIsSet(true)
 	{}
 
+public:
 	uint32 GetHash() const 
 	{
 		checkSlow(bIsSet);
 		return Hash; 
 	}
 	
-	void SetHash(uint32 InHash) 
-	{ 
-		Hash = InHash; 
-		bIsSet = true;
-	}
-
 	void Invalidate(uint32 InHash)
 	{
 		// the exact value we set here doesn't really matter, but just to keep the possible states consistent we set it 
@@ -67,35 +70,29 @@ protected:
 	uint8 bIsSet : 1;
 };
 
+
 /** @todo document	*/
 USTRUCT()
-struct MASSSPAWNER_API FMassEntityTemplate
+struct MASSSPAWNER_API FMassEntityTemplateData
 {
 	GENERATED_BODY()
 
 	typedef TFunction<void(UObject& /*Owner*/, FMassEntityView& /*EntityView*/, const EMassTranslationDirection /*CurrentDirection*/)> FObjectFragmentInitializerFunction;
 
-	FMassEntityTemplate() = default;
-	/** InArchetype is expected to be valid. The function will crash-check it. */
-	void SetArchetype(const FMassArchetypeHandle& InArchetype);
-	const FMassArchetypeHandle& GetArchetype() const { return Archetype; }
-
-	TConstArrayView<FObjectFragmentInitializerFunction> GetObjectFragmentInitializers() const { return ObjectInitializers; }
-	TArray<FObjectFragmentInitializerFunction>& GetMutableObjectFragmentInitializers() { return ObjectInitializers; }
+	FMassEntityTemplateData() = default;
 
-	bool IsValid() const { return Archetype.IsValid(); }
 	bool IsEmpty() const { return Composition.IsEmpty(); }
 
-	void SetTemplateID(FMassEntityTemplateID InTemplateID) { TemplateID = InTemplateID; }
-	FMassEntityTemplateID GetTemplateID() const { return TemplateID; }
-
-	void SetTemplateName(const FString& Name) { TemplateName = Name; }
+	TConstArrayView<FObjectFragmentInitializerFunction> GetObjectFragmentInitializers() const { return ObjectInitializers; }
 	const FString& GetTemplateName() const { return TemplateName; }
-	
 	const FMassArchetypeCompositionDescriptor& GetCompositionDescriptor() const { return Composition; }
 	const FMassArchetypeSharedFragmentValues& GetSharedFragmentValues() const { return SharedFragmentValues; }
 	TConstArrayView<FInstancedStruct> GetInitialFragmentValues() const { return InitialFragmentValues; }
 
+	TArray<FMassEntityTemplateData::FObjectFragmentInitializerFunction>& GetMutableObjectFragmentInitializers() { return ObjectInitializers; }
+
+	void SetTemplateName(const FString& Name) { TemplateName = Name; }
+	
 	template<typename T>
 	void AddFragment()
 	{
@@ -109,6 +106,7 @@ struct MASSSPAWNER_API FMassEntityTemplate
 		Composition.Fragments.Add(FragmentType);
 	}
 
+	// @todo this function is doing nothing if a given fragment's initial value has already been created. This seems inconsistent with the other AddFragment functions (especially AddFragment_GetRef).
 	void AddFragment(FConstStructView Fragment)
 	{
 		const UScriptStruct* FragmentType = Fragment.GetScriptStruct();
@@ -206,9 +204,6 @@ struct MASSSPAWNER_API FMassEntityTemplate
 		}
 	}
 
-	FString DebugGetDescription(FMassEntityManager* EntityManager = nullptr) const;
-	FString DebugGetArchetypeDescription(FMassEntityManager& EntityManager) const;
-
 	template<typename T>
 	bool HasFragment() const
 	{
@@ -243,24 +238,59 @@ struct MASSSPAWNER_API FMassEntityTemplate
 		return Composition.SharedFragments.Contains(ScriptStruct);
 	}
 
+	friend uint32 GetTypeHash(const FMassEntityTemplateData& Template);
+
+protected:
 	void Sort()
 	{
 		SharedFragmentValues.Sort();
 	}
 
-private:
-	FMassArchetypeHandle Archetype;
-
+protected:
 	FMassArchetypeCompositionDescriptor Composition;
 	FMassArchetypeSharedFragmentValues SharedFragmentValues;
-	
+
 	// Initial fragment values, this is not part of the archetype as it is the spawner job to set them.
 	TArray<FInstancedStruct> InitialFragmentValues;
 
 	// These functions will be called to initialize entity's UObject-based fragments
 	TArray<FObjectFragmentInitializerFunction> ObjectInitializers;
 
+	FString TemplateName;
+};
+
+
+struct MASSSPAWNER_API FMassEntityTemplate final : public FMassEntityTemplateData, public TSharedFromThis<FMassEntityTemplate> 
+{
+	using Super = FMassEntityTemplateData;
+	friend TSharedFromThis<FMassEntityTemplate>;
+
+	FMassEntityTemplate() = default;
+	FMassEntityTemplate(const FMassEntityTemplateData& InData, FMassEntityManager& EntityManager, FMassEntityTemplateID InTemplateID);
+	FMassEntityTemplate(FMassEntityTemplateData&& InData, FMassEntityManager& EntityManager, FMassEntityTemplateID InTemplateID);
+
+	/** InArchetype is expected to be valid. The function will crash-check it. */
+	void SetArchetype(const FMassArchetypeHandle& InArchetype);
+	const FMassArchetypeHandle& GetArchetype() const { return Archetype; }
+
+	bool IsValid() const { return Archetype.IsValid(); }
+
+	void SetTemplateID(FMassEntityTemplateID InTemplateID) { TemplateID = InTemplateID; }
+	FMassEntityTemplateID GetTemplateID() const { return TemplateID; }
+
+	FString DebugGetDescription(FMassEntityManager* EntityManager = nullptr) const;
+	FString DebugGetArchetypeDescription(FMassEntityManager& EntityManager) const;
+
+	static TSharedRef<FMassEntityTemplate> MakeFinalTemplate(FMassEntityManager& EntityManager, FMassEntityTemplateData&& TempTemplateData, FMassEntityTemplateID InTemplateID);
+
+private:
+	FMassArchetypeHandle Archetype;
 	FMassEntityTemplateID TemplateID;
+};
 
-	FString TemplateName;
+
+struct FMassEntityTemplateIDFactory
+{
+	static FMassEntityTemplateID Make(const FMassEntityTemplateData& TemplateData);
+	static FMassEntityTemplateID Make(TConstArrayView<UMassEntityTraitBase*> Traits);
 };
diff --git a/Engine/Plugins/Runtime/MassGameplay/Source/MassSpawner/Public/MassEntityTemplateRegistry.h b/Engine/Plugins/Runtime/MassGameplay/Source/MassSpawner/Public/MassEntityTemplateRegistry.h
--- a/Engine/Plugins/Runtime/MassGameplay/Source/MassSpawner/Public/MassEntityTemplateRegistry.h
+++ b/Engine/Plugins/Runtime/MassGameplay/Source/MassSpawner/Public/MassEntityTemplateRegistry.h
@@ -21,10 +21,16 @@ enum class EFragmentInitialization : uint8
 
 struct FMassEntityTemplateBuildContext
 {
-	explicit FMassEntityTemplateBuildContext(FMassEntityTemplate& InTemplate)
-		: Template(InTemplate)
+	explicit FMassEntityTemplateBuildContext(FMassEntityTemplateData& InTemplate, FMassEntityTemplateID InTemplateID = FMassEntityTemplateID())
+		: TemplateData(InTemplate)
+		, TemplateID(InTemplateID)
 	{}
 
+	void SetTemplateName(const FString& Name)
+	{
+		TemplateData.SetTemplateName(Name);
+	}
+
 	//----------------------------------------------------------------------//
 	// Fragments 
 	//----------------------------------------------------------------------//
@@ -32,35 +38,35 @@ struct FMassEntityTemplateBuildContext
 	T& AddFragment_GetRef()
 	{
 		TypeAdded(*T::StaticStruct());
-		return Template.AddFragment_GetRef<T>();
+		return TemplateData.AddFragment_GetRef<T>();
 	}
 
 	template<typename T>
 	void AddFragment()
 	{
 		TypeAdded(*T::StaticStruct());
-		Template.AddFragment<T>();
+		TemplateData.AddFragment<T>();
 	}
 
 	void AddFragment(FConstStructView InFragment)
 	{ 
 		checkf(InFragment.GetScriptStruct(), TEXT("Expecting a valid fragment type"));
 		TypeAdded(*InFragment.GetScriptStruct());
-		Template.AddFragment(InFragment);
+		TemplateData.AddFragment(InFragment);
 	}
 
 	template<typename T>
 	void AddTag()
 	{
 		// Tags can be added by multiple traits, so they do not follow the same rules as fragments
-		Template.AddTag<T>();
+		TemplateData.AddTag<T>();
 		TypeAdded(*T::StaticStruct());
 	}
 
 	void AddTag(const UScriptStruct& TagType)
 	{
 		// Tags can be added by multiple traits, so they do not follow the same rules as fragments
-		Template.AddTag(TagType);
+		TemplateData.AddTag(TagType);
 		TypeAdded(TagType);
 	}
 
@@ -68,60 +74,60 @@ struct FMassEntityTemplateBuildContext
 	void AddChunkFragment()
 	{
 		TypeAdded(*T::StaticStruct());
-		Template.AddChunkFragment<T>();
+		TemplateData.AddChunkFragment<T>();
 	}
 
 	void AddConstSharedFragment(const FConstSharedStruct& InSharedFragment)
 	{
 		checkf(InSharedFragment.GetScriptStruct(), TEXT("Expecting a valide shared fragment type"));
 		TypeAdded(*InSharedFragment.GetScriptStruct());
-		Template.AddConstSharedFragment(InSharedFragment);
+		TemplateData.AddConstSharedFragment(InSharedFragment);
 	}
 
 	void AddSharedFragment(const FSharedStruct& InSharedFragment)
 	{
 		checkf(InSharedFragment.GetScriptStruct(), TEXT("Expecting a valide shared fragment type"));
 		TypeAdded(*InSharedFragment.GetScriptStruct());
-		Template.AddSharedFragment(InSharedFragment);
+		TemplateData.AddSharedFragment(InSharedFragment);
 	}
 
 	template<typename T>
 	bool HasFragment() const
 	{
 		ensureMsgf(!BuildingTrait, TEXT("This method is not expected to be called within the build from trait call."));
-		return Template.HasFragment<T>();
+		return TemplateData.HasFragment<T>();
 	}
 	
 	bool HasFragment(const UScriptStruct& ScriptStruct) const
 	{
 		ensureMsgf(!BuildingTrait, TEXT("This method is not expected to be called within the build from trait call."));
-		return Template.HasFragment(ScriptStruct);
+		return TemplateData.HasFragment(ScriptStruct);
 	}
 
 	template<typename T>
 	bool HasTag() const
 	{
-		return Template.HasTag<T>();
+		return TemplateData.HasTag<T>();
 	}
 
 	template<typename T>
 	bool HasChunkFragment() const
 	{
 		ensureMsgf(!BuildingTrait, TEXT("This method is not expected to be called within the build from trait call."));
-		return Template.HasChunkFragment<T>();
+		return TemplateData.HasChunkFragment<T>();
 	}
 
 	template<typename T>
 	bool HasSharedFragment() const
 	{
 		ensureMsgf(!BuildingTrait, TEXT("This method is not expected to be called within the build from trait call."));
-		return Template.HasSharedFragment<T>();
+		return TemplateData.HasSharedFragment<T>();
 	}
 
 	bool HasSharedFragment(const UScriptStruct& ScriptStruct) const
 	{
 		ensureMsgf(!BuildingTrait, TEXT("This method is not expected to be called within the build from trait call."));
-		return Template.HasSharedFragment(ScriptStruct);
+		return TemplateData.HasSharedFragment(ScriptStruct);
 	}
 
 	//----------------------------------------------------------------------//
@@ -131,7 +137,7 @@ struct FMassEntityTemplateBuildContext
 	void AddTranslator()
 	{
 		TypeAdded(*T::StaticClass());
-		GetDefault<T>()->AppendRequiredTags(Template.GetMutableTags());
+		GetDefault<T>()->AppendRequiredTags(TemplateData.GetMutableTags());
 	}
 
 	//----------------------------------------------------------------------//
@@ -157,8 +163,8 @@ struct FMassEntityTemplateBuildContext
 	//----------------------------------------------------------------------//
 	// Template access
 	//----------------------------------------------------------------------//
-	FMassEntityTemplateID GetTemplateID() const { return Template.GetTemplateID(); }
-	TArray<FMassEntityTemplate::FObjectFragmentInitializerFunction>& GetMutableObjectFragmentInitializers() { return Template.GetMutableObjectFragmentInitializers(); }
+	FMassEntityTemplateID GetTemplateID() const { return TemplateID; }
+	TArray<FMassEntityTemplateData::FObjectFragmentInitializerFunction>& GetMutableObjectFragmentInitializers() { return TemplateData.GetMutableObjectFragmentInitializers(); }
 
 	//----------------------------------------------------------------------//
 	// Build methods
@@ -193,7 +199,8 @@ protected:
 	TMultiMap<const UStruct*, const UMassEntityTraitBase*> TraitAddedTypes;
 	TArray< TTuple<const UStruct*, const UMassEntityTraitBase*> > TraitsDependencies;
 
-	FMassEntityTemplate& Template;
+	FMassEntityTemplateData& TemplateData;
+	FMassEntityTemplateID TemplateID;
 };
 
 /** @todo document 
@@ -204,7 +211,9 @@ struct MASSSPAWNER_API FMassEntityTemplateRegistry
 	DECLARE_DELEGATE_ThreeParams(FStructToTemplateBuilderDelegate, const UWorld* /*World*/, const FConstStructView /*InStructInstance*/, FMassEntityTemplateBuildContext& /*BuildContext*/);
 
 	explicit FMassEntityTemplateRegistry(UObject* InOwner = nullptr);
-	
+
+	void ShutDown();
+
 	UWorld* GetWorld() const;
 
 	static FStructToTemplateBuilderDelegate& FindOrAdd(const UScriptStruct& DataType);
@@ -212,21 +221,42 @@ struct MASSSPAWNER_API FMassEntityTemplateRegistry
 	/** Removes all the cached template instances */
 	void DebugReset();
 
-	const FMassEntityTemplate* FindTemplateFromTemplateID(FMassEntityTemplateID TemplateID) const;
+	const TSharedRef<FMassEntityTemplate>* FindTemplateFromTemplateID(FMassEntityTemplateID TemplateID) const;
+
+	/**
+	 * Adds a template based on TemplateData
+	 */
+	const TSharedRef<FMassEntityTemplate>& AddTemplate(FMassEntityTemplateID TemplateID, FMassEntityTemplateData&& TemplateData);
+
+	const TSharedRef<FMassEntityTemplate>& AddTemplate(FMassEntityTemplateData&& TemplateData);
+
+	UE_DEPRECATED(5.3, "We no longer support fething mutable templates from the TemplateRegistry. Stored templates are considered const.")
 	FMassEntityTemplate* FindMutableTemplateFromTemplateID(FMassEntityTemplateID TemplateID);
 
-	FMassEntityTemplate& CreateTemplate(FMassEntityTemplateID TemplateID);
+	UE_DEPRECATED(5.3, "CreateTemplate is no longer available. Use AddTemplate instead.")
+	FMassEntityTemplate& CreateTemplate(const uint32 HashLookup, FMassEntityTemplateID TemplateID);
+
 	void DestroyTemplate(FMassEntityTemplateID TemplateID);
+
+	UE_DEPRECATED(5.3, "InitializeEntityTemplate is no longer available. Use AddTemplate instead.")
 	void InitializeEntityTemplate(FMassEntityTemplate& InOutTemplate) const;
 
-protected:
-	/** @return true if a template has been built, false otherwise */
-	bool BuildTemplateImpl(const FStructToTemplateBuilderDelegate& Builder, const FConstStructView StructInstance, FMassEntityTemplate& OutTemplate);
+	/** Stores the EntityManager the templates will be associated with. 
+	 *  Note that the function will only let users set the EntityManager once. Once it's set the subsequent calls will 
+	 *  have no effect. If attempting to set a different EntityManaget an ensure will trigger. */
+	void StoreEntityManager(const TSharedPtr<FMassEntityManager>& InEntityManager);
+	FMassEntityManager& GetEntityManagerChecked() { check(EntityManager); return *EntityManager; }
 
 protected:
 	static TMap<const UScriptStruct*, FStructToTemplateBuilderDelegate> StructBasedBuilders;
 
-	TMap<FMassEntityTemplateID, FMassEntityTemplate> TemplateIDToTemplateMap;
+	TMap<FMassEntityTemplateID, TSharedRef<FMassEntityTemplate>> TemplateIDToTemplateMap;
+
+	/** 
+	 * EntityManager the hosted templates are associated with. Storing instead of fetching at runtime to ensure all 
+	 *	templates are tied to the same EntityManager
+	 */
+	TSharedPtr<FMassEntityManager> EntityManager;
 
 	TWeakObjectPtr<UObject> Owner;
 };
