MassEntityTemplate changes following up the split of data and finalized part into separate types.
Most notably made FMassTemplateData a member rather than a parent for FMassEntityTemplate.

diff --git a/Engine/Plugins/Runtime/MassGameplay/Source/MassSpawner/Private/MassEntityConfigAsset.cpp b/Engine/Plugins/Runtime/MassGameplay/Source/MassSpawner/Private/MassEntityConfigAsset.cpp
--- a/Engine/Plugins/Runtime/MassGameplay/Source/MassSpawner/Private/MassEntityConfigAsset.cpp
+++ b/Engine/Plugins/Runtime/MassGameplay/Source/MassSpawner/Private/MassEntityConfigAsset.cpp
@@ -68,13 +68,13 @@ const FMassEntityTemplate& FMassEntityConfig::GetOrCreateEntityTemplate(const UW
 	// Dependency should add a fragment with default values (which later can be overridden),
 	// while setup would override values and should be run just once.
 
-	FMassEntityTemplate Template;
-	FMassEntityTemplateBuildContext BuildContext(Template, TemplateID);
+	FMassEntityTemplateData TemplateData;
+	FMassEntityTemplateBuildContext BuildContext(TemplateData, TemplateID);
 
 	BuildContext.BuildFromTraits(CombinedTraits, World);
 	BuildContext.SetTemplateName(GetNameSafe(ConfigOwner));
 
-	return TemplateRegistry.AddTemplate(TemplateID, MoveTemp(Template)).Get();
+	return TemplateRegistry.FindOrAddTemplate(TemplateID, MoveTemp(TemplateData)).Get();
 }
 
 void FMassEntityConfig::DestroyEntityTemplate(const UWorld& World) const
@@ -184,7 +184,7 @@ bool FMassEntityConfig::ValidateEntityTemplate(const UWorld& World)
 	Visited.Add(ConfigOwner);
 	GetCombinedTraits(CombinedTraits, Visited);
 
-	FMassEntityTemplate Template;
+	FMassEntityTemplateData Template;
 	FMassEntityTemplateBuildContext BuildContext(Template);
 
 	return BuildContext.BuildFromTraits(CombinedTraits, World);
diff --git a/Engine/Plugins/Runtime/MassGameplay/Source/MassSpawner/Private/MassEntityTemplateRegistry.cpp b/Engine/Plugins/Runtime/MassGameplay/Source/MassSpawner/Private/MassEntityTemplateRegistry.cpp
--- a/Engine/Plugins/Runtime/MassGameplay/Source/MassSpawner/Private/MassEntityTemplateRegistry.cpp
+++ b/Engine/Plugins/Runtime/MassGameplay/Source/MassSpawner/Private/MassEntityTemplateRegistry.cpp
@@ -15,42 +15,6 @@
 
 #define LOCTEXT_NAMESPACE "Mass"
 
-namespace FTemplateRegistryHelpers
-{
-	uint32 CalcHash(const FConstStructView StructInstance)
-	{
-		const UScriptStruct* Type = StructInstance.GetScriptStruct();
-		const uint8* Memory = StructInstance.GetMemory();
-		return Type && Memory ? Type->GetStructTypeHash(Memory) : 0;
-	}
-
-	void FragmentInstancesToTypes(TArrayView<const FInstancedStruct> FragmentList, TArray<const UScriptStruct*>& OutFragmentTypes)
-	{
-		for (const FInstancedStruct& Instance : FragmentList)
-		{
-			// at this point FragmentList is assumed to have no duplicates nor nulls
-			OutFragmentTypes.Add(Instance.GetScriptStruct());
-		}
-	}
-
-	void ResetEntityTemplates(const TArray<FString>& Args, UWorld* InWorld)
-	{
-		UMassSpawnerSubsystem* SpawnerSystem = UWorld::GetSubsystem<UMassSpawnerSubsystem>(InWorld);
-		if (ensure(SpawnerSystem))
-		{
-			FMassEntityTemplateRegistry& Registry = SpawnerSystem->GetMutableTemplateRegistryInstance();
-			Registry.DebugReset();
-		}
-	}
-
-	FAutoConsoleCommandWithWorldAndArgs EnableCategoryNameCmd(
-		TEXT("ai.mass.reset_entity_templates"),
-		TEXT("Clears all the runtime information cached by MassEntityTemplateRegistry. Will result in lazily building all entity templates again."),
-		FConsoleCommandWithWorldAndArgsDelegate::CreateStatic(&ResetEntityTemplates)
-	);
-} // namespace FTemplateRegistryHelpers
-
-
 //----------------------------------------------------------------------//
 // FMassEntityTemplateRegistry 
 //----------------------------------------------------------------------//
@@ -77,7 +41,7 @@ FMassEntityTemplateRegistry::FStructToTemplateBuilderDelegate& FMassEntityTempla
 	return StructBasedBuilders.FindOrAdd(&DataType);
 }
 
-void FMassEntityTemplateRegistry::StoreEntityManager(const TSharedPtr<FMassEntityManager>& InEntityManager)
+void FMassEntityTemplateRegistry::Initialize(const TSharedPtr<FMassEntityManager>& InEntityManager)
 {
 	if (EntityManager)
 	{
@@ -100,7 +64,7 @@ const TSharedRef<FMassEntityTemplate>* FMassEntityTemplateRegistry::FindTemplate
 	return TemplateIDToTemplateMap.Find(TemplateID);
 }
 
-const TSharedRef<FMassEntityTemplate>& FMassEntityTemplateRegistry::AddTemplate(FMassEntityTemplateID TemplateID, FMassEntityTemplateData&& TemplateData)
+const TSharedRef<FMassEntityTemplate>& FMassEntityTemplateRegistry::FindOrAddTemplate(FMassEntityTemplateID TemplateID, FMassEntityTemplateData&& TemplateData)
 {
 	check(EntityManager);
 	const TSharedRef<FMassEntityTemplate>* ExistingTemplate = FindTemplateFromTemplateID(TemplateID);
@@ -112,7 +76,7 @@ const TSharedRef<FMassEntityTemplate>& FMassEntityTemplateRegistry::AddTemplate(
 	return TemplateIDToTemplateMap.Add(TemplateID, FMassEntityTemplate::MakeFinalTemplate(*EntityManager, MoveTemp(TemplateData), TemplateID));
 }
 
-const TSharedRef<FMassEntityTemplate>& FMassEntityTemplateRegistry::AddTemplate(FMassEntityTemplateData&& TemplateData)
+const TSharedRef<FMassEntityTemplate>& FMassEntityTemplateRegistry::FindOrAddTemplate(FMassEntityTemplateData&& TemplateData)
 {
 	const FMassEntityTemplateID NewTemplateID = FMassEntityTemplateIDFactory::Make(TemplateData);
 
diff --git a/Engine/Plugins/Runtime/MassGameplay/Source/MassSpawner/Private/MassSpawnerSubsystem.cpp b/Engine/Plugins/Runtime/MassGameplay/Source/MassSpawner/Private/MassSpawnerSubsystem.cpp
--- a/Engine/Plugins/Runtime/MassGameplay/Source/MassSpawner/Private/MassSpawnerSubsystem.cpp
+++ b/Engine/Plugins/Runtime/MassGameplay/Source/MassSpawner/Private/MassSpawnerSubsystem.cpp
@@ -36,7 +36,7 @@ void UMassSpawnerSubsystem::Initialize(FSubsystemCollectionBase& Collection)
 	UWorld* World = GetWorld();
 	check(World);
 	EntityManager = UE::Mass::Utils::GetEntityManagerChecked(*World).AsShared();
-	TemplateRegistryInstance.StoreEntityManager(EntityManager);
+	TemplateRegistryInstance.Initialize(EntityManager);
 }
 
 void UMassSpawnerSubsystem::Deinitialize() 
diff --git a/Engine/Plugins/Runtime/MassGameplay/Source/MassSpawner/Public/MassEntityTemplate.h b/Engine/Plugins/Runtime/MassGameplay/Source/MassSpawner/Public/MassEntityTemplate.h
--- a/Engine/Plugins/Runtime/MassGameplay/Source/MassSpawner/Public/MassEntityTemplate.h
+++ b/Engine/Plugins/Runtime/MassGameplay/Source/MassSpawner/Public/MassEntityTemplate.h
@@ -71,7 +71,11 @@ protected:
 };
 
 
-/** @todo document	*/
+/** 
+ * Serves as data used to define and build finalized FMassEntityTemplate instances. Describes composition and initial
+ * values of fragments for entities created with this data, and lets users modify and extend the data. Once finalized as 
+ * FMassEntityTemplate the data will become immutable. 
+ */
 USTRUCT()
 struct MASSSPAWNER_API FMassEntityTemplateData
 {
@@ -240,7 +244,6 @@ struct MASSSPAWNER_API FMassEntityTemplateData
 
 	friend uint32 GetTypeHash(const FMassEntityTemplateData& Template);
 
-protected:
 	void Sort()
 	{
 		SharedFragmentValues.Sort();
@@ -259,10 +262,13 @@ protected:
 	FString TemplateName;
 };
 
-
-struct MASSSPAWNER_API FMassEntityTemplate final : public FMassEntityTemplateData, public TSharedFromThis<FMassEntityTemplate> 
+/**
+ * A finalized and const wrapper for FMassEntityTemplateData, associated with a Mass archetype and template ID. 
+ * Designed to never be changed. If a change is needed a copy of the hosted FMassEntityTemplateData needs to be made and 
+ * used to create another finalized FMassEntityTemplate (via FMassEntityTemplateManager).
+ */
+struct MASSSPAWNER_API FMassEntityTemplate final : public TSharedFromThis<FMassEntityTemplate> 
 {
-	using Super = FMassEntityTemplateData;
 	friend TSharedFromThis<FMassEntityTemplate>;
 
 	FMassEntityTemplate() = default;
@@ -283,7 +289,17 @@ struct MASSSPAWNER_API FMassEntityTemplate final : public FMassEntityTemplateDat
 
 	static TSharedRef<FMassEntityTemplate> MakeFinalTemplate(FMassEntityManager& EntityManager, FMassEntityTemplateData&& TempTemplateData, FMassEntityTemplateID InTemplateID);
 
+	//-----------------------------------------------------------------------------
+	// FMassEntityTemplateData getters
+	//-----------------------------------------------------------------------------
+	FORCEINLINE TConstArrayView<FMassEntityTemplateData::FObjectFragmentInitializerFunction> GetObjectFragmentInitializers() const { return TemplateData.GetObjectFragmentInitializers(); }
+	FORCEINLINE const FString& GetTemplateName() const { return TemplateData.GetTemplateName(); }
+	FORCEINLINE const FMassArchetypeCompositionDescriptor& GetCompositionDescriptor() const { return TemplateData.GetCompositionDescriptor(); }
+	FORCEINLINE const FMassArchetypeSharedFragmentValues& GetSharedFragmentValues() const { return TemplateData.GetSharedFragmentValues(); }
+	FORCEINLINE TConstArrayView<FInstancedStruct> GetInitialFragmentValues() const { return TemplateData.GetInitialFragmentValues(); }
+
 private:
+	FMassEntityTemplateData TemplateData;
 	FMassArchetypeHandle Archetype;
 	FMassEntityTemplateID TemplateID;
 };
diff --git a/Engine/Plugins/Runtime/MassGameplay/Source/MassSpawner/Public/MassEntityTemplateRegistry.h b/Engine/Plugins/Runtime/MassGameplay/Source/MassSpawner/Public/MassEntityTemplateRegistry.h
--- a/Engine/Plugins/Runtime/MassGameplay/Source/MassSpawner/Public/MassEntityTemplateRegistry.h
+++ b/Engine/Plugins/Runtime/MassGameplay/Source/MassSpawner/Public/MassEntityTemplateRegistry.h
@@ -212,6 +212,11 @@ struct MASSSPAWNER_API FMassEntityTemplateRegistry
 
 	explicit FMassEntityTemplateRegistry(UObject* InOwner = nullptr);
 
+	/** Initializes and stores the EntityManager the templates will be associated with. Needs to be called before any template operations.
+	 *  Note that the function will only let users set the EntityManager once. Once it's set the subsequent calls will
+	 *  have no effect. If attempting to set a different EntityManaget an ensure will trigger. */
+	void Initialize(const TSharedPtr<FMassEntityManager>& InEntityManager);
+
 	void ShutDown();
 
 	UWorld* GetWorld() const;
@@ -226,9 +231,9 @@ struct MASSSPAWNER_API FMassEntityTemplateRegistry
 	/**
 	 * Adds a template based on TemplateData
 	 */
-	const TSharedRef<FMassEntityTemplate>& AddTemplate(FMassEntityTemplateID TemplateID, FMassEntityTemplateData&& TemplateData);
+	const TSharedRef<FMassEntityTemplate>& FindOrAddTemplate(FMassEntityTemplateID TemplateID, FMassEntityTemplateData&& TemplateData);
 
-	const TSharedRef<FMassEntityTemplate>& AddTemplate(FMassEntityTemplateData&& TemplateData);
+	const TSharedRef<FMassEntityTemplate>& FindOrAddTemplate(FMassEntityTemplateData&& TemplateData);
 
 	UE_DEPRECATED(5.3, "We no longer support fething mutable templates from the TemplateRegistry. Stored templates are considered const.")
 	FMassEntityTemplate* FindMutableTemplateFromTemplateID(FMassEntityTemplateID TemplateID);
@@ -241,10 +246,6 @@ struct MASSSPAWNER_API FMassEntityTemplateRegistry
 	UE_DEPRECATED(5.3, "InitializeEntityTemplate is no longer available. Use AddTemplate instead.")
 	void InitializeEntityTemplate(FMassEntityTemplate& InOutTemplate) const;
 
-	/** Stores the EntityManager the templates will be associated with. 
-	 *  Note that the function will only let users set the EntityManager once. Once it's set the subsequent calls will 
-	 *  have no effect. If attempting to set a different EntityManaget an ensure will trigger. */
-	void StoreEntityManager(const TSharedPtr<FMassEntityManager>& InEntityManager);
 	FMassEntityManager& GetEntityManagerChecked() { check(EntityManager); return *EntityManager; }
 
 protected:
