CIS fix

diff --git a/Engine/Plugins/Runtime/MassGameplay/Source/MassSpawner/Private/MassEntityTemplate.cpp b/Engine/Plugins/Runtime/MassGameplay/Source/MassSpawner/Private/MassEntityTemplate.cpp
--- a/Engine/Plugins/Runtime/MassGameplay/Source/MassSpawner/Private/MassEntityTemplate.cpp
+++ b/Engine/Plugins/Runtime/MassGameplay/Source/MassSpawner/Private/MassEntityTemplate.cpp
@@ -51,22 +51,22 @@ TSharedRef<FMassEntityTemplate> FMassEntityTemplate::MakeFinalTemplate(FMassEnti
 }
 
 FMassEntityTemplate::FMassEntityTemplate(const FMassEntityTemplateData& InData, FMassEntityManager& EntityManager, FMassEntityTemplateID InTemplateID)
-	: Super(InData)
+	: TemplateData(InData)
 	, TemplateID(InTemplateID)
 {
 	// Sort anything there is to sort for later comparison purposes
-	Sort();
+	TemplateData.Sort();
 
 	const FMassArchetypeHandle ArchetypeHandle = EntityManager.CreateArchetype(GetCompositionDescriptor(), FName(GetTemplateName()));
 	SetArchetype(ArchetypeHandle);
 }
 
 FMassEntityTemplate::FMassEntityTemplate(FMassEntityTemplateData&& InData, FMassEntityManager& EntityManager, FMassEntityTemplateID InTemplateID)
-	: Super(InData)
+	: TemplateData(InData)
 	, TemplateID(InTemplateID)
 {
 	// Sort anything there is to sort for later comparison purposes
-	Sort();
+	TemplateData.Sort();
 
 	const FMassArchetypeHandle ArchetypeHandle = EntityManager.CreateArchetype(GetCompositionDescriptor(), FName(GetTemplateName()));
 	SetArchetype(ArchetypeHandle);
@@ -101,7 +101,7 @@ FString FMassEntityTemplate::DebugGetDescription(FMassEntityManager* EntityManag
 	else
 	{
 		Ar += TEXT("Composition:\n");
-		Composition.DebugOutputDescription(Ar);
+		GetCompositionDescriptor().DebugOutputDescription(Ar);
 	}
 
 #endif // WITH_MASSGAMEPLAY_DEBUG
