Renamed FISMCSharedDataMap to FMassISMCSharedDataMap to indicate the origin of the type. Similar change to FStaticMeshInstanceVisualizationMeshDesc and FISMCSharedData.

diff --git a/Engine/Plugins/Runtime/MassGameplay/Config/BaseMassGameplay.ini b/Engine/Plugins/Runtime/MassGameplay/Config/BaseMassGameplay.ini
--- a/Engine/Plugins/Runtime/MassGameplay/Config/BaseMassGameplay.ini
+++ b/Engine/Plugins/Runtime/MassGameplay/Config/BaseMassGameplay.ini
@@ -1,2 +1,4 @@
 [CoreRedirects]
-; no changes past UE 5.0
+; 5.3
++StructRedirects=(OldName="StaticMeshInstanceVisualizationMeshDesc",NewName="/Script/MassRepresentation.MassStaticMeshInstanceVisualizationMeshDesc")
++StructRedirects=(OldName="ISMCSharedData",NewName="/Script/MassRepresentation.MassISMCSharedData")
\ No newline at end of file
diff --git a/Engine/Plugins/Runtime/MassGameplay/Source/MassRepresentation/Private/MassVisualizationComponent.cpp b/Engine/Plugins/Runtime/MassGameplay/Source/MassRepresentation/Private/MassVisualizationComponent.cpp
--- a/Engine/Plugins/Runtime/MassGameplay/Source/MassRepresentation/Private/MassVisualizationComponent.cpp
+++ b/Engine/Plugins/Runtime/MassGameplay/Source/MassRepresentation/Private/MassVisualizationComponent.cpp
@@ -63,9 +63,9 @@ void UMassVisualizationComponent::ConstructStaticMeshComponents()
 			UE_LOG(LogMassRepresentation, Error, TEXT("No associated meshes for this intanced static mesh type"));
 			continue;
 		}
-		for (const FStaticMeshInstanceVisualizationMeshDesc& MeshDesc : Info.Desc.Meshes)
+		for (const FMassStaticMeshInstanceVisualizationMeshDesc& MeshDesc : Info.Desc.Meshes)
 		{			
-			FISMCSharedData* SharedData = ISMCSharedData.Find(MeshDesc);
+			FMassISMCSharedData* SharedData = ISMCSharedData.Find(MeshDesc);
 			UInstancedStaticMeshComponent* ISMC = SharedData ? SharedData->ISMC : nullptr;
 			if (SharedData)
 			{
@@ -91,7 +91,7 @@ void UMassVisualizationComponent::ConstructStaticMeshComponents()
 				ISMC->SetReceivesDecals(false);
 				ISMC->RegisterComponent();
 
-				ISMCSharedData.Emplace(MeshDesc, FISMCSharedData(ISMC));
+				ISMCSharedData.Emplace(MeshDesc, FMassISMCSharedData(ISMC));
 			}
 
 			Info.InstancedStaticMeshComponents.Add(ISMC);
@@ -116,7 +116,7 @@ void UMassVisualizationComponent::ConstructStaticMeshComponents()
 			}
 			AllLODSignificances.Insert(Significance, i);
 		};
-		for (const FStaticMeshInstanceVisualizationMeshDesc& MeshDesc : Info.Desc.Meshes)
+		for (const FMassStaticMeshInstanceVisualizationMeshDesc& MeshDesc : Info.Desc.Meshes)
 		{
 			UniqueInsertOrdered(MeshDesc.MinLODSignificance);
 			UniqueInsertOrdered(MeshDesc.MaxLODSignificance);
@@ -131,7 +131,7 @@ void UMassVisualizationComponent::ConstructStaticMeshComponents()
 
 			for (int j = 0; j < Info.Desc.Meshes.Num(); ++j)
 			{
-				const FStaticMeshInstanceVisualizationMeshDesc& MeshDesc = Info.Desc.Meshes[j];
+				const FMassStaticMeshInstanceVisualizationMeshDesc& MeshDesc = Info.Desc.Meshes[j];
 				const bool bAddMeshInRange = (Range.MinSignificance >= MeshDesc.MinLODSignificance && Range.MinSignificance < MeshDesc.MaxLODSignificance);
 				if (bAddMeshInRange)
 				{
@@ -191,7 +191,7 @@ void UMassVisualizationComponent::BeginVisualChanges()
 	// Reset instance transform scratch buffers
 	for (auto It = ISMCSharedData.CreateIterator(); It; ++It)
 	{
-		FISMCSharedData& SharedData = It.Value();
+		FMassISMCSharedData& SharedData = It.Value();
 		SharedData.UpdateInstanceIds.Reset();
 		SharedData.StaticMeshInstanceCustomFloats.Reset();
 		SharedData.StaticMeshInstanceTransforms.Reset();
@@ -207,7 +207,7 @@ void UMassVisualizationComponent::EndVisualChanges()
 	// Batch update gathered instance transforms
 	for (auto It = ISMCSharedData.CreateIterator(); It; ++It)
 	{
-		FISMCSharedData& SharedData = It.Value();
+		FMassISMCSharedData& SharedData = It.Value();
 		if (UInstancedStaticMeshComponent* InstancedStaticMeshComponent = SharedData.ISMC)
 		{
 			const int32 NumCustomDataFloats = SharedData.StaticMeshInstanceCustomFloats.Num() / (FMath::Max(1, SharedData.UpdateInstanceIds.Num()));
@@ -271,11 +271,11 @@ void UMassVisualizationComponent::EndVisualChanges()
 // FMassInstancedStaticMeshInfo
 //---------------------------------------------------------------
 
-void FMassInstancedStaticMeshInfo::ClearVisualInstance(FISMCSharedDataMap& ISMCSharedData)
+void FMassInstancedStaticMeshInfo::ClearVisualInstance(FMassISMCSharedDataMap& ISMCSharedData)
 {
 	for (int i = 0; i < Desc.Meshes.Num(); i++)
 	{
-		FISMCSharedData* SharedData = ISMCSharedData.Find(Desc.Meshes[i]);
+		FMassISMCSharedData* SharedData = ISMCSharedData.Find(Desc.Meshes[i]);
 		if (SharedData)
 		{
 			SharedData->RefCount -= 1;
@@ -306,7 +306,7 @@ void FMassLODSignificanceRange::AddBatchedTransform(const int32 InstanceId, cons
 			continue;
 		}
 
-		FISMCSharedData& SharedData = (*ISMCSharedDataPtr)[StaticMeshRefs[i]];
+		FMassISMCSharedData& SharedData = (*ISMCSharedDataPtr)[StaticMeshRefs[i]];
 
 		SharedData.UpdateInstanceIds.Add(InstanceId);
 		SharedData.StaticMeshInstanceTransforms.Add(Transform);
@@ -324,7 +324,7 @@ void FMassLODSignificanceRange::AddBatchedCustomDataFloats(const TArray<float>&
 			continue;
 		}
 
-		FISMCSharedData& SharedData = (*ISMCSharedDataPtr)[StaticMeshRefs[i]];
+		FMassISMCSharedData& SharedData = (*ISMCSharedDataPtr)[StaticMeshRefs[i]];
 		SharedData.StaticMeshInstanceCustomFloats.Append(CustomFloats);
 	}
 }
@@ -339,7 +339,7 @@ void FMassLODSignificanceRange::WriteCustomDataFloatsAtStartIndex(int32 StaticMe
 			return;
 		}
 
-		FISMCSharedData& SharedData = (*ISMCSharedDataPtr)[StaticMeshRefs[StaticMeshIndex]];
+		FMassISMCSharedData& SharedData = (*ISMCSharedDataPtr)[StaticMeshRefs[StaticMeshIndex]];
 
 		int32 StartIndex = FloatsPerInstance * SharedData.WriteIterator + StartFloatIndex;
 
diff --git a/Engine/Plugins/Runtime/MassGameplay/Source/MassRepresentation/Public/MassRepresentationSubsystem.h b/Engine/Plugins/Runtime/MassGameplay/Source/MassRepresentation/Public/MassRepresentationSubsystem.h
--- a/Engine/Plugins/Runtime/MassGameplay/Source/MassRepresentation/Public/MassRepresentationSubsystem.h
+++ b/Engine/Plugins/Runtime/MassGameplay/Source/MassRepresentation/Public/MassRepresentationSubsystem.h
@@ -12,7 +12,6 @@
 
 class UMassVisualizationComponent;
 class AMassVisualizer;
-class UHierarchicalInstancedStaticMeshComponent;
 struct FStaticMeshInstanceVisualizationDesc;
 struct FMassInstancedStaticMeshInfo;
 struct FMassActorSpawnRequestHandle;
diff --git a/Engine/Plugins/Runtime/MassGameplay/Source/MassRepresentation/Public/MassRepresentationTypes.h b/Engine/Plugins/Runtime/MassGameplay/Source/MassRepresentation/Public/MassRepresentationTypes.h
--- a/Engine/Plugins/Runtime/MassGameplay/Source/MassRepresentation/Public/MassRepresentationTypes.h
+++ b/Engine/Plugins/Runtime/MassGameplay/Source/MassRepresentation/Public/MassRepresentationTypes.h
@@ -35,7 +35,7 @@ enum class EMassActorEnabledType : uint8
 };
 
 USTRUCT()
-struct FStaticMeshInstanceVisualizationMeshDesc
+struct FMassStaticMeshInstanceVisualizationMeshDesc
 {
 	GENERATED_BODY()
 	
@@ -67,7 +67,7 @@ struct FStaticMeshInstanceVisualizationMeshDesc
 	UPROPERTY(EditAnywhere, Category = "Mass|Visual")
 	TEnumAsByte<EComponentMobility::Type> Mobility = EComponentMobility::Movable;
 
-	bool operator==(const FStaticMeshInstanceVisualizationMeshDesc& Other) const
+	bool operator==(const FMassStaticMeshInstanceVisualizationMeshDesc& Other) const
 	{
 		return Mesh == Other.Mesh && 
 			MaterialOverrides == Other.MaterialOverrides && 
@@ -77,7 +77,7 @@ struct FStaticMeshInstanceVisualizationMeshDesc
 			Mobility == Other.Mobility;
 	}
 	
-	friend FORCEINLINE uint32 GetTypeHash(const FStaticMeshInstanceVisualizationMeshDesc& MeshDesc)
+	friend FORCEINLINE uint32 GetTypeHash(const FMassStaticMeshInstanceVisualizationMeshDesc& MeshDesc)
 	{
 		uint32 Hash = 0x0;
 		Hash = PointerHash(MeshDesc.Mesh, Hash);
@@ -109,7 +109,7 @@ struct FStaticMeshInstanceVisualizationDesc : public FTableRowBase
 	 * visualize the given agent.
 	 */
 	UPROPERTY(EditAnywhere, Category = "Mass|Visual")
-	TArray<FStaticMeshInstanceVisualizationMeshDesc> Meshes;
+	TArray<FMassStaticMeshInstanceVisualizationMeshDesc> Meshes;
 
 	/** Boolean to enable code to transform the static meshes if not align the mass agent transform */
 	UPROPERTY(EditAnywhere, Category = "Mass|Visual")
@@ -128,13 +128,13 @@ struct FStaticMeshInstanceVisualizationDesc : public FTableRowBase
 class UInstancedStaticMeshComponent;
 
 USTRUCT()
-struct MASSREPRESENTATION_API FISMCSharedData
+struct MASSREPRESENTATION_API FMassISMCSharedData
 {
 	GENERATED_BODY()
 
-	FISMCSharedData() = default;
+	FMassISMCSharedData() = default;
 
-	FISMCSharedData(UInstancedStaticMeshComponent* InISMC)
+	FMassISMCSharedData(UInstancedStaticMeshComponent* InISMC)
 		: ISMC(InISMC),
 		RefCount(1)
 	{
@@ -157,7 +157,16 @@ struct MASSREPRESENTATION_API FISMCSharedData
 	int32 WriteIterator = 0;
 };
 
-typedef TMap<uint32, FISMCSharedData> FISMCSharedDataMap;
+using FMassISMCSharedDataMap = TMap<uint32, FMassISMCSharedData>;
+
+//struct FISMCSharedDataMap : public TMap<uint32, FMassISMCSharedData>
+//{
+//	FISMCSharedDataMap()
+//	{
+//		ensureMsgf(false, "FISMCSharedDataMap has been renamed to FMassISMCSharedDataMap, use that instead"));
+//	}
+//};
+
 
 USTRUCT()
 struct MASSREPRESENTATION_API FMassLODSignificanceRange
@@ -183,7 +192,7 @@ public:
 				continue;
 			}
 
-			FISMCSharedData& SharedData = (*ISMCSharedDataPtr)[StaticMeshRefs[i]];
+			FMassISMCSharedData& SharedData = (*ISMCSharedDataPtr)[StaticMeshRefs[i]];
 			const int32 StartIndex = SharedData.StaticMeshInstanceCustomFloats.AddDefaulted(StructSizeInFloats + NumFloatsToPad);
 			InCustomDataType* CustomData = reinterpret_cast<InCustomDataType*>(&SharedData.StaticMeshInstanceCustomFloats[StartIndex]);
 			*CustomData = InCustomData;
@@ -202,7 +211,7 @@ public:
 	UPROPERTY()
 	TArray<uint32> StaticMeshRefs;
 
-	FISMCSharedDataMap* ISMCSharedDataPtr = nullptr;
+	FMassISMCSharedDataMap* ISMCSharedDataPtr = nullptr;
 };
 
 USTRUCT()
@@ -314,7 +323,7 @@ public:
 protected:
 
 	/** Destroy the visual instance */
-	void ClearVisualInstance(FISMCSharedDataMap& ISMCSharedData);
+	void ClearVisualInstance(FMassISMCSharedDataMap& ISMCSharedData);
 
 	/** Information about this static mesh which will represent all instances */
 	UPROPERTY()
diff --git a/Engine/Plugins/Runtime/MassGameplay/Source/MassRepresentation/Public/MassVisualizationComponent.h b/Engine/Plugins/Runtime/MassGameplay/Source/MassRepresentation/Public/MassVisualizationComponent.h
--- a/Engine/Plugins/Runtime/MassGameplay/Source/MassRepresentation/Public/MassVisualizationComponent.h
+++ b/Engine/Plugins/Runtime/MassGameplay/Source/MassRepresentation/Public/MassVisualizationComponent.h
@@ -59,8 +59,7 @@ protected:
 	TArray<FMassInstancedStaticMeshInfo> InstancedStaticMeshInfos;
 	UE_MT_DECLARE_RW_RECURSIVE_ACCESS_DETECTOR(InstancedStaticMeshInfosDetector);
 
-	FISMCSharedDataMap ISMCSharedData;
-
+	FMassISMCSharedDataMap ISMCSharedData;
 
 	/** Whether there is a need to create a StaticMeshComponent */
 	bool bNeedStaticMeshComponentConstruction = false;
