Added a way to control the seed of MassEntitySpawnDataGenerators, mainly for debugging purposes, but can be used for runtime purposes as well.

diff --git a/Engine/Plugins/Runtime/MassGameplay/Source/MassSpawner/Private/MassEntityEQSSpawnPointsGenerator.cpp b/Engine/Plugins/Runtime/MassGameplay/Source/MassSpawner/Private/MassEntityEQSSpawnPointsGenerator.cpp
--- a/Engine/Plugins/Runtime/MassGameplay/Source/MassSpawner/Private/MassEntityEQSSpawnPointsGenerator.cpp
+++ b/Engine/Plugins/Runtime/MassGameplay/Source/MassSpawner/Private/MassEntityEQSSpawnPointsGenerator.cpp
@@ -55,7 +55,7 @@ void UMassEntityEQSSpawnPointsGenerator::OnEQSQueryFinished(TSharedPtr<FEnvQuery
 	EQSResult->GetAllAsLocations(Locations);
 
 	// Randomize them
-	FRandomStream RandomStream(GFrameNumber);
+	FRandomStream RandomStream(GetRandomSelectionSeed());
 	for (int32 I = 0; I < Locations.Num(); ++I)
 	{
 		const int32 J = RandomStream.RandHelper(Locations.Num());
diff --git a/Engine/Plugins/Runtime/MassGameplay/Source/MassSpawner/Private/MassEntityZoneGraphSpawnPointsGenerator.cpp b/Engine/Plugins/Runtime/MassGameplay/Source/MassSpawner/Private/MassEntityZoneGraphSpawnPointsGenerator.cpp
--- a/Engine/Plugins/Runtime/MassGameplay/Source/MassSpawner/Private/MassEntityZoneGraphSpawnPointsGenerator.cpp
+++ b/Engine/Plugins/Runtime/MassGameplay/Source/MassSpawner/Private/MassEntityZoneGraphSpawnPointsGenerator.cpp
@@ -27,7 +27,7 @@ void UMassEntityZoneGraphSpawnPointsGenerator::Generate(UObject& QueryOwner, TCo
 
 	TArray<FVector> Locations;
 	
-	const FRandomStream RandomStream(GFrameNumber);
+	const FRandomStream RandomStream(GetRandomSelectionSeed());
 	const TConstArrayView<FRegisteredZoneGraphData> RegisteredZoneGraphs = ZoneGraph->GetRegisteredZoneGraphData();
 	if (RegisteredZoneGraphs.IsEmpty())
 	{
diff --git a/Engine/Plugins/Runtime/MassGameplay/Source/MassSpawner/Public/MassEntitySpawnDataGeneratorBase.h b/Engine/Plugins/Runtime/MassGameplay/Source/MassSpawner/Public/MassEntitySpawnDataGeneratorBase.h
--- a/Engine/Plugins/Runtime/MassGameplay/Source/MassSpawner/Public/MassEntitySpawnDataGeneratorBase.h
+++ b/Engine/Plugins/Runtime/MassGameplay/Source/MassSpawner/Public/MassEntitySpawnDataGeneratorBase.h
@@ -66,4 +66,13 @@ public:
 	 * @param OutResults Generator result for each entity type that had > 0 entities assigned to it.
 	 */
 	void BuildResultsFromEntityTypes(const int32 SpawnCount, TConstArrayView<FMassSpawnedEntityType> EntityTypes, TArray<FMassEntitySpawnDataGeneratorResult>& OutResults) const;
+
+protected:
+	int32 GetRandomSelectionSeed() const { return bUseCustomSeed ? RandomSelectionSeed : GFrameNumber; }
+
+	UPROPERTY(Category = "Query", EditAnywhere, meta=(EditCondition="bUseCustomSeed"))
+	int32 RandomSelectionSeed = 0;
+
+	UPROPERTY()
+	uint8 bUseCustomSeed : 1;
 };
