Made Mass rely on Editor viewports rather than player controller in Simulation-in-editor mode.

This is required for functionning MassLOD in SIE, since even though a PlayerController does exist in the editor, it's "location" never changes.
The change also fixes an issue causing editor-viewport-based LOD viewers being recreated all the time.

diff --git a/Engine/Plugins/Runtime/MassGameplay/Source/MassLOD/Private/MassLODSubsystem.cpp b/Engine/Plugins/Runtime/MassGameplay/Source/MassLOD/Private/MassLODSubsystem.cpp
--- a/Engine/Plugins/Runtime/MassGameplay/Source/MassLOD/Private/MassLODSubsystem.cpp
+++ b/Engine/Plugins/Runtime/MassGameplay/Source/MassLOD/Private/MassLODSubsystem.cpp
@@ -8,6 +8,7 @@
 #include "MassSimulationSubsystem.h"
 #if WITH_EDITOR
 #include "CoreGlobals.h" // GIsEditor
+#include "Editor.h" // GEditor
 #include "LevelEditorViewport.h"
 #include "Editor/EditorEngine.h"
 #endif // WITH_EDITOR
@@ -38,7 +39,8 @@ void UMassLODSubsystem::Initialize(FSubsystemCollectionBase& Collection)
 		check(SimSystem);
 		SimSystem->GetOnProcessingPhaseStarted(EMassProcessingPhase::PrePhysics).AddUObject(this, &UMassLODSubsystem::OnPrePhysicsPhaseStarted);
 #if WITH_EDITOR
-		bUseEditorLevelViewports = (GIsEditor && World->WorldType == EWorldType::Editor);
+		bIgnorePlayerControllersDueToSimulation = (GEditor && GEditor->IsSimulateInEditorInProgress());
+		bUseEditorLevelViewports = bIgnorePlayerControllersDueToSimulation || (GIsEditor && World->WorldType == EWorldType::Editor);
 #endif // WITH_EDITOR
 	}
 
@@ -145,7 +147,11 @@ void UMassLODSubsystem::SynchronizeViewers()
 			continue;
 		}
 
-		if (ViewerInfo.PlayerController != nullptr)
+		if (ViewerInfo.PlayerController != nullptr
+#if WITH_EDITOR
+			&& bIgnorePlayerControllersDueToSimulation == false
+#endif // WITH_EDITOR
+			)
 		{
 			LocalViewerMap.Add(GetTypeHash(ViewerInfo.PlayerController->GetFName()), ViewerInfo.Handle);
 		}
@@ -206,7 +212,10 @@ void UMassLODSubsystem::SynchronizeViewers()
 			if (LevelVC && LevelVC->IsPerspective())
 			{
 				const int32 HashValue = GetTypeHash(LevelVC);
-				AddEditorViewer(HashValue, ClientIndex);
+				if (LocalViewerMap.Remove(HashValue) == 0)
+				{
+					AddEditorViewer(HashValue, ClientIndex);
+				}
 			}
 		}
 	}
@@ -298,6 +307,13 @@ void UMassLODSubsystem::AddViewer(APlayerController* PlayerController, FName Str
 
 void UMassLODSubsystem::AddPlayerViewer(APlayerController& PlayerController)
 {
+#if WITH_EDITOR
+	if (bIgnorePlayerControllersDueToSimulation)
+	{
+		return;
+	}
+#endif // WITH_EDITOR
+
 	const int32 HashValue = GetTypeHash(PlayerController.GetFName());
 
 	FMassViewerHandle& ViewerHandle = ViewerMap.FindOrAdd(HashValue, FMassViewerHandle());
diff --git a/Engine/Plugins/Runtime/MassGameplay/Source/MassLOD/Public/MassLODSubsystem.h b/Engine/Plugins/Runtime/MassGameplay/Source/MassLOD/Public/MassLODSubsystem.h
--- a/Engine/Plugins/Runtime/MassGameplay/Source/MassLOD/Public/MassLODSubsystem.h
+++ b/Engine/Plugins/Runtime/MassGameplay/Source/MassLOD/Public/MassLODSubsystem.h
@@ -150,6 +150,7 @@ private:
 
 #if WITH_EDITOR
 	bool bUseEditorLevelViewports = false;
+	bool bIgnorePlayerControllersDueToSimulation = false;
 #endif // WITH_EDITOR
 
 	/** Free list of indices in the sparse viewer array */
