Fixed a failed EQS Mass spawn locations query still resulting in a non-empty Results output.

diff --git a/Engine/Plugins/Runtime/MassGameplay/Source/MassSpawner/Private/MassEntityEQSSpawnPointsGenerator.cpp b/Engine/Plugins/Runtime/MassGameplay/Source/MassSpawner/Private/MassEntityEQSSpawnPointsGenerator.cpp
--- a/Engine/Plugins/Runtime/MassGameplay/Source/MassSpawner/Private/MassEntityEQSSpawnPointsGenerator.cpp
+++ b/Engine/Plugins/Runtime/MassGameplay/Source/MassSpawner/Private/MassEntityEQSSpawnPointsGenerator.cpp
@@ -47,6 +47,7 @@ void UMassEntityEQSSpawnPointsGenerator::OnEQSQueryFinished(TSharedPtr<FEnvQuery
 	{
 		UE_VLOG_UELOG(this, LogMassSpawner, Error, TEXT("EQS query failed or result is invalid"));
 		// Return empty result.
+		Results.Reset();
 		FinishedGeneratingSpawnPointsDelegate.Execute(Results);
 		return;
 	}
