PR #10220: EQS query can be parameterized in the Mass Spawner Instance (Contributed by [REDACTED_PERSON])

diff --git a/Engine/Plugins/Runtime/MassGameplay/Source/MassSpawner/Private/MassEntityEQSSpawnPointsGenerator.cpp b/Engine/Plugins/Runtime/MassGameplay/Source/MassSpawner/Private/MassEntityEQSSpawnPointsGenerator.cpp
--- a/Engine/Plugins/Runtime/MassGameplay/Source/MassSpawner/Private/MassEntityEQSSpawnPointsGenerator.cpp
+++ b/Engine/Plugins/Runtime/MassGameplay/Source/MassSpawner/Private/MassEntityEQSSpawnPointsGenerator.cpp
@@ -108,4 +108,18 @@ void UMassEntityEQSSpawnPointsGenerator::OnEQSQueryFinished(TSharedPtr<FEnvQuery
 #endif // ENABLE_VISUAL_LOG
 
 	FinishedGeneratingSpawnPointsDelegate.Execute(Results);
-}
\ No newline at end of file
+}
+
+#if WITH_EDITOR
+void UMassEntityEQSSpawnPointsGenerator::PostEditChangeProperty(FPropertyChangedEvent& PropertyChangedEvent)
+{
+	static const FName EQSRequestName = GET_MEMBER_NAME_CHECKED(UMassEntityEQSSpawnPointsGenerator, EQSRequest);
+
+	Super::PostEditChangeProperty(PropertyChangedEvent);
+
+	if (PropertyChangedEvent.MemberProperty && PropertyChangedEvent.MemberProperty->GetFName() == EQSRequestName)
+	{
+		EQSRequest.PostEditChangeProperty(*this, PropertyChangedEvent);
+	}
+}
+#endif
\ No newline at end of file
diff --git a/Engine/Plugins/Runtime/MassGameplay/Source/MassSpawner/Public/MassEntityEQSSpawnPointsGenerator.h b/Engine/Plugins/Runtime/MassGameplay/Source/MassSpawner/Public/MassEntityEQSSpawnPointsGenerator.h
--- a/Engine/Plugins/Runtime/MassGameplay/Source/MassSpawner/Public/MassEntityEQSSpawnPointsGenerator.h
+++ b/Engine/Plugins/Runtime/MassGameplay/Source/MassSpawner/Public/MassEntityEQSSpawnPointsGenerator.h
@@ -21,6 +21,11 @@ public:
 	virtual void Generate(UObject& QueryOwner, TConstArrayView<FMassSpawnedEntityType> EntityTypes, int32 Count, FFinishedGeneratingSpawnDataSignature& FinishedGeneratingSpawnPointsDelegate) const override;
 
 protected:
+
+#if WITH_EDITOR
+	virtual void PostEditChangeProperty(FPropertyChangedEvent& PropertyChangedEvent) override;
+#endif
+
 	void OnEQSQueryFinished(TSharedPtr<FEnvQueryResult> EQSResult, TArray<FMassEntitySpawnDataGeneratorResult> Results, FFinishedGeneratingSpawnDataSignature FinishedGeneratingSpawnPointsDelegate) const;
 
 	UPROPERTY(Category = "Query", EditAnywhere)
