Fixed UMassRepresentationActorManagement::ReleaseAnyActorOrCancelAnySpawning to not cause issues if given Representation's actor indices are INDEX_NONE

diff --git a/Engine/Plugins/Runtime/MassGameplay/Source/MassRepresentation/Private/MassRepresentationActorManagement.cpp b/Engine/Plugins/Runtime/MassGameplay/Source/MassRepresentation/Private/MassRepresentationActorManagement.cpp
--- a/Engine/Plugins/Runtime/MassGameplay/Source/MassRepresentation/Private/MassRepresentationActorManagement.cpp
+++ b/Engine/Plugins/Runtime/MassGameplay/Source/MassRepresentation/Private/MassRepresentationActorManagement.cpp
@@ -124,8 +124,11 @@ void UMassRepresentationActorManagement::ReleaseAnyActorOrCancelAnySpawning(UMas
 		ActorInfo.ResetAndUpdateHandleMap();
 	}
 	// Try releasing both as we can have a low res actor and a high res spawning request
-	RepresentationSubsystem.ReleaseTemplateActorOrCancelSpawning(MassAgent, Representation.HighResTemplateActorIndex, Actor, Representation.ActorSpawnRequestHandle);
-	if (Representation.LowResTemplateActorIndex != Representation.HighResTemplateActorIndex)
+	if (Representation.HighResTemplateActorIndex != INDEX_NONE)
+	{
+		RepresentationSubsystem.ReleaseTemplateActorOrCancelSpawning(MassAgent, Representation.HighResTemplateActorIndex, Actor, Representation.ActorSpawnRequestHandle);
+	}	
+	if (Representation.LowResTemplateActorIndex != Representation.HighResTemplateActorIndex && Representation.LowResTemplateActorIndex != INDEX_NONE)
 	{
 		RepresentationSubsystem.ReleaseTemplateActorOrCancelSpawning(MassAgent, Representation.LowResTemplateActorIndex, Actor, Representation.ActorSpawnRequestHandle);
 	}
