Massaged MassEntityConfig validation error output so the the logged lines allow one to identify the specific asset causing the issue

diff --git a/Engine/Plugins/Runtime/MassGameplay/Source/MassSpawner/Private/MassEntityTemplateRegistry.cpp b/Engine/Plugins/Runtime/MassGameplay/Source/MassSpawner/Private/MassEntityTemplateRegistry.cpp
--- a/Engine/Plugins/Runtime/MassGameplay/Source/MassSpawner/Private/MassEntityTemplateRegistry.cpp
+++ b/Engine/Plugins/Runtime/MassGameplay/Source/MassSpawner/Private/MassEntityTemplateRegistry.cpp
@@ -124,7 +124,7 @@ bool FMassEntityTemplateBuildContext::ValidateBuildContext(const UWorld& World)
 	// Loop through all the registered fragments and make sure only one trait registered them.
 	const UStruct* CurrentStruct = nullptr;
 	const UMassEntityTraitBase* CurrentTrait = nullptr;
-	bool bHeaderOutputed = false;
+	bool bHeaderOutputted = false;
 	bool bFragmentHasMultipleOwners = false;
 	for (const auto& Pair : TraitAddedTypes)
 	{
@@ -134,15 +134,17 @@ bool FMassEntityTemplateBuildContext::ValidateBuildContext(const UWorld& World)
 			CurrentTrait = Pair.Value;
 			check(CurrentTrait);
 			CurrentTrait->ValidateTemplate(*this, World);
-			bHeaderOutputed = false;
+			bHeaderOutputted = false;
 		}
 		else
 		{
-			if (!bHeaderOutputed)
+			if (!bHeaderOutputted)
 			{
-				UE_LOG(LogMass, Warning, TEXT("Fragment(%s) was added multiple time and should only be added by one trait. Fragment was added by:"), CurrentStruct ? *CurrentStruct->GetName() : TEXT("null"));
+				UE_LOG(LogMass, Warning, TEXT("%s: Fragment(%s) was added multiple time and should only be added by one trait. Fragment was added by:")
+					, CurrentTrait ? *GetNameSafe(CurrentTrait->GetOuter()) : TEXT("None")
+					, CurrentStruct ? *CurrentStruct->GetName() : TEXT("null"));
 				UE_LOG(LogMass, Warning, TEXT("\t\t%s"), CurrentTrait ? *CurrentTrait->GetClass()->GetName() : TEXT("null"));
-				bHeaderOutputed = true;
+				bHeaderOutputted = true;
 			}
 			UE_LOG(LogMass, Warning, TEXT("\t\t%s"), *Pair.Value->GetClass()->GetName());
 			bFragmentHasMultipleOwners = true;
@@ -151,22 +153,23 @@ bool FMassEntityTemplateBuildContext::ValidateBuildContext(const UWorld& World)
 
 	// Loop through all the traits dependencies and check if they have been added
 	CurrentTrait = nullptr;
-	bHeaderOutputed = false;
+	bHeaderOutputted = false;
 	bool bMissingFragmentDependencies = false;
 	for (const auto& Dependency : TraitsDependencies)
 	{
 		if (CurrentTrait != Dependency.Get<1>())
 		{
 			CurrentTrait = Dependency.Get<1>();
-			bHeaderOutputed = false;
+			bHeaderOutputted = false;
 		}
 		if (!TraitAddedTypes.Contains(Dependency.Get<0>()))
 		{
-			if (!bHeaderOutputed)
+			if (!bHeaderOutputted)
 			{
 				check(CurrentTrait);
-				UE_LOG(LogMass, Error, TEXT("Trait(%s) has missing dependency:"),  *CurrentTrait->GetClass()->GetName() );
-				bHeaderOutputed = true;
+				UE_LOG(LogMass, Error, TEXT("%s: Trait(%s) has missing dependency:"), *GetNameSafe(CurrentTrait->GetOuter())
+					, *CurrentTrait->GetClass()->GetName());
+				bHeaderOutputted = true;
 			}
 			UE_LOG(LogMass, Error, TEXT("\t\t%s"), *Dependency.Get<0>()->GetName());
 			bMissingFragmentDependencies = true;
