Enabled access the Mass' phase manager from its editor subsystem.

In an upcoming change a minimal amount of optional work is being done during the start and end of a processing phase. In order to register to the required callbacks, the phase manager was exposed.

diff --git a/Engine/Plugins/Experimental/TypedElementsDataStorage/Source/TypedElementsDataStorage/Private/TypedElementDatabase.cpp b/Engine/Plugins/Experimental/TypedElementsDataStorage/Source/TypedElementsDataStorage/Private/TypedElementDatabase.cpp
--- a/Engine/Plugins/Experimental/TypedElementsDataStorage/Source/TypedElementsDataStorage/Private/TypedElementDatabase.cpp
+++ b/Engine/Plugins/Experimental/TypedElementsDataStorage/Source/TypedElementsDataStorage/Private/TypedElementDatabase.cpp
@@ -28,7 +28,7 @@ void UTypedElementDatabase::Initialize()
 	check(Mass);
 	Mass->GetOnPreTickDelegate().AddUObject(this, &UTypedElementDatabase::OnPreMassTick);
 
-	ActiveEditorEntityManager = Mass->GetMutableEntityManager().AsShared();
+	ActiveEditorEntityManager = Mass->GetMutableEntityManager();
 }
 
 void UTypedElementDatabase::Deinitialize()
diff --git a/Engine/Plugins/Runtime/MassEntity/Source/MassEntityEditor/Public/MassEntityEditorSubsystem.h b/Engine/Plugins/Runtime/MassEntity/Source/MassEntityEditor/Public/MassEntityEditorSubsystem.h
--- a/Engine/Plugins/Runtime/MassEntity/Source/MassEntityEditor/Public/MassEntityEditorSubsystem.h
+++ b/Engine/Plugins/Runtime/MassEntity/Source/MassEntityEditor/Public/MassEntityEditorSubsystem.h
@@ -24,7 +24,8 @@ public:
 	UMassEntityEditorSubsystem();
 	~UMassEntityEditorSubsystem();
 
-	FMassEntityManager& GetMutableEntityManager() { return EntityManager.Get(); }
+	TSharedRef<FMassEntityManager> GetMutableEntityManager() { return EntityManager; }
+	TSharedRef<FMassProcessingPhaseManager> GetMutablePhaseManager() { return PhaseManager; }
 
 	FOnPreTick& GetOnPreTickDelegate() { return OnPreTickDelegate; }
 
diff --git a/Engine/Plugins/Runtime/MassGameplay/Source/MassGameplayEditor/Private/MassActorEditorSubsystem.cpp b/Engine/Plugins/Runtime/MassGameplay/Source/MassGameplayEditor/Private/MassActorEditorSubsystem.cpp
--- a/Engine/Plugins/Runtime/MassGameplay/Source/MassGameplayEditor/Private/MassActorEditorSubsystem.cpp
+++ b/Engine/Plugins/Runtime/MassGameplay/Source/MassGameplayEditor/Private/MassActorEditorSubsystem.cpp
@@ -13,6 +13,6 @@ void UMassActorEditorSubsystem::Initialize(FSubsystemCollectionBase& Collection)
 {
 	UMassEntityEditorSubsystem* MassEditorEditorSubsystem = Collection.InitializeDependency<UMassEntityEditorSubsystem>();
 	check(MassEditorEditorSubsystem);
-	FMassEntityManager& MassEntityManager = MassEditorEditorSubsystem->GetMutableEntityManager();
-	ActorManager = MakeShareable(new FMassActorManager(MassEntityManager.AsShared()));
+	TSharedRef<FMassEntityManager> MassEntityManager = MassEditorEditorSubsystem->GetMutableEntityManager();
+	ActorManager = MakeShareable(new FMassActorManager(MassEntityManager));
 }
