Added a way to find a trait instance of a specific class in a given MassEntityConfig

diff --git a/Engine/Plugins/Runtime/MassGameplay/Source/MassSpawner/Private/MassEntityConfigAsset.cpp b/Engine/Plugins/Runtime/MassGameplay/Source/MassSpawner/Private/MassEntityConfigAsset.cpp
--- a/Engine/Plugins/Runtime/MassGameplay/Source/MassSpawner/Private/MassEntityConfigAsset.cpp
+++ b/Engine/Plugins/Runtime/MassGameplay/Source/MassSpawner/Private/MassEntityConfigAsset.cpp
@@ -24,6 +24,19 @@ FMassEntityConfig::FMassEntityConfig(UObject& InOwner)
 
 }
 
+const UMassEntityTraitBase* FMassEntityConfig::FindTrait(TSubclassOf<UMassEntityTraitBase> TraitClass, const bool bExactMatch) const
+{
+	for (const TObjectPtr<UMassEntityTraitBase> Trait : Traits)
+	{
+		if (Trait && (bExactMatch ? Trait->GetClass() == TraitClass : Trait->IsA(TraitClass)))
+		{
+			return Trait;
+		}
+	}
+
+	return Parent ? Parent->FindTrait(TraitClass, bExactMatch) :  nullptr;
+}
+
 #if WITH_EDITOR
 void UMassEntityConfigAsset::ValidateEntityConfig()
 {
diff --git a/Engine/Plugins/Runtime/MassGameplay/Source/MassSpawner/Public/MassEntityConfigAsset.h b/Engine/Plugins/Runtime/MassGameplay/Source/MassSpawner/Public/MassEntityConfigAsset.h
--- a/Engine/Plugins/Runtime/MassGameplay/Source/MassSpawner/Public/MassEntityConfigAsset.h
+++ b/Engine/Plugins/Runtime/MassGameplay/Source/MassSpawner/Public/MassEntityConfigAsset.h
@@ -46,6 +46,9 @@ struct MASSSPAWNER_API FMassEntityConfig
 	/** @return View to the array of features defined on this config */
 	TConstArrayView<UMassEntityTraitBase*> GetTraits() const { return Traits; }
 
+	/** Looks for a trait of the indicated type, accepting all child classes as well, unless bExactMatch == true */
+	const UMassEntityTraitBase* FindTrait(TSubclassOf<UMassEntityTraitBase> TraitClass, const bool bExactMatch = false) const;
+
 	/** Adds Trait to the collection of traits hosted by this FMassEntityConfig instance */
 	void AddTrait(UMassEntityTraitBase& Trait);
 
@@ -121,6 +124,12 @@ public:
 		Config.DestroyEntityTemplate(World);
 	}
 
+	/** Looks for a trait of the indicated type, accepting all child classes as well, unless bExactMatch == true */
+	const UMassEntityTraitBase* FindTrait(TSubclassOf<UMassEntityTraitBase> TraitClass, const bool bExactMatch = false) const
+	{
+		return Config.FindTrait(TraitClass, bExactMatch);
+	}
+
 #if WITH_EDITOR
 
 	UFUNCTION(CallInEditor, Category = "Entity Config")
