Added bits to FMassEntityTemplateData to support creation of a given template's variation - create and instance and then fetch mutable fragment instances to modify the values being stored.

Also improved template's hashing by including the template name in the calculations.

diff --git a/Engine/Plugins/Runtime/MassGameplay/Source/MassSpawner/Private/MassEntityTemplate.cpp b/Engine/Plugins/Runtime/MassGameplay/Source/MassSpawner/Private/MassEntityTemplate.cpp
--- a/Engine/Plugins/Runtime/MassGameplay/Source/MassSpawner/Private/MassEntityTemplate.cpp
+++ b/Engine/Plugins/Runtime/MassGameplay/Source/MassSpawner/Private/MassEntityTemplate.cpp
@@ -17,8 +17,17 @@ FString FMassEntityTemplateID::ToString() const
 //----------------------------------------------------------------------//
 //  FMassEntityTemplateData
 //----------------------------------------------------------------------//
+FMassEntityTemplateData::FMassEntityTemplateData(const FMassEntityTemplate& InFinalizedTemplate)
+	: FMassEntityTemplateData(InFinalizedTemplate.GetTemplateData())
+{
+
+}
+
 uint32 GetTypeHash(const FMassEntityTemplateData& Template)
 {
+	// @todo: using the template name is temporary solution to allow to tell two templates apart based on something 
+	// else than composition. Ideally we would hash the fragment values instead.
+	const int32 NameHash = GetTypeHash(Template.GetTemplateName());
 	const uint32 CompositionHash = Template.Composition.CalculateHash();
 	const uint32 SharedFragmentValuesHash = GetTypeHash(Template.SharedFragmentValues);
 	
@@ -39,7 +48,7 @@ uint32 GetTypeHash(const FMassEntityTemplateData& Template)
 	}*/
 
 	// @todo maybe better to have two separate hashes - one for composition and maybe shared fragments, and then the other one for the rest, and use multimap for look up
-	return HashCombine(HashCombine(CompositionHash, SharedFragmentValuesHash), HashCombine(InitialValuesHash, InitializersHash));
+	return HashCombine(NameHash, HashCombine(HashCombine(CompositionHash, SharedFragmentValuesHash), HashCombine(InitialValuesHash, InitializersHash)));
 }
 
 //-----------------------------------------------------------------------------
diff --git a/Engine/Plugins/Runtime/MassGameplay/Source/MassSpawner/Public/MassEntityTemplate.h b/Engine/Plugins/Runtime/MassGameplay/Source/MassSpawner/Public/MassEntityTemplate.h
--- a/Engine/Plugins/Runtime/MassGameplay/Source/MassSpawner/Public/MassEntityTemplate.h
+++ b/Engine/Plugins/Runtime/MassGameplay/Source/MassSpawner/Public/MassEntityTemplate.h
@@ -14,6 +14,7 @@
 class UMassEntityTraitBase;
 struct FMassEntityView;
 struct FMassEntityTemplateIDFactory;
+struct FMassEntityTemplate;
 
 
 //ID of the template an entity is using
@@ -84,6 +85,7 @@ struct MASSSPAWNER_API FMassEntityTemplateData
 	typedef TFunction<void(UObject& /*Owner*/, FMassEntityView& /*EntityView*/, const EMassTranslationDirection /*CurrentDirection*/)> FObjectFragmentInitializerFunction;
 
 	FMassEntityTemplateData() = default;
+	explicit FMassEntityTemplateData(const FMassEntityTemplate& InFinalizedTemplate);
 
 	bool IsEmpty() const { return Composition.IsEmpty(); }
 
@@ -143,6 +145,14 @@ struct MASSSPAWNER_API FMassEntityTemplateData
 		return InitialFragmentValues.Emplace_GetRef(T::StaticStruct()).template GetMutable<T>();
 	}
 
+	template<typename T>
+	T* GetMutableFragment()
+	{
+		static_assert(TIsDerivedFrom<T, FMassFragment>::IsDerived, "Given struct doesn't represent a valid fragment type. Make sure to inherit from FMassFragment or one of its child-types.");
+		FInstancedStruct* Fragment = InitialFragmentValues.FindByPredicate(FStructTypeEqualOperator(T::StaticStruct()));
+		return Fragment ? &Fragment->template GetMutable<T>() : (T*)nullptr;
+	}
+
 	template<typename T>
 	void AddTag()
 	{
@@ -152,7 +162,7 @@ struct MASSSPAWNER_API FMassEntityTemplateData
 	
 	void AddTag(const UScriptStruct& TagType)
 	{
-		checkf(TagType.IsChildOf(FMassFragment::StaticStruct()), TEXT("Given struct doesn't represent a valid mass tag type. Make sure to inherit from FMassTag or one of its child-types."));
+		checkf(TagType.IsChildOf(FMassTag::StaticStruct()), TEXT("Given struct doesn't represent a valid mass tag type. Make sure to inherit from FMassTag or one of its child-types."));
 		Composition.Tags.Add(TagType);
 	}
 
@@ -298,6 +308,8 @@ struct MASSSPAWNER_API FMassEntityTemplate final : public TSharedFromThis<FMassE
 	FORCEINLINE const FMassArchetypeSharedFragmentValues& GetSharedFragmentValues() const { return TemplateData.GetSharedFragmentValues(); }
 	FORCEINLINE TConstArrayView<FInstancedStruct> GetInitialFragmentValues() const { return TemplateData.GetInitialFragmentValues(); }
 
+	const FMassEntityTemplateData& GetTemplateData() const { return TemplateData; }
+
 private:
 	FMassEntityTemplateData TemplateData;
 	FMassArchetypeHandle Archetype;
