Added a function for fetching an already created instance of a given FMassFragment type from FMassEntityTemplateBuildContext

diff --git a/Engine/Plugins/Runtime/MassGameplay/Source/MassSpawner/Public/MassEntityTemplateRegistry.h b/Engine/Plugins/Runtime/MassGameplay/Source/MassSpawner/Public/MassEntityTemplateRegistry.h
--- a/Engine/Plugins/Runtime/MassGameplay/Source/MassSpawner/Public/MassEntityTemplateRegistry.h
+++ b/Engine/Plugins/Runtime/MassGameplay/Source/MassSpawner/Public/MassEntityTemplateRegistry.h
@@ -91,6 +91,15 @@ struct FMassEntityTemplateBuildContext
 		TemplateData.AddSharedFragment(InSharedFragment);
 	}
 
+	template<typename T>
+	T& GetFragmentChecked()
+	{
+		check(TraitAddedTypes.Find(T::StaticStruct()) != nullptr);
+		T* FragmentInstance = TemplateData.GetMutableFragment<T>();
+		check(FragmentInstance);
+		return *FragmentInstance;
+	}
+
 	template<typename T>
 	bool HasFragment() const
 	{
