Added unite tests for MassEntityTemplate creation and mutating.
Added the MassGamplayTestSuite module, since these are the first MassGameplay tests.

diff --git a/Engine/Plugins/Runtime/MassEntity/Source/MassEntityTestSuite/Private/MassEntityTest.cpp b/Engine/Plugins/Runtime/MassEntity/Source/MassEntityTestSuite/Private/MassEntityTest.cpp
--- a/Engine/Plugins/Runtime/MassEntity/Source/MassEntityTestSuite/Private/MassEntityTest.cpp
+++ b/Engine/Plugins/Runtime/MassEntity/Source/MassEntityTestSuite/Private/MassEntityTest.cpp
@@ -145,7 +145,7 @@ struct FEntityTest_EntityCreationFromInstances : FEntityTestBase
 		AITEST_EQUAL("There should be one entity across the whole system", EntityManager->DebugGetEntityCount(), 1);
 		AITEST_EQUAL("Entity\'s archetype should be the Ints one", EntityManager->GetArchetypeForEntity(Entity), IntsArchetype);
 		AITEST_EQUAL("The created entity should have been added to the Ints archetype", EntityManager->DebugGetArchetypeEntitiesCount(IntsArchetype), 1);
-		AITEST_EQUAL("The entity should have the new component with the correct value set", EntityManager->GetFragmentDataChecked<FTestFragment_Int>(Entity).Value, TestIntValue);
+		AITEST_EQUAL("The entity should have the new component with the correct value set", EntityManager->GetFragmentDataChecked<FTestFragment_Int>(Entity).Value, FTestFragment_Int::TestIntValue);
 
 		return true;
 	}
@@ -201,7 +201,7 @@ struct FEntityTest_AddingFragmentInstance : FEntityTestBase
 		AITEST_EQUAL("The archetype containing just the new fragment should not be affected", EntityManager->DebugGetArchetypeEntitiesCount(IntsArchetype), 0);
 		// this test was originally failing due to FEntityData.CurrentArchetype not getting updated during entity moving between archetypes
 		AITEST_EQUAL("The entity should get associated with the new archetype", EntityManager->GetArchetypeForEntity(Entity), FloatsIntsArchetype);
-		AITEST_EQUAL("The entity should have the new component with the correct value set", EntityManager->GetFragmentDataChecked<FTestFragment_Int>(Entity).Value, TestIntValue);
+		AITEST_EQUAL("The entity should have the new component with the correct value set", EntityManager->GetFragmentDataChecked<FTestFragment_Int>(Entity).Value, FTestFragment_Int::TestIntValue);
 
 		return true;
 	}
@@ -296,7 +296,7 @@ struct FEntityTest_EntityReservationAndBuildingFromInstances : FEntityTestBase
 		AITEST_EQUAL("There should be one entity across the whole system", EntityManager->DebugGetEntityCount(), 1);
 		AITEST_EQUAL("Entity\'s archetype should be the Ints one", EntityManager->GetArchetypeForEntity(ReservedEntity), IntsArchetype);
 		AITEST_EQUAL("The created entity should have been added to the Ints archetype", EntityManager->DebugGetArchetypeEntitiesCount(IntsArchetype), 1);
-		AITEST_EQUAL("The entity should have the new component with the correct value set", EntityManager->GetFragmentDataChecked<FTestFragment_Int>(ReservedEntity).Value, TestIntValue);
+		AITEST_EQUAL("The entity should have the new component with the correct value set", EntityManager->GetFragmentDataChecked<FTestFragment_Int>(ReservedEntity).Value, FTestFragment_Int::TestIntValue);
 		EntityManager->DestroyEntity(ReservedEntity);
 		AITEST_EQUAL("There should not be any entity across the whole system", EntityManager->DebugGetEntityCount(), 0);
 		AITEST_EQUAL("The original archetype should now have no entities", EntityManager->DebugGetArchetypeEntitiesCount(FloatsArchetype), 0);
diff --git a/Engine/Plugins/Runtime/MassEntity/Source/MassEntityTestSuite/Private/MassTestTypes.cpp b/Engine/Plugins/Runtime/MassEntity/Source/MassEntityTestSuite/Private/MassTestTypes.cpp
--- a/Engine/Plugins/Runtime/MassEntity/Source/MassEntityTestSuite/Private/MassTestTypes.cpp
+++ b/Engine/Plugins/Runtime/MassEntity/Source/MassEntityTestSuite/Private/MassTestTypes.cpp
@@ -34,7 +34,7 @@ bool FEntityTestBase::SetUp()
 	FloatsIntsArchetype = EntityManager->CreateArchetype(MakeArrayView(FragmentTypes, 2));
 
 	FTestFragment_Int IntFrag;
-	IntFrag.Value = TestIntValue;
+	IntFrag.Value = FTestFragment_Int::TestIntValue;
 	InstanceInt = FInstancedStruct::Make(IntFrag);
 
 	return true;
diff --git a/Engine/Plugins/Runtime/MassEntity/Source/MassEntityTestSuite/Public/MassEntityTestTypes.h b/Engine/Plugins/Runtime/MassEntity/Source/MassEntityTestSuite/Public/MassEntityTestTypes.h
--- a/Engine/Plugins/Runtime/MassEntity/Source/MassEntityTestSuite/Public/MassEntityTestTypes.h
+++ b/Engine/Plugins/Runtime/MassEntity/Source/MassEntityTestSuite/Public/MassEntityTestTypes.h
@@ -40,6 +40,8 @@ struct FTestFragment_Int : public FMassFragment
 	int32 Value = 0;
 
 	FTestFragment_Int(const int32 InValue = 0) : Value(InValue) {}
+
+	static constexpr int32 TestIntValue = 123456;
 };
 
 USTRUCT()
@@ -75,6 +77,33 @@ struct FTestFragment_Array : public FMassFragment
 	}
 };
 
+USTRUCT()
+struct FTestChunkFragment_Int : public FMassChunkFragment
+{
+	GENERATED_BODY()
+	int32 Value = 0;
+
+	FTestChunkFragment_Int(const int32 InValue = 0) : Value(InValue) {}
+};
+
+USTRUCT()
+struct FTestSharedFragment_Int : public FMassSharedFragment
+{
+	GENERATED_BODY()
+	int32 Value = 0;
+
+	FTestSharedFragment_Int(const int32 InValue = 0) : Value(InValue) {}
+};
+
+USTRUCT()
+struct FTestSharedFragment_Float : public FMassSharedFragment
+{
+	GENERATED_BODY()
+	float Value = 0.f;
+
+	FTestSharedFragment_Float(const float InValue = 0) : Value(InValue) {}
+};
+
 /** @todo rename to FTestTag */
 USTRUCT()
 struct FTestFragment_Tag : public FMassTag
@@ -228,8 +257,6 @@ struct MASSENTITYTESTSUITE_API FExecutionTestBase : FAITestBase
 	virtual bool SetUp() override;
 };
 
-const int TestIntValue = 123456;
-
 struct MASSENTITYTESTSUITE_API FEntityTestBase : FExecutionTestBase
 {
 	FMassArchetypeHandle EmptyArchetype;
diff --git a/Engine/Plugins/Runtime/MassGameplay/MassGameplay.uplugin b/Engine/Plugins/Runtime/MassGameplay/MassGameplay.uplugin
--- a/Engine/Plugins/Runtime/MassGameplay/MassGameplay.uplugin
+++ b/Engine/Plugins/Runtime/MassGameplay/MassGameplay.uplugin
@@ -72,6 +72,11 @@
 		{
 			"Name": "MassGameplayExternalTraits",
 			"Type": "Runtime"
+		},
+		{
+			"Name": "MassGameplayTestSuite",
+			"Type": "UncookedOnly",
+			"LoadingPhase": "Default"
 		}
 	],
 	"Plugins": [
diff --git a/Engine/Plugins/Runtime/MassGameplay/Source/MassGameplayTestSuite/MassGameplayTestSuite.Build.cs b/Engine/Plugins/Runtime/MassGameplay/Source/MassGameplayTestSuite/MassGameplayTestSuite.Build.cs
new file mode 100644
--- /dev/null
+++ b/Engine/Plugins/Runtime/MassGameplay/Source/MassGameplayTestSuite/MassGameplayTestSuite.Build.cs
@@ -0,0 +1,35 @@
+// Copyright Epic Games, Inc. All Rights Reserved.
+
+using UnrealBuildTool;
+
+namespace UnrealBuildTool.Rules
+{
+	public class MassGameplayTestSuite : ModuleRules
+	{
+		public MassGameplayTestSuite(ReadOnlyTargetRules Target) : base(Target)
+		{
+			OptimizeCode = CodeOptimization.Never;
+			bUseUnity = false;
+
+			PCHUsage = ModuleRules.PCHUsageMode.UseExplicitOrSharedPCHs;
+
+			PublicIncludePaths.AddRange(
+				new string[] {
+				}
+			);
+
+			PublicDependencyModuleNames.AddRange(
+				new string[] {
+					"Core",
+					"CoreUObject",
+					"Engine",
+					"AITestSuite",
+					"MassEntity",
+					"MassEntityTestSuite",
+					"MassSpawner",
+					"StructUtils"
+				}
+			);
+		}
+	}
+}
\ No newline at end of file
diff --git a/Engine/Plugins/Runtime/MassGameplay/Source/MassGameplayTestSuite/Private/MassEntityTemplateTest.cpp b/Engine/Plugins/Runtime/MassGameplay/Source/MassGameplayTestSuite/Private/MassEntityTemplateTest.cpp
new file mode 100644
--- /dev/null
+++ b/Engine/Plugins/Runtime/MassGameplay/Source/MassGameplayTestSuite/Private/MassEntityTemplateTest.cpp
@@ -0,0 +1,192 @@
+// Copyright Epic Games, Inc. All Rights Reserved.
+
+#include "AITestsCommon.h"
+#include "MassEntityTypes.h"
+#include "MassEntityTestTypes.h"
+#include "MassEntityTemplate.h"
+#include "MassEntityTemplateRegistry.h"
+
+#define LOCTEXT_NAMESPACE "MassTest"
+
+UE_DISABLE_OPTIMIZATION_SHIP
+
+namespace FMassEntityTemplateTest
+{
+
+struct FEntityTemplateBase : FExecutionTestBase
+{
+	FMassEntityTemplateRegistry TemplateRegistry;
+
+	virtual bool SetUp() override
+	{
+		if (FExecutionTestBase::SetUp())
+		{
+			TemplateRegistry.Initialize(EntityManager);
+			return true;
+		}
+		return false;
+	}
+};
+
+
+struct FEntityTemplate_Empty : FEntityTemplateBase
+{
+	virtual bool InstantTest() override
+	{
+		FMassEntityTemplateData TemplateData;
+
+		const TSharedRef<FMassEntityTemplate>& FinalizedTemplate = TemplateRegistry.FindOrAddTemplate(MoveTemp(TemplateData));
+		
+		const FMassEntityTemplateID TemplateID = FinalizedTemplate->GetTemplateID();
+		AITEST_TRUE("Empty template is expected to be registered as a valid template", TemplateID.IsValid());
+
+		const FMassArchetypeHandle& ArchetypeHandle = FinalizedTemplate->GetArchetype();
+		AITEST_TRUE("Empty template is expected to produce a valid (if empty) archetype", ArchetypeHandle.IsValid());
+
+		const FMassArchetypeCompositionDescriptor& ArchetypeDescription = EntityManager->GetArchetypeComposition(ArchetypeHandle);
+		AITEST_TRUE("Empty template is expected to produce an empty-composition archetype", ArchetypeDescription.IsEmpty());
+
+		return true;
+	}
+};
+IMPLEMENT_AI_INSTANT_TEST(FEntityTemplate_Empty, "System.Mass.EntityTemplate.Empty");
+
+
+struct FEntityTemplate_Composition : FEntityTemplateBase
+{
+	virtual bool InstantTest() override
+	{
+		FMassEntityTemplateData TemplateData;
+
+		TemplateData.AddFragment<FTestFragment_Int>();
+		TemplateData.AddFragment<FTestFragment_Float>();
+		TemplateData.AddChunkFragment<FTestChunkFragment_Int>();
+		{
+			FSharedStruct SharedFragment = EntityManager->GetOrCreateSharedFragmentByHash<FTestSharedFragment_Int>(0);
+			TemplateData.AddSharedFragment(SharedFragment);
+		}
+		{
+			FSharedStruct SharedFragment = EntityManager->GetOrCreateSharedFragmentByHash<FTestSharedFragment_Float>(1);
+			TemplateData.AddConstSharedFragment(SharedFragment);
+		}
+		TemplateData.Sort();
+
+		FMassArchetypeCompositionDescriptor ExpectedComposition;
+		ExpectedComposition.Fragments.Add<FTestFragment_Int>();
+		ExpectedComposition.Fragments.Add<FTestFragment_Float>();
+		ExpectedComposition.ChunkFragments.Add<FTestChunkFragment_Int>();
+		ExpectedComposition.SharedFragments.Add<FTestSharedFragment_Int>();
+		ExpectedComposition.SharedFragments.Add<FTestSharedFragment_Float>();
+
+		AITEST_TRUE("The composition should end up being the same regardless of whether the data is added via a template data or a composition descriptor"
+			, ExpectedComposition.IsEquivalent(TemplateData.GetCompositionDescriptor()));
+
+		const TSharedRef<FMassEntityTemplate>& FinalizedTemplate = TemplateRegistry.FindOrAddTemplate(MoveTemp(TemplateData));
+
+		const FMassArchetypeHandle& ArchetypeHandle = FinalizedTemplate->GetArchetype();
+		AITEST_TRUE("Empty template is expected to produce a valid (if empty) archetype", ArchetypeHandle.IsValid());
+
+		const FMassArchetypeCompositionDescriptor& ArchetypeDescription = EntityManager->GetArchetypeComposition(ArchetypeHandle);
+		AITEST_TRUE("Empty template is expected to produce an empty-composition archetype", ArchetypeDescription.IsEquivalent(ExpectedComposition));
+
+		return true;
+	}
+};
+IMPLEMENT_AI_INSTANT_TEST(FEntityTemplate_Composition, "System.Mass.EntityTemplate.Composition");
+
+
+struct FEntityTemplate_Trivial : FEntityTemplateBase
+{
+	virtual bool InstantTest() override
+	{
+		FMassEntityTemplateData EmptyTemplateData;
+		FMassEntityTemplateData TemplateDataA;
+		FMassEntityTemplateData TemplateDataB;
+
+		TemplateDataA.AddFragment<FTestFragment_Int>();
+		TemplateDataB.AddFragment<FTestFragment_Int>();
+
+		const TSharedRef<FMassEntityTemplate>& FinalizedEmptyTemplate = TemplateRegistry.FindOrAddTemplate(MoveTemp(EmptyTemplateData));
+		const TSharedRef<FMassEntityTemplate>& FinalizedTemplateA = TemplateRegistry.FindOrAddTemplate(MoveTemp(TemplateDataA));
+		const TSharedRef<FMassEntityTemplate>& FinalizedTemplateB = TemplateRegistry.FindOrAddTemplate(MoveTemp(TemplateDataA));
+
+		AITEST_NOT_EQUAL("Non-empty template data should result in a finalized template different from the empty one", FinalizedEmptyTemplate->GetTemplateID(), FinalizedTemplateA->GetTemplateID());
+		AITEST_EQUAL("Non-empty template data should result in the very same finalized template", FinalizedTemplateA, FinalizedTemplateB);
+
+		const FMassArchetypeHandle& ArchetypeHandle = FinalizedTemplateA->GetArchetype();
+		AITEST_TRUE("Non-empty template is expected to produce a valid archetype", ArchetypeHandle.IsValid());
+
+		return true;
+	}
+};
+IMPLEMENT_AI_INSTANT_TEST(FEntityTemplate_Trivial, "System.Mass.EntityTemplate.Trivial");
+
+
+struct FEntityTemplate_Modified : FEntityTemplateBase
+{
+	virtual bool InstantTest() override
+	{
+		UScriptStruct* IntFragmentType = FTestFragment_Int::StaticStruct();
+		FMassArchetypeHandle IntsArchetype = EntityManager->CreateArchetype(MakeArrayView(&IntFragmentType, 1));
+
+		FMassEntityTemplateData OriginalTemplateData;
+		FTestFragment_Int& IntFragment = OriginalTemplateData.AddFragment_GetRef<FTestFragment_Int>();
+		IntFragment.Value = FTestFragment_Int::TestIntValue;
+
+		FMassEntityTemplateData MovedOriginalTemplateData = OriginalTemplateData;
+		const TSharedRef<FMassEntityTemplate>& FinalizedOriginalTemplate = TemplateRegistry.FindOrAddTemplate(MoveTemp(MovedOriginalTemplateData));
+
+		FMassEntityTemplateData NewTemplateData(*FinalizedOriginalTemplate);
+
+		AITEST_TRUE("Template data created from a finalized template should match the composition of the original data"
+			, NewTemplateData.GetCompositionDescriptor().IsEquivalent(OriginalTemplateData.GetCompositionDescriptor()));
+
+		// need to set a new name, otherwise the hash calculated while creating finalized template will end up being the 
+		// same as the original template's. This is the case since we're not using struct's contents while calculating hash (yet).
+		NewTemplateData.SetTemplateName(TEXT("ModifiedTemplate"));
+		FTestFragment_Int* CoppiedFragment = NewTemplateData.GetMutableFragment<FTestFragment_Int>();
+		AITEST_NOT_NULL("The fragment instance is expected to be found", CoppiedFragment);
+		AITEST_EQUAL("The fragment instance is expected to have the same value as the original one", CoppiedFragment->Value, IntFragment.Value);
+		CoppiedFragment->Value = FTestFragment_Int::TestIntValue + 1;
+		AITEST_NOT_EQUAL("Modifying the coppied instance of the fragment doesn't affect the original", CoppiedFragment->Value, IntFragment.Value);
+
+		const TSharedRef<FMassEntityTemplate>& FinalizedModifiedTemplate = TemplateRegistry.FindOrAddTemplate(MoveTemp(NewTemplateData));
+
+		AITEST_NOT_EQUAL("The original and modified templates should end up resulting in two different templates", FinalizedModifiedTemplate->GetTemplateID(), FinalizedOriginalTemplate->GetTemplateID());
+		AITEST_EQUAL("The original and modified templates should still point at the same archetype", FinalizedModifiedTemplate->GetArchetype(), FinalizedOriginalTemplate->GetArchetype());
+
+		AITEST_EQUAL("The resulting archetype should match the IntArchetypeHandle", FinalizedModifiedTemplate->GetArchetype(), IntsArchetype);
+
+		return true;
+	}
+};
+IMPLEMENT_AI_INSTANT_TEST(FEntityTemplate_Modified, "System.Mass.EntityTemplate.Modification");
+
+
+//struct FEntityTemplate_Modified : FEntityTemplateBase
+//{
+//	virtual bool InstantTest() override
+//	{
+//		FMassEntityTemplateData OriginalTemplateData;
+//		OriginalTemplateData.AddFragment<FTestFragment_Int>();
+//
+//		FMassEntityTemplateData OriginalTemplateDataCopy = OriginalTemplateData;
+//		const TSharedRef<FMassEntityTemplate>& FinalizedOriginalTemplate = TemplateRegistry.FindOrAddTemplate(MoveTemp(OriginalTemplateData));
+//		
+//
+//		AITEST_NOT_EQUAL("Non-empty template data should result in a finalized template different from the empty one", FinalizedEmptyTemplate->GetTemplateID(), FinalizedTemplateA->GetTemplateID());
+//		AITEST_EQUAL("Non-empty template data should result in the very same finalized template", FinalizedTemplateA, FinalizedTemplateB);
+//
+//		const FMassArchetypeHandle& ArchetypeHandle = FinalizedTemplateA->GetArchetype();
+//		AITEST_TRUE("Non-empty template is expected to produce a valid archetype", ArchetypeHandle.IsValid());*/
+//
+//		return true;
+//	}
+//};
+//IMPLEMENT_AI_INSTANT_TEST(FEntityTemplate_Modified, "System.Mass.EntityTemplate.Modification");
+
+} // FMassEntityTemplateTest
+
+UE_ENABLE_OPTIMIZATION_SHIP
+
+#undef LOCTEXT_NAMESPACE
diff --git a/Engine/Plugins/Runtime/MassGameplay/Source/MassGameplayTestSuite/Private/MassGameplayTestSuiteModule.cpp b/Engine/Plugins/Runtime/MassGameplay/Source/MassGameplayTestSuite/Private/MassGameplayTestSuiteModule.cpp
new file mode 100644
--- /dev/null
+++ b/Engine/Plugins/Runtime/MassGameplay/Source/MassGameplayTestSuite/Private/MassGameplayTestSuiteModule.cpp
@@ -0,0 +1,13 @@
+// Copyright Epic Games, Inc. All Rights Reserved.
+
+#include "MassGameplayTestSuiteModule.h"
+
+#define LOCTEXT_NAMESPACE "MassTest"
+
+class FMassGameplayTestSuiteModule : public IMassGameplayTestSuiteModule
+{
+};
+
+IMPLEMENT_MODULE(FMassGameplayTestSuiteModule, MassGameplayTestSuite)
+
+#undef LOCTEXT_NAMESPACE
diff --git a/Engine/Plugins/Runtime/MassGameplay/Source/MassGameplayTestSuite/Public/MassGameplayTestSuiteModule.h b/Engine/Plugins/Runtime/MassGameplay/Source/MassGameplayTestSuite/Public/MassGameplayTestSuiteModule.h
new file mode 100644
--- /dev/null
+++ b/Engine/Plugins/Runtime/MassGameplay/Source/MassGameplayTestSuite/Public/MassGameplayTestSuiteModule.h
@@ -0,0 +1,37 @@
+// Copyright Epic Games, Inc. All Rights Reserved.
+
+#pragma once
+
+#include "CoreMinimal.h"
+#include "Modules/ModuleManager.h"
+
+/**
+* The public interface to this module
+*/
+class IMassGameplayTestSuiteModule : public IModuleInterface
+{
+
+public:
+
+	/**
+	* Singleton-like access to this module's interface.  This is just for convenience!
+	* Beware of calling this during the shutdown phase, though.  Your module might have been unloaded already.
+	*
+	* @return Returns singleton instance, loading the module on demand if needed
+	*/
+	static inline IMassGameplayTestSuiteModule& Get()
+	{
+		return FModuleManager::LoadModuleChecked<IMassGameplayTestSuiteModule>("MassGameplayTestSuite");
+	}
+
+	/**
+	* Checks to see if this module is loaded and ready.  It is only valid to call Get() if IsAvailable() returns true.
+	*
+	* @return True if the module is loaded and ready to use
+	*/
+	static inline bool IsAvailable()
+	{
+		return FModuleManager::Get().IsModuleLoaded("MassGameplayTestSuite");
+	}
+};
+
