Temporary fix to UE::MassSpawner::HashTraits to avoid an easy to repro hash clash.

diff --git a/Engine/Plugins/Runtime/MassGameplay/Source/MassSpawner/Private/MassSpawnerTypes.cpp b/Engine/Plugins/Runtime/MassGameplay/Source/MassSpawner/Private/MassSpawnerTypes.cpp
--- a/Engine/Plugins/Runtime/MassGameplay/Source/MassSpawner/Private/MassSpawnerTypes.cpp
+++ b/Engine/Plugins/Runtime/MassGameplay/Source/MassSpawner/Private/MassSpawnerTypes.cpp
@@ -24,6 +24,10 @@ namespace UE::MassSpawner
 		{
 			FArchiveObjectCRC32AgentConfig Archive;
 			CRC = Archive.Crc32(Trait, CRC);
+			// @todo this piece is here to avoid an easy to repro hash class - all one needs to do is to add a 
+			// trait subclass that sets different default values
+			check(Trait && Trait->GetClass());
+			CRC = HashCombine(CRC, GetTypeHash(Trait->GetClass()->GetName()));
 		}
 		return CRC;
 	}
