Fixed FMassClientBubbleSerializerBase instances not releasing their FMassEntityManager shared pointer soon enough

diff --git a/Engine/Plugins/Runtime/MassGameplay/Source/MassReplication/Private/MassClientBubbleInfoBase.cpp b/Engine/Plugins/Runtime/MassGameplay/Source/MassReplication/Private/MassClientBubbleInfoBase.cpp
--- a/Engine/Plugins/Runtime/MassGameplay/Source/MassReplication/Private/MassClientBubbleInfoBase.cpp
+++ b/Engine/Plugins/Runtime/MassGameplay/Source/MassReplication/Private/MassClientBubbleInfoBase.cpp
@@ -86,11 +86,15 @@ void AMassClientBubbleInfoBase::BeginPlay()
 
 void AMassClientBubbleInfoBase::EndPlay(const EEndPlayReason::Type EndPlayReason)
 {
-	for (const FMassClientBubbleSerializerBase* Serializer : Serializers)
+	for (FMassClientBubbleSerializerBase* Serializer : Serializers)
 	{
-		check(Serializer->GetClientHandler());
+		if (ensure(Serializer))
+		{
+			Serializer->ShutDown();
 
-		Serializer->GetClientHandler()->Reset();
+			check(Serializer->GetClientHandler());
+			Serializer->GetClientHandler()->Reset();
+		}
 	}
 
 	Super::EndPlay(EndPlayReason);
diff --git a/Engine/Plugins/Runtime/MassGameplay/Source/MassReplication/Private/MassClientBubbleSerializerBase.cpp b/Engine/Plugins/Runtime/MassGameplay/Source/MassReplication/Private/MassClientBubbleSerializerBase.cpp
--- a/Engine/Plugins/Runtime/MassGameplay/Source/MassReplication/Private/MassClientBubbleSerializerBase.cpp
+++ b/Engine/Plugins/Runtime/MassGameplay/Source/MassReplication/Private/MassClientBubbleSerializerBase.cpp
@@ -5,6 +5,14 @@
 #include "MassEntityUtils.h"
 
 
+void FMassClientBubbleSerializerBase::ShutDown()
+{
+	EntityManager = nullptr;
+	World = nullptr;
+	SpawnerSubsystem = nullptr;
+	ReplicationSubsystem = nullptr;
+}
+
 void FMassClientBubbleSerializerBase::InitializeForWorld(UWorld& InWorld)
 {
 	World = &InWorld;
diff --git a/Engine/Plugins/Runtime/MassGameplay/Source/MassReplication/Public/MassClientBubbleSerializerBase.h b/Engine/Plugins/Runtime/MassGameplay/Source/MassReplication/Public/MassClientBubbleSerializerBase.h
--- a/Engine/Plugins/Runtime/MassGameplay/Source/MassReplication/Public/MassClientBubbleSerializerBase.h
+++ b/Engine/Plugins/Runtime/MassGameplay/Source/MassReplication/Public/MassClientBubbleSerializerBase.h
@@ -41,6 +41,8 @@ public:
 	void SetClientHandler(IClientBubbleHandlerInterface& InClientHandler) { ClientHandler = &InClientHandler; }
 	IClientBubbleHandlerInterface* GetClientHandler() const { return ClientHandler; }
 
+	void ShutDown();
+
 private:
 	UPROPERTY(Transient)
 	TObjectPtr<UWorld> World = nullptr;
