Fixed mismatched copy constructor/assignment operator in FMassNetworkID.

diff --git a/Engine/Plugins/Runtime/MassGameplay/Source/MassCommon/Public/MassCommonTypes.h b/Engine/Plugins/Runtime/MassGameplay/Source/MassCommon/Public/MassCommonTypes.h
--- a/Engine/Plugins/Runtime/MassGameplay/Source/MassCommon/Public/MassCommonTypes.h
+++ b/Engine/Plugins/Runtime/MassGameplay/Source/MassCommon/Public/MassCommonTypes.h
@@ -44,7 +44,6 @@ struct FMassNetworkID : public FSequentialIDBase
 	GENERATED_BODY()
 
 	FMassNetworkID() = default;
-	FMassNetworkID(const FMassNetworkID& Other) = default;
 	explicit FMassNetworkID(uint32 InID) : FSequentialIDBase(InID) {}
 };
 
