Minor massage of how UMassActorSpawnerSubsystem bind's to MassSimulationSubsystem.PhaseManager's events

diff --git a/Engine/Plugins/Runtime/MassGameplay/Source/MassActors/Private/MassActorSpawnerSubsystem.cpp b/Engine/Plugins/Runtime/MassGameplay/Source/MassActors/Private/MassActorSpawnerSubsystem.cpp
--- a/Engine/Plugins/Runtime/MassGameplay/Source/MassActors/Private/MassActorSpawnerSubsystem.cpp
+++ b/Engine/Plugins/Runtime/MassGameplay/Source/MassActors/Private/MassActorSpawnerSubsystem.cpp
@@ -341,16 +341,13 @@ void UMassActorSpawnerSubsystem::ProcessPendingDestruction(const double MaxTimeS
 
 void UMassActorSpawnerSubsystem::Initialize(FSubsystemCollectionBase& Collection)
 {
-	Collection.InitializeDependency(UMassSimulationSubsystem::StaticClass());
+	UMassSimulationSubsystem* SimSystem = Collection.InitializeDependency<UMassSimulationSubsystem>();
+	check(SimSystem);
+
 	Super::Initialize(Collection);
 
-	if (const UWorld* World = GetWorld())
-	{
-		UMassSimulationSubsystem* SimSystem = UWorld::GetSubsystem<UMassSimulationSubsystem>(World);
-		check(SimSystem);
-		SimSystem->GetOnProcessingPhaseStarted(EMassProcessingPhase::PrePhysics).AddUObject(this, &UMassActorSpawnerSubsystem::OnPrePhysicsPhaseStarted);
-		SimSystem->GetOnProcessingPhaseFinished(EMassProcessingPhase::PrePhysics).AddUObject(this, &UMassActorSpawnerSubsystem::OnPrePhysicsPhaseFinished);
-	}
+	SimSystem->GetOnProcessingPhaseStarted(EMassProcessingPhase::PrePhysics).AddUObject(this, &UMassActorSpawnerSubsystem::OnPrePhysicsPhaseStarted);
+	SimSystem->GetOnProcessingPhaseFinished(EMassProcessingPhase::PrePhysics).AddUObject(this, &UMassActorSpawnerSubsystem::OnPrePhysicsPhaseFinished);
 }
 
 void UMassActorSpawnerSubsystem::Deinitialize()
