Fixed a flaw in MassSignals' delivery of delayed signals resulting in potentially sending out the notifications the very same frame they were posted (and similar issues with longer-waiting signals, off by one frame).

diff --git a/Engine/Plugins/Runtime/MassGameplay/Source/MassSignals/Private/MassSignalSubsystem.cpp b/Engine/Plugins/Runtime/MassGameplay/Source/MassSignals/Private/MassSignalSubsystem.cpp
--- a/Engine/Plugins/Runtime/MassGameplay/Source/MassSignals/Private/MassSignalSubsystem.cpp
+++ b/Engine/Plugins/Runtime/MassGameplay/Source/MassSignals/Private/MassSignalSubsystem.cpp
@@ -10,16 +10,31 @@
 
 CSV_DEFINE_CATEGORY(MassSignalsCounters, true);
 
+void UMassSignalSubsystem::Initialize(FSubsystemCollectionBase& Collection)
+{
+	Super::Initialize(Collection);
+	CachedWorld = &GetWorldRef();
+	checkf(CachedWorld, TEXT("UMassSignalSubsystem instances are expected to always be tied to a valid UWorld instance"));
+}
+
+void UMassSignalSubsystem::Deinitialize()
+{
+	CachedWorld = nullptr;
+	Super::Deinitialize();
+}
+
 void UMassSignalSubsystem::Tick(float DeltaTime)
 {
+	CA_ASSUME(CachedWorld);
+	const double CurrentTime = CachedWorld->GetTimeSeconds();
+
 	for (int i = 0; i < DelayedSignals.Num();)
 	{
 		FDelayedSignal& DelayedSignal = DelayedSignals[i];
-		DelayedSignal.DelayInSeconds -= DeltaTime;
-		if(DelayedSignal.DelayInSeconds <= 0.0f)
+		if (DelayedSignal.TargetTimestamp <= CurrentTime)
 		{
 			SignalEntities(DelayedSignal.SignalName, MakeArrayView(DelayedSignal.Entities));
-			DelayedSignals.RemoveAtSwap(i);
+			DelayedSignals.RemoveAtSwap(i, 1, /*bAllowShrinking=*/false);
 		}
 		else
 		{
@@ -64,7 +79,10 @@ void UMassSignalSubsystem::DelaySignalEntities(FName SignalName, TConstArrayView
 	FDelayedSignal& DelayedSignal = DelayedSignals.Emplace_GetRef();
 	DelayedSignal.SignalName = SignalName;
 	DelayedSignal.Entities = Entities;
-	DelayedSignal.DelayInSeconds = DelayInSeconds;
+
+	check(CachedWorld);
+	DelayedSignal.TargetTimestamp = CachedWorld->GetTimeSeconds() + DelayInSeconds;
+
 	UE_CVLOG(Entities.Num() == 1, this, LogMassSignals, Log, TEXT("Delay signal [%s] to entity [%s] in %.2f"), *SignalName.ToString(), *Entities[0].DebugGetDescription(), DelayInSeconds);
 	UE_CVLOG(Entities.Num() > 1,this, LogMassSignals, Log, TEXT("Delay signal [%s] to %d entities in %.2f"), *SignalName.ToString(), Entities.Num(), DelayInSeconds);
 }
diff --git a/Engine/Plugins/Runtime/MassGameplay/Source/MassSignals/Public/MassSignalSubsystem.h b/Engine/Plugins/Runtime/MassGameplay/Source/MassSignals/Public/MassSignalSubsystem.h
--- a/Engine/Plugins/Runtime/MassGameplay/Source/MassSignals/Public/MassSignalSubsystem.h
+++ b/Engine/Plugins/Runtime/MassGameplay/Source/MassSignals/Public/MassSignalSubsystem.h
@@ -95,6 +95,11 @@ public:
 	void DelaySignalEntitiesDeferred(FMassExecutionContext& Context, FName SignalName, TConstArrayView<FMassEntityHandle> Entities, const float DelayInSeconds);
 
 protected:
+	// USubsystem implementation Begin
+	virtual void Initialize(FSubsystemCollectionBase& Collection) override;
+	virtual void Deinitialize() override;
+	// USubsystem implementation End	
+
 	virtual void Tick(float DeltaTime) override;
 	virtual TStatId GetStatId() const override;
 
@@ -104,10 +109,13 @@ protected:
 	{
 		FName SignalName;
 		TArray<FMassEntityHandle> Entities;
-		float DelayInSeconds;
+		double TargetTimestamp;
 	};
 
 	TArray<FDelayedSignal> DelayedSignals;
+
+	UPROPERTY(transient)
+	TObjectPtr<UWorld> CachedWorld;
 };
 
 template<>
