Made UMassSpawnLocationProcessor run on clients as well since there's no reason not to - if the entities get spawned we do want the to get configured

diff --git a/Engine/Plugins/Runtime/MassGameplay/Source/MassSpawner/Private/MassSpawnLocationProcessor.cpp b/Engine/Plugins/Runtime/MassGameplay/Source/MassSpawner/Private/MassSpawnLocationProcessor.cpp
--- a/Engine/Plugins/Runtime/MassGameplay/Source/MassSpawner/Private/MassSpawnLocationProcessor.cpp
+++ b/Engine/Plugins/Runtime/MassGameplay/Source/MassSpawner/Private/MassSpawnLocationProcessor.cpp
@@ -31,53 +31,49 @@ void UMassSpawnLocationProcessor::Execute(FMassEntityManager& EntityManager, FMa
 	}
 
 	const UWorld* World = EntityManager.GetWorld();
-
 	check(World);
 
 	const ENetMode NetMode = World->GetNetMode();
 
-	if (NetMode != NM_Client)
+	FMassTransformsSpawnData& AuxData = Context.GetMutableAuxData().GetMutable<FMassTransformsSpawnData>();
+	TArray<FTransform>& Transforms = AuxData.Transforms;
+
+	const int32 NumSpawnTransforms = Transforms.Num();
+	if (NumSpawnTransforms == 0)
 	{
-		FMassTransformsSpawnData& AuxData = Context.GetMutableAuxData().GetMutable<FMassTransformsSpawnData>();
-		TArray<FTransform>& Transforms = AuxData.Transforms;
+		UE_VLOG_UELOG(this, LogMass, Error, TEXT("No spawn transforms provided. Entity transforms won't be initialized."));
+		return;
+	}
 
-		const int32 NumSpawnTransforms = Transforms.Num();
-		if (NumSpawnTransforms == 0)
+	int32 NumRequiredSpawnTransforms = 0;
+	EntityQuery.ForEachEntityChunk(EntityManager, Context, [&NumRequiredSpawnTransforms](const FMassExecutionContext& Context)
 		{
-			UE_VLOG_UELOG(this, LogMass, Error, TEXT("No spawn transforms provided. Entity transforms won't be initialized."));
-			return;
-		}
+			NumRequiredSpawnTransforms += Context.GetNumEntities();
+		});
 
-		int32 NumRequiredSpawnTransforms = 0;
-		EntityQuery.ForEachEntityChunk(EntityManager, Context, [&NumRequiredSpawnTransforms](const FMassExecutionContext& Context)
-			{
-				NumRequiredSpawnTransforms += Context.GetNumEntities();
-			});
+	const int32 NumToAdd = NumRequiredSpawnTransforms - NumSpawnTransforms;
+	if (NumToAdd > 0)
+	{
+		UE_VLOG_UELOG(this, LogMass, Warning,
+			TEXT("Not enough spawn locations provided (%d) for all entities (%d). Existing locations will be reused randomly to fill the %d missing positions."),
+			NumSpawnTransforms, NumRequiredSpawnTransforms, NumToAdd);
 
-		const int32 NumToAdd = NumRequiredSpawnTransforms - NumSpawnTransforms;
-		if (NumToAdd > 0)
+		Transforms.AddUninitialized(NumToAdd);
+		for (int i = 0; i < NumToAdd; ++i)
 		{
-			UE_VLOG_UELOG(this, LogMass, Warning,
-				TEXT("Not enough spawn locations provided (%d) for all entities (%d). Existing locations will be reused randomly to fill the %d missing positions."),
-				NumSpawnTransforms, NumRequiredSpawnTransforms, NumToAdd);
-
-			Transforms.AddUninitialized(NumToAdd);
-			for (int i = 0; i < NumToAdd; ++i)
-			{
-				Transforms[NumSpawnTransforms + i] = Transforms[FMath::RandRange(0, NumSpawnTransforms - 1)];
-			}
+			Transforms[NumSpawnTransforms + i] = Transforms[FMath::RandRange(0, NumSpawnTransforms - 1)];
 		}
+	}
 
-		EntityQuery.ForEachEntityChunk(EntityManager, Context, [&Transforms, this](FMassExecutionContext& Context)
+	EntityQuery.ForEachEntityChunk(EntityManager, Context, [&Transforms, this](FMassExecutionContext& Context)
+		{
+			const TArrayView<FTransformFragment> LocationList = Context.GetMutableFragmentView<FTransformFragment>();
+			const int32 NumEntities = Context.GetNumEntities();
+			for (int32 i = 0; i < NumEntities; ++i)
 			{
-				const TArrayView<FTransformFragment> LocationList = Context.GetMutableFragmentView<FTransformFragment>();
-				const int32 NumEntities = Context.GetNumEntities();
-				for (int32 i = 0; i < NumEntities; ++i)
-				{
-					const int32 AuxIndex = FMath::RandRange(0, Transforms.Num() - 1);
-					LocationList[i].GetMutableTransform() = Transforms[AuxIndex];
-					Transforms.RemoveAtSwap(AuxIndex, 1, /*bAllowShrinking=*/false);
-				}
-			});
-	}
+				const int32 AuxIndex = FMath::RandRange(0, Transforms.Num() - 1);
+				LocationList[i].GetMutableTransform() = Transforms[AuxIndex];
+				Transforms.RemoveAtSwap(AuxIndex, 1, /*bAllowShrinking=*/false);
+			}
+		});
 }
