Deprecated UMassSpawnerSubsystem::DestroyEntities flavor using FMassEntityTemplateID parameter, since the param is unused and not needed.

diff --git a/Engine/Plugins/Runtime/MassGameplay/Source/MassActors/Private/MassAgentSubsystem.cpp b/Engine/Plugins/Runtime/MassGameplay/Source/MassActors/Private/MassAgentSubsystem.cpp
--- a/Engine/Plugins/Runtime/MassGameplay/Source/MassActors/Private/MassAgentSubsystem.cpp
+++ b/Engine/Plugins/Runtime/MassGameplay/Source/MassActors/Private/MassAgentSubsystem.cpp
@@ -230,7 +230,7 @@ void UMassAgentSubsystem::UnregisterAgentComponent(UMassAgentComponent& AgentCom
 			}
 			else if (ensure(EntityTemplate))
 			{
-				SpawnerSystem->DestroyEntities(EntityTemplate->GetTemplateID(), TArrayView<FMassEntityHandle>(&Entity, 1));
+				SpawnerSystem->DestroyEntities(TArrayView<FMassEntityHandle>(&Entity, 1));
 			}
 		}
 		else if (AgentComp.IsEntityPendingCreation())
diff --git a/Engine/Plugins/Runtime/MassGameplay/Source/MassReplication/Public/MassClientBubbleHandler.h b/Engine/Plugins/Runtime/MassGameplay/Source/MassReplication/Public/MassClientBubbleHandler.h
--- a/Engine/Plugins/Runtime/MassGameplay/Source/MassReplication/Public/MassClientBubbleHandler.h
+++ b/Engine/Plugins/Runtime/MassGameplay/Source/MassReplication/Public/MassClientBubbleHandler.h
@@ -382,7 +382,7 @@ void TClientBubbleHandlerBase<AgentArrayItem>::UpdateAgentsToRemove()
 template<typename AgentArrayItem>
 void TClientBubbleHandlerBase<AgentArrayItem>::PreReplicatedRemove(const TArrayView<int32> RemovedIndices, int32 FinalSize)
 {
-	TMap<FMassEntityTemplateID, TArray<FMassEntityHandle>> EntitiesToDestroy;
+	TArray<FMassEntityHandle> EntitiesToDestroy;
 
 	UWorld* World = Serializer->GetWorld();
 	check(World);
@@ -403,21 +403,14 @@ void TClientBubbleHandlerBase<AgentArrayItem>::PreReplicatedRemove(const TArrayV
 		// We do need to check the ReplicationID in this case
 		if (Entity.IsSet())
 		{
-			TArray<FMassEntityHandle>& EntityArray = EntitiesToDestroy.FindOrAdd(RemovedItem.Agent.GetTemplateID());
-			EntityArray.Push(Entity);
-
+			EntitiesToDestroy.Add(Entity);
+			
 			check(AgentsRemoveDataMap.Find(RemovedItem.Agent.GetNetID()) == nullptr);
 			AgentsRemoveDataMap.Add(RemovedItem.Agent.GetNetID(), FMassAgentRemoveData(World->GetRealTimeSeconds()));
 		}
 	}
 
-	// Batch destroy agents per template
-	for (const TPair<FMassEntityTemplateID, TArray<FMassEntityHandle>>& Item : EntitiesToDestroy)
-	{
-		const FMassEntityTemplateID& TemplateID = Item.Key;
-		const TArray<FMassEntityHandle>& EntityArray = Item.Value;
-		SpawnerSubsystem->DestroyEntities(TemplateID, EntityArray);
-	}
+	SpawnerSubsystem->DestroyEntities(EntitiesToDestroy);
 }
 #endif //UE_REPLICATION_COMPILE_CLIENT_CODE
 
diff --git a/Engine/Plugins/Runtime/MassGameplay/Source/MassSpawner/Private/MassSpawner.cpp b/Engine/Plugins/Runtime/MassGameplay/Source/MassSpawner/Private/MassSpawner.cpp
--- a/Engine/Plugins/Runtime/MassGameplay/Source/MassSpawner/Private/MassSpawner.cpp
+++ b/Engine/Plugins/Runtime/MassGameplay/Source/MassSpawner/Private/MassSpawner.cpp
@@ -498,7 +498,7 @@ void AMassSpawner::DoDespawning()
 
 	for (const FSpawnedEntities& SpawnedEntities : AllSpawnedEntities)
 	{
-		SpawnerSystem->DestroyEntities(SpawnedEntities.TemplateID, SpawnedEntities.Entities);
+		SpawnerSystem->DestroyEntities(SpawnedEntities.Entities);
 	}
 	AllSpawnedEntities.Reset();
 
@@ -549,7 +549,7 @@ bool AMassSpawner::DespawnEntity(const FMassEntityHandle Entity)
 		const int32 Index = SpawnedEntities.Entities.Find(Entity);
 		if (Index != INDEX_NONE)
 		{
-			SpawnerSystem->DestroyEntities(SpawnedEntities.TemplateID, MakeArrayView(&Entity, 1));
+			SpawnerSystem->DestroyEntities(MakeArrayView(&Entity, 1));
 			SpawnedEntities.Entities.RemoveAtSwap(Index, 1, /*bAllowShrinking=*/false);
 			return true;
 		}
diff --git a/Engine/Plugins/Runtime/MassGameplay/Source/MassSpawner/Private/MassSpawnerSubsystem.cpp b/Engine/Plugins/Runtime/MassGameplay/Source/MassSpawner/Private/MassSpawnerSubsystem.cpp
--- a/Engine/Plugins/Runtime/MassGameplay/Source/MassSpawner/Private/MassSpawnerSubsystem.cpp
+++ b/Engine/Plugins/Runtime/MassGameplay/Source/MassSpawner/Private/MassSpawnerSubsystem.cpp
@@ -69,7 +69,7 @@ void UMassSpawnerSubsystem::SpawnEntities(FMassEntityTemplateID TemplateID, cons
 	DoSpawning(EntityTemplate->Get(), NumberToSpawn, SpawnData, InitializerClass, OutEntities);
 }
 
-void UMassSpawnerSubsystem::DestroyEntities(const FMassEntityTemplateID TemplateID, TConstArrayView<FMassEntityHandle> Entities)
+void UMassSpawnerSubsystem::DestroyEntities(TConstArrayView<FMassEntityHandle> Entities)
 {
 	TRACE_CPUPROFILER_EVENT_SCOPE_STR("MassSpawnerSubsystem_DestroyEntities")
 
diff --git a/Engine/Plugins/Runtime/MassGameplay/Source/MassSpawner/Public/MassSpawnerSubsystem.h b/Engine/Plugins/Runtime/MassGameplay/Source/MassSpawner/Public/MassSpawnerSubsystem.h
--- a/Engine/Plugins/Runtime/MassGameplay/Source/MassSpawner/Public/MassSpawnerSubsystem.h
+++ b/Engine/Plugins/Runtime/MassGameplay/Source/MassSpawner/Public/MassSpawnerSubsystem.h
@@ -33,7 +33,13 @@ public:
 
 	void SpawnEntities(FMassEntityTemplateID TemplateID, const uint32 NumberToSpawn, FConstStructView SpawnData, TSubclassOf<UMassProcessor> InitializerClass, TArray<FMassEntityHandle>& OutEntities);
 
-	void DestroyEntities(const FMassEntityTemplateID TemplateID, TConstArrayView<FMassEntityHandle> Entities);
+	void DestroyEntities(TConstArrayView<FMassEntityHandle> Entities);
+
+	UE_DEPRECATED(5.3, "This flavor of DestroyEntities has been deprecated. Use the one without the FMassEntityTemplateID parameter")
+	void DestroyEntities(const FMassEntityTemplateID TemplateID, TConstArrayView<FMassEntityHandle> Entities)
+	{
+		DestroyEntities(Entities);
+	}
 
 	const FMassEntityTemplateRegistry& GetTemplateRegistryInstance() const { return TemplateRegistryInstance; }
 	FMassEntityTemplateRegistry& GetMutableTemplateRegistryInstance() { return TemplateRegistryInstance; }
