Added a custom ISM component to handle Mass static entities visualization

diff --git a/Engine/Plugins/Runtime/MassGameplay/Source/MassRepresentation/Private/MassInstancedStaticMeshComponent.cpp b/Engine/Plugins/Runtime/MassGameplay/Source/MassRepresentation/Private/MassInstancedStaticMeshComponent.cpp
new file mode 100644
--- /dev/null
+++ b/Engine/Plugins/Runtime/MassGameplay/Source/MassRepresentation/Private/MassInstancedStaticMeshComponent.cpp
@@ -0,0 +1,251 @@
+// Copyright Epic Games, Inc. All Rights Reserved.
+
+#include "MassInstancedStaticMeshComponent.h"
+#include "MassRepresentationTypes.h"
+#include "AI/NavigationSystemBase.h"
+
+
+UMassInstancedStaticMeshComponent::UMassInstancedStaticMeshComponent()
+{
+	Mobility == EComponentMobility::Stationary;
+}
+
+void UMassInstancedStaticMeshComponent::ApplyVisualChanges(const FMassISMCSharedData& SharedData)
+{
+	ensure(Mobility != EComponentMobility::Movable);
+
+	bool bWorkDone = false;
+	if (SharedData.GetUpdateInstanceIds().Num())
+	{
+		AddInstancesWithIds(SharedData.GetUpdateInstanceIds(), SharedData.GetStaticMeshInstanceTransforms()
+			, NumCustomDataFloats, SharedData.GetStaticMeshInstanceCustomFloats());
+		bWorkDone = true;
+	}
+
+	if (SharedData.GetRemoveInstanceIds().Num())
+	{
+		RemoveInstanceWithIds(SharedData.GetRemoveInstanceIds());
+		bWorkDone = true;
+	}
+
+	// if UHierarchicalInstancedStaticMeshComponent was used as the parent class you'd do the following here:
+	/*if (bWorkDone)
+	{
+		BuildTreeIfOutdated(true, false);
+	}*/
+}
+
+void UMassInstancedStaticMeshComponent::AddInstancesWithIds(TConstArrayView<int32> InstanceIds, TConstArrayView<FTransform> InstanceTransforms, int32 InNumCustomDataFloats
+	, TConstArrayView<float> CustomFloatData, bool bWorldSpace, TArray<int32>* OutAddedIndices)
+{
+	constexpr float EqualTolerance = 1e-6;
+
+	const int32 StartingCount = PerInstanceSMData.Num();
+
+	// if these are the first entities we're adding we need to set NumCustomDataFloats so that the PerInstanceSMCustomData
+	// gets populated properly by the AddInstancesInternal call below
+	if (StartingCount == 0 && CustomFloatData.Num() && Mobility != EComponentMobility::Static)
+	{
+		NumCustomDataFloats = InNumCustomDataFloats;
+	}
+
+	check(InstanceIds.Num() == InstanceTransforms.Num());
+	TArray<int32> NewIndices = AddInstancesInternal(InstanceTransforms, /*bShouldReturnIndices=*/true, bWorldSpace);
+	if (OutAddedIndices)
+	{
+		OutAddedIndices->Append(NewIndices);
+	}
+
+	check(InstanceIds.Num() == NewIndices.Num());
+	PerInstanceIds.AddDefaulted(PerInstanceSMData.Num() - PerInstanceIds.Num());
+
+	for (int32 i = 0; i < InstanceIds.Num(); ++i)
+	{
+		InstanceIdToInstanceIndexMap.Add(InstanceIds[i], NewIndices[i]);
+		PerInstanceIds[NewIndices[i]] = InstanceIds[i];
+	}
+
+	if (CustomFloatData.Num() && Mobility != EComponentMobility::Static)
+	{
+		checkf(NumCustomDataFloats == InNumCustomDataFloats, TEXT("Adding instances with a Custrom Floats count inconsisntent with previously added instances"));
+
+		for (int32 i = 0; i < NewIndices.Num(); ++i)
+		{
+			const int32 InstanceIndex = NewIndices[i];
+			const int32 TargetCustomDataOffset = InstanceIndex * NumCustomDataFloats;
+			const int32 SrcCustomDataOffset = i * NumCustomDataFloats;
+			for (int32 FloatIndex = 0; FloatIndex < NumCustomDataFloats; ++FloatIndex)
+			{
+				// we're making a change only if any of the input data differs of the currently stored ones
+				if (FMath::Abs(CustomFloatData[SrcCustomDataOffset + FloatIndex] - PerInstanceSMCustomData[TargetCustomDataOffset + FloatIndex]) > EqualTolerance)
+				{
+					// Update the component's data in place.
+					FMemory::Memcpy(&PerInstanceSMCustomData[TargetCustomDataOffset], &CustomFloatData[SrcCustomDataOffset], NumCustomDataFloats * sizeof(float));
+
+					// Record in a command buffer for future use.
+					// Using AddInstance here rather than SetCustomData because AddInstancesInternal we used to create 
+					// instances doesn't add commands to InstanceUpdateCmdBuffer
+					InstanceUpdateCmdBuffer.AddInstance(InstanceIds[i], InstanceTransforms[i].ToMatrixWithScale(), FMatrix()
+						, MakeArrayView((const float*)&CustomFloatData[SrcCustomDataOffset], NumCustomDataFloats));
+
+					break;
+				}
+			}
+		}
+	}
+	else
+	{
+		// since AddInstancesInternal called above doesn't add any commands we need to add them here.
+		// The "there are custom floats" path is using a different, command flavor
+		for (const FTransform& InstanceTransform : InstanceTransforms)
+		{
+			InstanceUpdateCmdBuffer.AddInstance(InstanceTransform.ToMatrixWithScale());
+		}
+	}
+}
+
+bool UMassInstancedStaticMeshComponent::RemoveInstanceInternal(const int32 InstanceIndex, const bool bInstanceAlreadyRemoved)
+{
+#if WITH_EDITOR
+	DeletionState = bInstanceAlreadyRemoved ? EInstanceDeletionReason::EntryAlreadyRemoved : EInstanceDeletionReason::EntryRemoval;
+#endif
+
+	const int32 LastInstanceIndex = PerInstanceSMData.Num();
+
+	// remove instance
+	if (!bInstanceAlreadyRemoved && PerInstanceSMData.IsValidIndex(InstanceIndex))
+	{
+		const bool bWasNavRelevant = bNavigationRelevant;
+
+		// Request navigation update
+		PartialNavigationUpdate(InstanceIndex);
+
+		PerInstanceSMData.RemoveAtSwap(InstanceIndex, 1, false);
+		PerInstanceSMCustomData.RemoveAt(InstanceIndex * NumCustomDataFloats, NumCustomDataFloats, false);
+
+		// If it's the last instance, unregister the component since component with no instances are not registered. 
+		// (because of GetInstanceCount() > 0 in UInstancedStaticMeshComponent::IsNavigationRelevant())
+		if (bWasNavRelevant && GetInstanceCount() == 0)
+		{
+			bNavigationRelevant = false;
+			FNavigationSystem::UnregisterComponent(*this);
+		}
+	}
+
+#if WITH_EDITOR
+	// remove selection flag if array is filled in
+	if (SelectedInstances.IsValidIndex(InstanceIndex))
+	{
+		SelectedInstances.RemoveAtSwap(InstanceIndex);
+	}
+#endif
+
+	// update the physics state
+	if (bPhysicsStateCreated && InstanceBodies.IsValidIndex(InstanceIndex))
+	{
+		// Clean up physics for removed instance
+		if (InstanceBodies[InstanceIndex])
+		{
+			InstanceBodies[InstanceIndex]->TermBody();
+			delete InstanceBodies[InstanceIndex];
+		}
+
+		if (InstanceIndex == LastInstanceIndex)
+		{
+			// If we removed the last instance in the array we just need to remove it from the InstanceBodies array too.
+			InstanceBodies.RemoveAt(InstanceIndex);
+		}
+		else
+		{
+			if (InstanceBodies[LastInstanceIndex])
+			{
+				// term physics for swapped instance
+				InstanceBodies[LastInstanceIndex]->TermBody();
+			}
+
+			// swap in the last instance body if we have one
+			InstanceBodies.RemoveAtSwap(InstanceIndex, 1, false);
+
+			// recreate physics for the instance we swapped in the removed item's place
+			if (InstanceBodies[InstanceIndex])
+			{
+				InitInstanceBody(InstanceIndex, InstanceBodies[InstanceIndex]);
+			}
+		}
+	}
+
+	// Notify that these instances have been removed/relocated
+	if (FInstancedStaticMeshDelegates::OnInstanceIndexUpdated.IsBound())
+	{
+		TArray<FInstancedStaticMeshDelegates::FInstanceIndexUpdateData, TInlineAllocator<2>> IndexUpdates;
+		IndexUpdates.Reserve(1 + (PerInstanceSMData.Num() - InstanceIndex));
+
+		IndexUpdates.Add(FInstancedStaticMeshDelegates::FInstanceIndexUpdateData{ FInstancedStaticMeshDelegates::EInstanceIndexUpdateType::Removed, InstanceIndex });
+		if (InstanceIndex != LastInstanceIndex)
+		{
+			// ISMs use swap remove, so the last index has been moved to the spot we removed from
+			IndexUpdates.Add(FInstancedStaticMeshDelegates::FInstanceIndexUpdateData{ FInstancedStaticMeshDelegates::EInstanceIndexUpdateType::Relocated, InstanceIndex, LastInstanceIndex });
+		}
+
+		FInstancedStaticMeshDelegates::OnInstanceIndexUpdated.Broadcast(this, IndexUpdates);
+	}
+
+	// Force recreation of the render data
+	InstanceUpdateCmdBuffer.Edit();
+	MarkRenderStateDirty();
+#if WITH_EDITOR
+	DeletionState = EInstanceDeletionReason::NotDeleting;
+#endif
+	return true;
+}
+
+bool UMassInstancedStaticMeshComponent::RemoveInstanceWithIds(TConstArrayView<int32> InstanceIds)
+{
+	TArray<int32> InstanceIndices;
+	InstanceIndices.Reserve(InstanceIds.Num());
+	// Inform the cmd buffer which instances were removed.
+	for (const int32 InstanceId : InstanceIds)
+	{
+		const int32* InstanceIndexPtr = InstanceIdToInstanceIndexMap.Find(InstanceId);
+		if (InstanceIndexPtr)
+		{
+			if (*InstanceIndexPtr != INDEX_NONE)
+			{
+				InstanceIndices.Add(*InstanceIndexPtr);
+			}
+			InstanceIdToInstanceIndexMap.Remove(InstanceId);
+		}
+	}
+
+	InstanceIndices.Sort(TGreater<int32>());
+
+	if (!PerInstanceSMData.IsValidIndex(InstanceIndices[0]) || !PerInstanceSMData.IsValidIndex(InstanceIndices.Last()))
+	{
+		return false;
+	}
+
+	// update the Id <-> Index mappings
+	for (int32 Index : InstanceIndices)
+	{
+		InstanceUpdateCmdBuffer.HideInstance(Index);
+
+		if (Index == PerInstanceIds.Num() - 1)
+		{
+			PerInstanceIds.RemoveAt(Index, 1, /*bAllowShrinking=*/false);
+		}
+		else
+		{
+			PerInstanceIds.RemoveAtSwap(Index, 1, /*bAllowShrinking=*/false);
+			const int32 NewIdAtIndex = PerInstanceIds[Index];
+			InstanceIdToInstanceIndexMap.FindChecked(NewIdAtIndex) = Index;
+		}
+	}
+	PerInstanceIds.Shrink();
+
+	for (const int32 InstanceIndex : InstanceIndices)
+	{
+		RemoveInstanceInternal(InstanceIndex, /*bInstanceAlreadyRemoved=*/false);
+	}
+
+	return true;
+}
diff --git a/Engine/Plugins/Runtime/MassGameplay/Source/MassRepresentation/Public/MassInstancedStaticMeshComponent.h b/Engine/Plugins/Runtime/MassGameplay/Source/MassRepresentation/Public/MassInstancedStaticMeshComponent.h
new file mode 100644
--- /dev/null
+++ b/Engine/Plugins/Runtime/MassGameplay/Source/MassRepresentation/Public/MassInstancedStaticMeshComponent.h
@@ -0,0 +1,34 @@
+// Copyright Epic Games, Inc. All Rights Reserved.
+
+#pragma once 
+
+#include "Components/InstancedStaticMeshComponent.h"
+#include "MassInstancedStaticMeshComponent.generated.h"
+
+
+struct FMassISMCSharedData;
+
+UCLASS()
+class MASSREPRESENTATION_API UMassInstancedStaticMeshComponent : public UInstancedStaticMeshComponent
+{
+	GENERATED_BODY()
+
+public:
+	UMassInstancedStaticMeshComponent();
+
+	void ApplyVisualChanges(const FMassISMCSharedData& SharedData);
+		
+protected:
+	/**
+	 * Add multiple instances to this component. The newly added instances are associated with Ids as indicated by InstanceIds.
+	 * @param bWorldSpace determines whether input transforms are in world-space (true) or local-space (false)
+	 * @param OutAddedIndices if provided will be appended with indices of newly created instances
+	 */
+	void AddInstancesWithIds(TConstArrayView<int32> InstanceIds, TConstArrayView<FTransform> InstanceTransforms, int32 InNumCustomDataFloats
+		, TConstArrayView<float> CustomFloatData, bool bWorldSpace = false, TArray<int32>* OutAddedIndices = nullptr);
+
+	/** Remove the instances given as instance Ids. Returns True on success. */
+	bool RemoveInstanceWithIds(TConstArrayView<int32> InstanceIds);
+
+	bool RemoveInstanceInternal(const int32 InstanceIndex, const bool bInstanceAlreadyRemoved);
+};
diff --git a/Engine/Source/Runtime/Engine/Public/PrimitiveInstanceUpdateCommand.h b/Engine/Source/Runtime/Engine/Public/PrimitiveInstanceUpdateCommand.h
--- a/Engine/Source/Runtime/Engine/Public/PrimitiveInstanceUpdateCommand.h
+++ b/Engine/Source/Runtime/Engine/Public/PrimitiveInstanceUpdateCommand.h
@@ -40,9 +40,9 @@ struct FInstanceUpdateCmdBuffer
 	ENGINE_API FInstanceUpdateCmdBuffer();
 
 	// Commands that can modify render data in place
-	void HideInstance(int32 RenderIndex);
-	void AddInstance(const FMatrix& InTransform);
-	void AddInstance(int32 InstanceId, const FMatrix& InTransform, const FMatrix& InPreviousTransform, TConstArrayView<float> InCustomDataFloats);
+	void ENGINE_API HideInstance(int32 RenderIndex);
+	void ENGINE_API AddInstance(const FMatrix& InTransform);
+	void ENGINE_API AddInstance(int32 InstanceId, const FMatrix& InTransform, const FMatrix& InPreviousTransform, TConstArrayView<float> InCustomDataFloats);
 	void UpdateInstance(int32 RenderIndex, const FMatrix& InTransform);
 	void UpdateInstance(int32 RenderIndex, const FMatrix& InTransform, const FMatrix& InPreviousTransform);
 	void SetEditorData(int32 RenderIndex, const FColor& Color, bool bSelected);
@@ -52,7 +52,7 @@ struct FInstanceUpdateCmdBuffer
 	int32 NumInlineCommands() const { return Cmds.Num(); }
 
 	// Command that can't be in-lined and should cause full buffer rebuild
-	void Edit();
+	void ENGINE_API Edit();
 	void Reset();
 	int32 NumTotalCommands() const { return NumEdits; };
 
