Mass CIS fix

diff --git a/Engine/Plugins/Runtime/MassGameplay/Source/MassRepresentation/Public/MassRepresentationTypes.h b/Engine/Plugins/Runtime/MassGameplay/Source/MassRepresentation/Public/MassRepresentationTypes.h
--- a/Engine/Plugins/Runtime/MassGameplay/Source/MassRepresentation/Public/MassRepresentationTypes.h
+++ b/Engine/Plugins/Runtime/MassGameplay/Source/MassRepresentation/Public/MassRepresentationTypes.h
@@ -141,6 +141,12 @@ struct MASSREPRESENTATION_API FMassISMCSharedData
 
 	}
 
+	TConstArrayView<int32> GetUpdateInstanceIds() const { return UpdateInstanceIds; }
+	TConstArrayView<FTransform> GetStaticMeshInstanceTransforms() const { return StaticMeshInstanceTransforms; }
+	TConstArrayView<FTransform> GetStaticMeshInstancePrevTransforms() const { return StaticMeshInstancePrevTransforms; }
+	TConstArrayView<int32> GetRemoveInstanceIds() const { return RemoveInstanceIds; }
+	TConstArrayView<float> GetStaticMeshInstanceCustomFloats() const { return StaticMeshInstanceCustomFloats; }
+
 	UInstancedStaticMeshComponent* ISMC = nullptr;
 
 	int32 RefCount = 1;
@@ -149,6 +155,7 @@ struct MASSREPRESENTATION_API FMassISMCSharedData
 	TArray<int32> UpdateInstanceIds;
 	TArray<FTransform> StaticMeshInstanceTransforms;
 	TArray<FTransform> StaticMeshInstancePrevTransforms;
+	TArray<int32> RemoveInstanceIds;
 
 	/** Buffer holding current frame custom floats for the static mesh instances, used to batch update the ISMs custom data */
 	TArray<float> StaticMeshInstanceCustomFloats;
