CIS fix

diff --git a/Engine/Plugins/Runtime/MassGameplay/Source/MassRepresentation/Private/MassInstancedStaticMeshComponent.cpp b/Engine/Plugins/Runtime/MassGameplay/Source/MassRepresentation/Private/MassInstancedStaticMeshComponent.cpp
--- a/Engine/Plugins/Runtime/MassGameplay/Source/MassRepresentation/Private/MassInstancedStaticMeshComponent.cpp
+++ b/Engine/Plugins/Runtime/MassGameplay/Source/MassRepresentation/Private/MassInstancedStaticMeshComponent.cpp
@@ -7,7 +7,7 @@
 
 UMassInstancedStaticMeshComponent::UMassInstancedStaticMeshComponent()
 {
-	Mobility == EComponentMobility::Stationary;
+	Mobility = EComponentMobility::Stationary;
 }
 
 void UMassInstancedStaticMeshComponent::ApplyVisualChanges(const FMassISMCSharedData& SharedData)
