UMassVisualizationComponent changes required for Mass stationary entities visualization.

* Chaged UMassVisualizationComponent::FindOrAddVisualDesc to configure newly created FMassInstancedStaticMeshInfo instances so that they can be used right after the function call.

* Changed when UMassVisualizationComponent cleans up its accumulated runtime data - used to reset on first Mass processing phase start. Now it cleans up after consuming the data. This lets us accumulate data outside of regular Mass processing (like during spawning).

* Improved UMassVisualizationComponent handling of UHierarchicalInstancedStaticMeshComponent by calling UHierarchicalInstancedStaticMeshComponent::BuildTreeIfOutdated after adding/removing instances to a given component.

diff --git a/Engine/Plugins/Runtime/MassGameplay/Source/MassRepresentation/Private/MassRepresentationTypes.cpp b/Engine/Plugins/Runtime/MassGameplay/Source/MassRepresentation/Private/MassRepresentationTypes.cpp
--- a/Engine/Plugins/Runtime/MassGameplay/Source/MassRepresentation/Private/MassRepresentationTypes.cpp
+++ b/Engine/Plugins/Runtime/MassGameplay/Source/MassRepresentation/Private/MassRepresentationTypes.cpp
@@ -3,6 +3,7 @@
 #include "MassRepresentationTypes.h"
 #include "MassRepresentationUtils.h"
 #include "MassCommandBuffer.h"
+#include "Components/InstancedStaticMeshComponent.h"
 
 DEFINE_LOG_CATEGORY(LogMassRepresentation);
 
@@ -72,4 +73,13 @@ default: \
 
 #undef CASE_SWAP_TAGS
 }
-} // UE::MassRepresentation
\ No newline at end of file
+} // UE::MassRepresentation
+
+
+//-----------------------------------------------------------------------------
+// FMassStaticMeshInstanceVisualizationMeshDesc
+//-----------------------------------------------------------------------------
+FMassStaticMeshInstanceVisualizationMeshDesc::FMassStaticMeshInstanceVisualizationMeshDesc()
+{
+	ISMComponentClass = UInstancedStaticMeshComponent::StaticClass();
+}
diff --git a/Engine/Plugins/Runtime/MassGameplay/Source/MassRepresentation/Private/MassVisualizationComponent.cpp b/Engine/Plugins/Runtime/MassGameplay/Source/MassRepresentation/Private/MassVisualizationComponent.cpp
--- a/Engine/Plugins/Runtime/MassGameplay/Source/MassRepresentation/Private/MassVisualizationComponent.cpp
+++ b/Engine/Plugins/Runtime/MassGameplay/Source/MassRepresentation/Private/MassVisualizationComponent.cpp
@@ -10,6 +10,8 @@
 #include "RenderUtils.h"
 #include "SceneInterface.h"
 
+#include "Components/HierarchicalInstancedStaticMeshComponent.h"
+#include "MassInstancedStaticMeshComponent.h"
 
 //---------------------------------------------------------------
 // UMassVisualizationComponent
@@ -37,6 +39,14 @@ int16 UMassVisualizationComponent::FindOrAddVisualDesc(const FStaticMeshInstance
 	if (VisualIndex == INDEX_NONE)
 	{
 		VisualIndex = InstancedStaticMeshInfos.Emplace(Desc);
+
+		for (const FMassStaticMeshInstanceVisualizationMeshDesc& MeshDesc : Desc.Meshes)
+		{
+			ISMCSharedData.FindOrAdd(MeshDesc, FMassISMCSharedData());
+		}
+
+		BuildLODSignificanceForInfo(InstancedStaticMeshInfos[VisualIndex]);
+
 		bNeedStaticMeshComponentConstruction = true;
 	}
 	check(VisualIndex < INT16_MAX);
@@ -64,16 +74,14 @@ void UMassVisualizationComponent::ConstructStaticMeshComponents()
 			continue;
 		}
 		for (const FMassStaticMeshInstanceVisualizationMeshDesc& MeshDesc : Info.Desc.Meshes)
-		{			
+		{
 			FMassISMCSharedData* SharedData = ISMCSharedData.Find(MeshDesc);
-			UInstancedStaticMeshComponent* ISMC = SharedData ? SharedData->ISMC : nullptr;
-			if (SharedData)
-			{
-				SharedData->RefCount += 1;
-			}
-			else
+			UInstancedStaticMeshComponent* ISMC = SharedData ? SharedData->GetISMComponent() : nullptr;
+
+			if (SharedData == nullptr || SharedData->GetISMComponent() == nullptr)
 			{
-				ISMC = NewObject<UInstancedStaticMeshComponent>(ActorOwner);
+				ISMC = NewObject<UInstancedStaticMeshComponent>(ActorOwner, MeshDesc.ISMComponentClass);
+				CA_ASSUME(ISMC);
 				ISMC->SetStaticMesh(MeshDesc.Mesh);
 				for (int32 ElementIndex = 0; ElementIndex < MeshDesc.MaterialOverrides.Num(); ++ElementIndex)
 				{
@@ -91,36 +99,56 @@ void UMassVisualizationComponent::ConstructStaticMeshComponents()
 				ISMC->SetReceivesDecals(false);
 				ISMC->RegisterComponent();
 
-				ISMCSharedData.Emplace(MeshDesc, FMassISMCSharedData(ISMC));
+				if (SharedData == nullptr)
+				{
+					SharedData = &ISMCSharedData.Emplace(MeshDesc, FMassISMCSharedData(ISMC));
+				}
+				else
+				{
+					SharedData->SetISMComponent(*ISMC);
+				}
 			}
 
-			Info.InstancedStaticMeshComponents.Add(ISMC);
+			check(SharedData);
+			Info.AddISMComponent(*SharedData);
 		}
 
 		// Build the LOD significance ranges
-		TArray<float> AllLODSignificances;
-		auto UniqueInsertOrdered = [&AllLODSignificances](const float Significance)
+		if (Info.LODSignificanceRanges.Num() == 0)
+		{
+			BuildLODSignificanceForInfo(Info);
+		}
+	}
+}
+
+void UMassVisualizationComponent::BuildLODSignificanceForInfo(FMassInstancedStaticMeshInfo& Info)
+{
+	TArray<float> AllLODSignificances;
+	auto UniqueInsertOrdered = [&AllLODSignificances](const float Significance)
+	{
+		int i = 0;
+		for (; i < AllLODSignificances.Num(); ++i)
 		{
-			int i = 0;
-			for (; i < AllLODSignificances.Num(); ++i)
+			// I did not use epsilon check here on purpose, because it will make it hard later meshes inside.
+			if (Significance == AllLODSignificances[i])
 			{
-				// I did not use epsilon check here on purpose, because it will make it hard later meshes inside.
-				if (Significance == AllLODSignificances[i])
-				{
-					return;
-				}
-				if (AllLODSignificances[i] > Significance)
-				{
-					break;
-				}
+				return;
+			}
+			if (AllLODSignificances[i] > Significance)
+			{
+				break;
 			}
-			AllLODSignificances.Insert(Significance, i);
-		};
-		for (const FMassStaticMeshInstanceVisualizationMeshDesc& MeshDesc : Info.Desc.Meshes)
-		{
-			UniqueInsertOrdered(MeshDesc.MinLODSignificance);
-			UniqueInsertOrdered(MeshDesc.MaxLODSignificance);
 		}
+		AllLODSignificances.Insert(Significance, i);
+	};
+	for (const FMassStaticMeshInstanceVisualizationMeshDesc& MeshDesc : Info.Desc.Meshes)
+	{
+		UniqueInsertOrdered(MeshDesc.MinLODSignificance);
+		UniqueInsertOrdered(MeshDesc.MaxLODSignificance);
+	}
+
+	if (AllLODSignificances.Num() > 1)
+	{
 		Info.LODSignificanceRanges.SetNum(AllLODSignificances.Num() - 1);
 		for (int i = 0; i < Info.LODSignificanceRanges.Num(); ++i)
 		{
@@ -139,7 +167,6 @@ void UMassVisualizationComponent::ConstructStaticMeshComponents()
 				}
 			}
 		}
-
 	}
 }
 
@@ -155,7 +182,7 @@ void UMassVisualizationComponent::ClearAllVisualInstances()
 	// Pool should already be empty, got a problem if it's not
 	for (auto It = ISMCSharedData.CreateIterator(); It; ++It)
 	{
-		if (UInstancedStaticMeshComponent* InstancedStaticMeshComponent = It.Value().ISMC)
+		if (UInstancedStaticMeshComponent* InstancedStaticMeshComponent = It.Value().GetISMComponent())
 		{
 			InstancedStaticMeshComponent->ClearInstances();
 			InstancedStaticMeshComponent->DestroyComponent();
@@ -187,17 +214,6 @@ void UMassVisualizationComponent::BeginVisualChanges()
 		ConstructStaticMeshComponents();
 		bNeedStaticMeshComponentConstruction = false;
 	}
-
-	// Reset instance transform scratch buffers
-	for (auto It = ISMCSharedData.CreateIterator(); It; ++It)
-	{
-		FMassISMCSharedData& SharedData = It.Value();
-		SharedData.UpdateInstanceIds.Reset();
-		SharedData.StaticMeshInstanceCustomFloats.Reset();
-		SharedData.StaticMeshInstanceTransforms.Reset();
-		SharedData.StaticMeshInstancePrevTransforms.Reset();
-		SharedData.WriteIterator = 0;
-	}
 }
 
 void UMassVisualizationComponent::EndVisualChanges()
@@ -208,8 +224,14 @@ void UMassVisualizationComponent::EndVisualChanges()
 	for (auto It = ISMCSharedData.CreateIterator(); It; ++It)
 	{
 		FMassISMCSharedData& SharedData = It.Value();
-		if (UInstancedStaticMeshComponent* InstancedStaticMeshComponent = SharedData.ISMC)
+		if (UInstancedStaticMeshComponent* InstancedStaticMeshComponent = SharedData.GetISMComponent())
 		{
+			if (UMassInstancedStaticMeshComponent* MassISMComponent = Cast<UMassInstancedStaticMeshComponent>(InstancedStaticMeshComponent))
+			{
+				MassISMComponent->ApplyVisualChanges(SharedData);
+				continue;
+			}
+
 			const int32 NumCustomDataFloats = SharedData.StaticMeshInstanceCustomFloats.Num() / (FMath::Max(1, SharedData.UpdateInstanceIds.Num()));
 
 			// Ensure InstanceCustomData is passed if NumCustomDataFloats > 0. If it is, also make sure
@@ -223,6 +245,10 @@ void UMassVisualizationComponent::EndVisualChanges()
 			if ((bool)UE::MassRepresentation::GCallUpdateInstances)
 			{
 				InstancedStaticMeshComponent->UpdateInstances(SharedData.UpdateInstanceIds, SharedData.StaticMeshInstanceTransforms, SharedData.StaticMeshInstancePrevTransforms, NumCustomDataFloats, SharedData.StaticMeshInstanceCustomFloats);
+				if (UHierarchicalInstancedStaticMeshComponent* HISMComp = Cast<UHierarchicalInstancedStaticMeshComponent>(InstancedStaticMeshComponent))
+				{
+					HISMComp->BuildTreeIfOutdated(true, false);
+				}
 			}
 			else
 			{
@@ -265,6 +291,13 @@ void UMassVisualizationComponent::EndVisualChanges()
 			}
 		}
 	}
+
+	// Reset instance transform scratch buffers
+	for (auto It = ISMCSharedData.CreateIterator(); It; ++It)
+	{
+		FMassISMCSharedData& SharedData = It.Value();
+		SharedData.ResetAccumulatedData();
+	}
 }
 
 //---------------------------------------------------------------
@@ -276,15 +309,14 @@ void FMassInstancedStaticMeshInfo::ClearVisualInstance(FMassISMCSharedDataMap& I
 	for (int i = 0; i < Desc.Meshes.Num(); i++)
 	{
 		FMassISMCSharedData* SharedData = ISMCSharedData.Find(Desc.Meshes[i]);
-		if (SharedData)
+		if (SharedData && SharedData->ReleaseReference() == 0)
 		{
-			SharedData->RefCount -= 1;
-			if (SharedData->RefCount == 0)
+			if (UInstancedStaticMeshComponent* ISMC = SharedData->GetISMComponent())
 			{
-				SharedData->ISMC->ClearInstances();
-				SharedData->ISMC->DestroyComponent();
-				ISMCSharedData.Remove(Desc.Meshes[i]);
+				ISMC->ClearInstances();
+				ISMC->DestroyComponent();
 			}
+			ISMCSharedData.Remove(Desc.Meshes[i]);
 		}
 	}
 
@@ -329,6 +361,16 @@ void FMassLODSignificanceRange::AddBatchedCustomDataFloats(const TArray<float>&
 	}
 }
 
+void FMassLODSignificanceRange::RemoveInstance(const int32 InstanceId)
+{
+	check(ISMCSharedDataPtr);
+	for (int i = 0; i < StaticMeshRefs.Num(); i++)
+	{
+		FMassISMCSharedData& SharedData = (*ISMCSharedDataPtr)[StaticMeshRefs[i]];
+		SharedData.RemoveInstanceIds.Add(InstanceId);
+	}
+}
+
 void FMassLODSignificanceRange::WriteCustomDataFloatsAtStartIndex(int32 StaticMeshIndex, const TArrayView<float>& CustomFloats, const int32 FloatsPerInstance, const int32 StartFloatIndex, const TArray<uint32>& ExcludeStaticMeshRefs)
 {
 	check(ISMCSharedDataPtr);
diff --git a/Engine/Plugins/Runtime/MassGameplay/Source/MassRepresentation/Public/MassRepresentationTypes.h b/Engine/Plugins/Runtime/MassGameplay/Source/MassRepresentation/Public/MassRepresentationTypes.h
--- a/Engine/Plugins/Runtime/MassGameplay/Source/MassRepresentation/Public/MassRepresentationTypes.h
+++ b/Engine/Plugins/Runtime/MassGameplay/Source/MassRepresentation/Public/MassRepresentationTypes.h
@@ -10,6 +10,9 @@
 
 class UMaterialInterface;
 class UStaticMesh;
+struct FMassLODSignificanceRange;
+class UMassVisualizationComponent;
+class UInstancedStaticMeshComponent;
 
 DECLARE_LOG_CATEGORY_EXTERN(LogMassRepresentation, Log, All);
 
@@ -35,9 +38,11 @@ enum class EMassActorEnabledType : uint8
 };
 
 USTRUCT()
-struct FMassStaticMeshInstanceVisualizationMeshDesc
+struct MASSREPRESENTATION_API FMassStaticMeshInstanceVisualizationMeshDesc
 {
 	GENERATED_BODY()
+
+	FMassStaticMeshInstanceVisualizationMeshDesc();
 	
 	/** The static mesh visual representation */
 	UPROPERTY(EditAnywhere, Category = "Mass|Visual")
@@ -64,8 +69,11 @@ struct FMassStaticMeshInstanceVisualizationMeshDesc
 	bool bCastShadows = false;
 	
 	/** Controls the mobility of the ISM */
+	EComponentMobility::Type Mobility = EComponentMobility::Movable;
+
+	/** InstancedStaticMeshComponent class to use to manage instances described by this struct instance */
 	UPROPERTY(EditAnywhere, Category = "Mass|Visual")
-	TEnumAsByte<EComponentMobility::Type> Mobility = EComponentMobility::Movable;
+	TSubclassOf<UInstancedStaticMeshComponent> ISMComponentClass;
 
 	bool operator==(const FMassStaticMeshInstanceVisualizationMeshDesc& Other) const
 	{
@@ -127,18 +135,34 @@ struct FStaticMeshInstanceVisualizationDesc : public FTableRowBase
 
 class UInstancedStaticMeshComponent;
 
-USTRUCT()
+
 struct MASSREPRESENTATION_API FMassISMCSharedData
 {
-	GENERATED_BODY()
-
 	FMassISMCSharedData() = default;
 
 	FMassISMCSharedData(UInstancedStaticMeshComponent* InISMC)
-		: ISMC(InISMC),
-		RefCount(1)
+		: ISMC(InISMC)
 	{
+	}
 
+	void SetISMComponent(UInstancedStaticMeshComponent& InISMC)
+	{
+		check(ISMC == nullptr && RefCount == 0);
+		ISMC = &InISMC;
+	}
+
+	UInstancedStaticMeshComponent* GetISMComponent() { return ISMC; }
+	int32 StoreReference() { return ++RefCount; }
+	int32 ReleaseReference() { ensure(RefCount >= 0); return --RefCount; }
+
+	void ResetAccumulatedData()
+	{
+		UpdateInstanceIds.Reset();
+		StaticMeshInstanceCustomFloats.Reset();
+		StaticMeshInstanceTransforms.Reset();
+		StaticMeshInstancePrevTransforms.Reset();
+		RemoveInstanceIds.Reset();
+		WriteIterator = 0;
 	}
 
 	TConstArrayView<int32> GetUpdateInstanceIds() const { return UpdateInstanceIds; }
@@ -147,10 +171,9 @@ struct MASSREPRESENTATION_API FMassISMCSharedData
 	TConstArrayView<int32> GetRemoveInstanceIds() const { return RemoveInstanceIds; }
 	TConstArrayView<float> GetStaticMeshInstanceCustomFloats() const { return StaticMeshInstanceCustomFloats; }
 
-	UInstancedStaticMeshComponent* ISMC = nullptr;
-
-	int32 RefCount = 1;
-
+protected:
+	friend FMassLODSignificanceRange;
+	friend UMassVisualizationComponent;
 	/** Buffer holding current frame transforms for the static mesh instances, used to batch update the transforms */
 	TArray<int32> UpdateInstanceIds;
 	TArray<FTransform> StaticMeshInstanceTransforms;
@@ -162,18 +185,13 @@ struct MASSREPRESENTATION_API FMassISMCSharedData
 
 	// When initially adding to StaticMeshInstanceCustomFloats, can use the size as the write iterator, but on subsequent processors, we need to know where to start writing
 	int32 WriteIterator = 0;
+
+	UInstancedStaticMeshComponent* ISMC = nullptr;
+	int32 RefCount = 0;
 };
 
 using FMassISMCSharedDataMap = TMap<uint32, FMassISMCSharedData>;
 
-//struct FISMCSharedDataMap : public TMap<uint32, FMassISMCSharedData>
-//{
-//	FISMCSharedDataMap()
-//	{
-//		ensureMsgf(false, "FISMCSharedDataMap has been renamed to FMassISMCSharedDataMap, use that instead"));
-//	}
-//};
-
 
 USTRUCT()
 struct MASSREPRESENTATION_API FMassLODSignificanceRange
@@ -208,6 +226,8 @@ public:
 
 	void AddBatchedCustomDataFloats(const TArray<float>& CustomFloats, const TArray<uint32>& ExcludeStaticMeshRefs);
 
+	void RemoveInstance(const int32 InstanceId);
+
 	void WriteCustomDataFloatsAtStartIndex(int32 StaticMeshIndex, const TArrayView<float>& CustomFloats, const int32 FloatsPerInstance, const int32 StartIndex, const TArray<uint32>& ExcludeStaticMeshRefs);
 
 	/** LOD Significance range */
@@ -215,7 +235,7 @@ public:
 	float MaxSignificance;
 
 	/** The component handling these instances */
-	UPROPERTY()
+	UPROPERTY(VisibleAnywhere, Category = "Mass/Debug")
 	TArray<uint32> StaticMeshRefs;
 
 	FMassISMCSharedDataMap* ISMCSharedDataPtr = nullptr;
@@ -276,6 +296,14 @@ public:
 		}
 	}
 
+	FORCEINLINE void RemoveInstance(const int32 InstanceId, const float LODSignificance)
+	{
+		if (FMassLODSignificanceRange* Range = GetLODSignificanceRange(LODSignificance))
+		{
+			Range->RemoveInstance(InstanceId);
+		}
+	}
+
 	// Adds the specified struct reinterpreted as custom floats to our custom data. Individual members of the specified struct should always fit into a float.
 	// When adding any custom data, the custom data must be added for every instance.
 	template<typename InCustomDataType>
@@ -327,21 +355,31 @@ public:
 		}
 	}
 
+	void AddISMComponent(FMassISMCSharedData& SharedData)
+	{
+		if (ensure(SharedData.GetISMComponent()))
+		{
+			InstancedStaticMeshComponents.Add(SharedData.GetISMComponent());
+			SharedData.StoreReference();
+		}
+	}
+
 protected:
 
 	/** Destroy the visual instance */
 	void ClearVisualInstance(FMassISMCSharedDataMap& ISMCSharedData);
 
 	/** Information about this static mesh which will represent all instances */
-	UPROPERTY()
-		FStaticMeshInstanceVisualizationDesc Desc;
+	UPROPERTY(VisibleAnywhere, Category = "Mass/Debug")
+	FStaticMeshInstanceVisualizationDesc Desc;
 
 	/** The component handling these instances */
-	UPROPERTY()
-		TArray<TObjectPtr<UInstancedStaticMeshComponent>> InstancedStaticMeshComponents;
+	UPROPERTY(VisibleAnywhere, Category = "Mass/Debug")
+	TArray<TObjectPtr<UInstancedStaticMeshComponent>> InstancedStaticMeshComponents;
+
+	UPROPERTY(VisibleAnywhere, Category = "Mass/Debug")
+	TArray<FMassLODSignificanceRange> LODSignificanceRanges;
 
-	UPROPERTY()
-		TArray<FMassLODSignificanceRange> LODSignificanceRanges;
 	friend class UMassVisualizationComponent;
 };
 
diff --git a/Engine/Plugins/Runtime/MassGameplay/Source/MassRepresentation/Public/MassVisualizationComponent.h b/Engine/Plugins/Runtime/MassGameplay/Source/MassRepresentation/Public/MassVisualizationComponent.h
--- a/Engine/Plugins/Runtime/MassGameplay/Source/MassRepresentation/Public/MassVisualizationComponent.h
+++ b/Engine/Plugins/Runtime/MassGameplay/Source/MassRepresentation/Public/MassVisualizationComponent.h
@@ -54,6 +54,8 @@ protected:
 	/** Overridden to make sure this component is only added to a MassVisualizer actor */
 	virtual void PostInitProperties() override;
 
+	void BuildLODSignificanceForInfo(FMassInstancedStaticMeshInfo& Info);
+
 	/** The information of all the instanced static meshes */
 	UPROPERTY(Transient)
 	TArray<FMassInstancedStaticMeshInfo> InstancedStaticMeshInfos;
