FGuid-powered MassEntityTemplateID. In principle it's stable, i.e. independent of memory pointers. Might still contain kinks, to be found and annihilated.

Also:
* removed 32-bit hash support from MassEntityTemplateID in favor of 64-bit hash.
* moved the "get combined traits" step out of GetEntityTemplateInternal and into a separate function.

diff --git a/Engine/Plugins/Runtime/MassGameplay/Source/MassGameplayTestSuite/Private/MassEntityTemplateTest.cpp b/Engine/Plugins/Runtime/MassGameplay/Source/MassGameplayTestSuite/Private/MassEntityTemplateTest.cpp
--- a/Engine/Plugins/Runtime/MassGameplay/Source/MassGameplayTestSuite/Private/MassEntityTemplateTest.cpp
+++ b/Engine/Plugins/Runtime/MassGameplay/Source/MassGameplayTestSuite/Private/MassEntityTemplateTest.cpp
@@ -33,12 +33,14 @@ struct FEntityTemplate_Empty : FEntityTemplateBase
 {
 	virtual bool InstantTest() override
 	{
+		const FMassEntityTemplateID InitialTemplateID = FMassEntityTemplateIDFactory::Make(FGuid::NewGuid());
 		FMassEntityTemplateData TemplateData;
 
-		const TSharedRef<FMassEntityTemplate>& FinalizedTemplate = TemplateRegistry.FindOrAddTemplate(MoveTemp(TemplateData));
+		const TSharedRef<FMassEntityTemplate>& FinalizedTemplate = TemplateRegistry.FindOrAddTemplate(InitialTemplateID, MoveTemp(TemplateData));
 		
 		const FMassEntityTemplateID TemplateID = FinalizedTemplate->GetTemplateID();
 		AITEST_TRUE("Empty template is expected to be registered as a valid template", TemplateID.IsValid());
+		AITEST_EQUAL("Empty template's ID is expected to be the same as the input template", TemplateID, InitialTemplateID);
 
 		const FMassArchetypeHandle& ArchetypeHandle = FinalizedTemplate->GetArchetype();
 		AITEST_TRUE("Empty template is expected to produce a valid (if empty) archetype", ArchetypeHandle.IsValid());
@@ -81,7 +83,8 @@ struct FEntityTemplate_Composition : FEntityTemplateBase
 		AITEST_TRUE("The composition should end up being the same regardless of whether the data is added via a template data or a composition descriptor"
 			, ExpectedComposition.IsEquivalent(TemplateData.GetCompositionDescriptor()));
 
-		const TSharedRef<FMassEntityTemplate>& FinalizedTemplate = TemplateRegistry.FindOrAddTemplate(MoveTemp(TemplateData));
+		const FMassEntityTemplateID TemplateID = FMassEntityTemplateIDFactory::Make(FGuid::NewGuid());
+		const TSharedRef<FMassEntityTemplate>& FinalizedTemplate = TemplateRegistry.FindOrAddTemplate(TemplateID, MoveTemp(TemplateData));
 
 		const FMassArchetypeHandle& ArchetypeHandle = FinalizedTemplate->GetArchetype();
 		AITEST_TRUE("Empty template is expected to produce a valid (if empty) archetype", ArchetypeHandle.IsValid());
@@ -106,9 +109,9 @@ struct FEntityTemplate_Trivial : FEntityTemplateBase
 		TemplateDataA.AddFragment<FTestFragment_Int>();
 		TemplateDataB.AddFragment<FTestFragment_Int>();
 
-		const TSharedRef<FMassEntityTemplate>& FinalizedEmptyTemplate = TemplateRegistry.FindOrAddTemplate(MoveTemp(EmptyTemplateData));
-		const TSharedRef<FMassEntityTemplate>& FinalizedTemplateA = TemplateRegistry.FindOrAddTemplate(MoveTemp(TemplateDataA));
-		const TSharedRef<FMassEntityTemplate>& FinalizedTemplateB = TemplateRegistry.FindOrAddTemplate(MoveTemp(TemplateDataA));
+		const TSharedRef<FMassEntityTemplate>& FinalizedEmptyTemplate = TemplateRegistry.FindOrAddTemplate(FMassEntityTemplateIDFactory::Make(FGuid::NewGuid()), MoveTemp(EmptyTemplateData));
+		const TSharedRef<FMassEntityTemplate>& FinalizedTemplateA = TemplateRegistry.FindOrAddTemplate(FMassEntityTemplateIDFactory::Make(FGuid::NewGuid()), MoveTemp(TemplateDataA));
+		const TSharedRef<FMassEntityTemplate>& FinalizedTemplateB = TemplateRegistry.FindOrAddTemplate(FMassEntityTemplateIDFactory::Make(FGuid::NewGuid()), MoveTemp(TemplateDataA));
 
 		AITEST_NOT_EQUAL("Non-empty template data should result in a finalized template different from the empty one", FinalizedEmptyTemplate->GetTemplateID(), FinalizedTemplateA->GetTemplateID());
 		AITEST_EQUAL("Non-empty template data should result in the very same finalized template", FinalizedTemplateA, FinalizedTemplateB);
@@ -134,7 +137,7 @@ struct FEntityTemplate_Modified : FEntityTemplateBase
 		IntFragment.Value = FTestFragment_Int::TestIntValue;
 
 		FMassEntityTemplateData MovedOriginalTemplateData = OriginalTemplateData;
-		const TSharedRef<FMassEntityTemplate>& FinalizedOriginalTemplate = TemplateRegistry.FindOrAddTemplate(MoveTemp(MovedOriginalTemplateData));
+		const TSharedRef<FMassEntityTemplate>& FinalizedOriginalTemplate = TemplateRegistry.FindOrAddTemplate(FMassEntityTemplateIDFactory::Make(FGuid::NewGuid()), MoveTemp(MovedOriginalTemplateData));
 
 		FMassEntityTemplateData NewTemplateData(*FinalizedOriginalTemplate);
 
@@ -150,7 +153,8 @@ struct FEntityTemplate_Modified : FEntityTemplateBase
 		CoppiedFragment->Value = FTestFragment_Int::TestIntValue + 1;
 		AITEST_NOT_EQUAL("Modifying the coppied instance of the fragment doesn't affect the original", CoppiedFragment->Value, IntFragment.Value);
 
-		const TSharedRef<FMassEntityTemplate>& FinalizedModifiedTemplate = TemplateRegistry.FindOrAddTemplate(MoveTemp(NewTemplateData));
+		const FMassEntityTemplateID NewTemplateID = FMassEntityTemplateIDFactory::MakeFlavor(FinalizedOriginalTemplate->GetTemplateID(), 1);
+		const TSharedRef<FMassEntityTemplate>& FinalizedModifiedTemplate = TemplateRegistry.FindOrAddTemplate(NewTemplateID, MoveTemp(NewTemplateData));
 
 		AITEST_NOT_EQUAL("The original and modified templates should end up resulting in two different templates", FinalizedModifiedTemplate->GetTemplateID(), FinalizedOriginalTemplate->GetTemplateID());
 		AITEST_EQUAL("The original and modified templates should still point at the same archetype", FinalizedModifiedTemplate->GetArchetype(), FinalizedOriginalTemplate->GetArchetype());
diff --git a/Engine/Plugins/Runtime/MassGameplay/Source/MassSpawner/Private/MassEntityConfigAsset.cpp b/Engine/Plugins/Runtime/MassGameplay/Source/MassSpawner/Private/MassEntityConfigAsset.cpp
--- a/Engine/Plugins/Runtime/MassGameplay/Source/MassSpawner/Private/MassEntityConfigAsset.cpp
+++ b/Engine/Plugins/Runtime/MassGameplay/Source/MassSpawner/Private/MassEntityConfigAsset.cpp
@@ -18,10 +18,13 @@
 #define LOCTEXT_NAMESPACE "Mass"
 
 
+//-----------------------------------------------------------------------------
+// FMassEntityConfig
+//-----------------------------------------------------------------------------
 FMassEntityConfig::FMassEntityConfig(UObject& InOwner)
 	: ConfigOwner(&InOwner)
 {
-
+	ConfigGuid = FGuid::NewGuid();
 }
 
 const UMassEntityTraitBase* FMassEntityConfig::FindTrait(TSubclassOf<UMassEntityTraitBase> TraitClass, const bool bExactMatch) const
@@ -37,37 +40,10 @@
 	return Parent ? Parent->FindTrait(TraitClass, bExactMatch) :  nullptr;
 }
 
-#if WITH_EDITOR
-void UMassEntityConfigAsset::ValidateEntityConfig()
-{
-	if (UWorld* EditorWorld = GEditor->GetEditorWorldContext().World())
-	{
-		if (Config.ValidateEntityTemplate(*EditorWorld))
-		{
-			const FText InfoText = LOCTEXT("MassEntityConfigAssetNoErrorsDetected", "There were no errors nor warnings detected during validation of the EntityConfigAsset");
-			
-			FMessageLog EditorInfo("MassEntity");
-			EditorInfo.Info(InfoText);
-
-			FNotificationInfo Info(InfoText);
-			Info.bFireAndForget = true;
-			Info.bUseThrobber = false;
-			Info.FadeOutDuration = 0.5f;
-			Info.ExpireDuration = 5.0f;
-			if (TSharedPtr<SNotificationItem> Notification = FSlateNotificationManager::Get().AddNotification(Info))
-			{
-				Notification->SetCompletionState(SNotificationItem::CS_Success);
-			}
-		}
-	}
-}
-#endif // WITH_EDITOR
-
 const FMassEntityTemplate& FMassEntityConfig::GetOrCreateEntityTemplate(const UWorld& World) const
 {
 	FMassEntityTemplateID TemplateID;
-	TArray<UMassEntityTraitBase*> CombinedTraits;
-	if (const FMassEntityTemplate* ExistingTemplate = GetEntityTemplateInternal(World, TemplateID, CombinedTraits))
+	if (const FMassEntityTemplate* ExistingTemplate = GetEntityTemplateInternal(World, TemplateID))
 	{
 		return *ExistingTemplate;
 	}
@@ -84,6 +60,9 @@
 	FMassEntityTemplateData TemplateData;
 	FMassEntityTemplateBuildContext BuildContext(TemplateData, TemplateID);
 
+	TArray<UMassEntityTraitBase*> CombinedTraits;
+	GetCombinedTraits(CombinedTraits);
+
 	BuildContext.BuildFromTraits(CombinedTraits, World);
 	BuildContext.SetTemplateName(GetNameSafe(ConfigOwner));
 
@@ -93,8 +72,7 @@
 void FMassEntityConfig::DestroyEntityTemplate(const UWorld& World) const
 {
 	FMassEntityTemplateID TemplateID;
-	TArray<UMassEntityTraitBase*> CombinedTraits;
-	const FMassEntityTemplate* Template = GetEntityTemplateInternal(World, TemplateID, CombinedTraits);
+	const FMassEntityTemplate* Template = GetEntityTemplateInternal(World, TemplateID);
 	if (Template == nullptr)
 	{
 		return;
@@ -104,6 +82,9 @@
 	check(SpawnerSystem);
 	FMassEntityTemplateRegistry& TemplateRegistry = SpawnerSystem->GetMutableTemplateRegistryInstance();
 
+	TArray<UMassEntityTraitBase*> CombinedTraits;
+	GetCombinedTraits(CombinedTraits);
+
 	for (const UMassEntityTraitBase* Trait : CombinedTraits)
 	{
 		check(Trait);
@@ -118,33 +99,32 @@
 const FMassEntityTemplate& FMassEntityConfig::GetEntityTemplateChecked(const UWorld& World) const
 {
 	FMassEntityTemplateID TemplateID;
-	TArray<UMassEntityTraitBase*> CombinedTraits;
-	const FMassEntityTemplate* ExistingTemplate = GetEntityTemplateInternal(World, TemplateID, CombinedTraits);
+	const FMassEntityTemplate* ExistingTemplate = GetEntityTemplateInternal(World, TemplateID);
 	check(ExistingTemplate);
 	return *ExistingTemplate;
 }
 
-const FMassEntityTemplate* FMassEntityConfig::GetEntityTemplateInternal(const UWorld& World, FMassEntityTemplateID& OutTemplateID, TArray<UMassEntityTraitBase*>& OutCombinedTraits) const
+const FMassEntityTemplate* FMassEntityConfig::GetEntityTemplateInternal(const UWorld& World, FMassEntityTemplateID& OutTemplateID) const
 {
 	UMassSpawnerSubsystem* SpawnerSystem = UWorld::GetSubsystem<UMassSpawnerSubsystem>(&World);
 	check(SpawnerSystem);
 	const FMassEntityTemplateRegistry& TemplateRegistry = SpawnerSystem->GetTemplateRegistryInstance();
 
-	// Combine all the features into one array
-	// @todo this is an inefficient way assuming given template is expected to have already been created. Figure out a way to cache it.
-	TArray<const UObject*> Visited;
-	OutCombinedTraits.Reset();
-	Visited.Add(ConfigOwner);
-	GetCombinedTraits(OutCombinedTraits, Visited);
-
 	// Return existing template if found.
-	// TODO: cache the hash.
-	OutTemplateID = FMassEntityTemplateIDFactory::Make(OutCombinedTraits);
+	OutTemplateID = FMassEntityTemplateIDFactory::Make(ConfigGuid);
 	const TSharedRef<FMassEntityTemplate>* TemplateFound = TemplateRegistry.FindTemplateFromTemplateID(OutTemplateID);
 	return TemplateFound ? &TemplateFound->Get() : nullptr;
 }
 
-void FMassEntityConfig::GetCombinedTraits(TArray<UMassEntityTraitBase*>& OutTraits, TArray<const UObject*>& Visited) const
+void FMassEntityConfig::GetCombinedTraits(TArray<UMassEntityTraitBase*>& OutTraits) const
 {
+	TArray<const UObject*> Visited;
+	OutTraits.Reset();
+	Visited.Add(ConfigOwner);
+	GetCombinedTraitsInternal(OutTraits, Visited);
+}
+
+void FMassEntityConfig::GetCombinedTraitsInternal(TArray<UMassEntityTraitBase*>& OutTraits, TArray<const UObject*>& Visited) const
+{
 	if (Parent)
 	{
@@ -162,7 +142,7 @@
 		else
 		{
 			Visited.Add(Parent);
-			Parent->GetConfig().GetCombinedTraits(OutTraits, Visited);
+			Parent->GetConfig().GetCombinedTraitsInternal(OutTraits, Visited);
 		}
 	}
 
@@ -192,10 +172,8 @@
 
 bool FMassEntityConfig::ValidateEntityTemplate(const UWorld& World)
 {
-	TArray<const UObject*> Visited;
 	TArray<UMassEntityTraitBase*> CombinedTraits;
-	Visited.Add(ConfigOwner);
-	GetCombinedTraits(CombinedTraits, Visited);
+	GetCombinedTraits(CombinedTraits);
 
 	FMassEntityTemplateData Template;
 	FMassEntityTemplateBuildContext BuildContext(Template);
@@ -228,8 +206,36 @@
 
 void FMassEntityConfig::GetCombinedTraits(TArray<UMassEntityTraitBase*>& OutTraits, TArray<const UObject*>& Visited, const UObject& InConfigOwner) const
 {
-	return GetCombinedTraits(OutTraits, Visited);
+	return GetCombinedTraitsInternal(OutTraits, Visited);
 }
 
+//-----------------------------------------------------------------------------
+// UMassEntityConfigAsset
+//-----------------------------------------------------------------------------
+#if WITH_EDITOR
+void UMassEntityConfigAsset::ValidateEntityConfig()
+{
+	if (UWorld* EditorWorld = GEditor->GetEditorWorldContext().World())
+	{
+		if (Config.ValidateEntityTemplate(*EditorWorld))
+		{
+			const FText InfoText = LOCTEXT("MassEntityConfigAssetNoErrorsDetected", "There were no errors nor warnings detected during validation of the EntityConfigAsset");
 
+			FMessageLog EditorInfo("MassEntity");
+			EditorInfo.Info(InfoText);
+
+			FNotificationInfo Info(InfoText);
+			Info.bFireAndForget = true;
+			Info.bUseThrobber = false;
+			Info.FadeOutDuration = 0.5f;
+			Info.ExpireDuration = 5.0f;
+			if (TSharedPtr<SNotificationItem> Notification = FSlateNotificationManager::Get().AddNotification(Info))
+			{
+				Notification->SetCompletionState(SNotificationItem::CS_Success);
+			}
+		}
+	}
+}
+#endif // WITH_EDITOR
+
 #undef LOCTEXT_NAMESPACE 
diff --git a/Engine/Plugins/Runtime/MassGameplay/Source/MassSpawner/Private/MassEntityTemplate.cpp b/Engine/Plugins/Runtime/MassGameplay/Source/MassSpawner/Private/MassEntityTemplate.cpp
--- a/Engine/Plugins/Runtime/MassGameplay/Source/MassSpawner/Private/MassEntityTemplate.cpp
+++ b/Engine/Plugins/Runtime/MassGameplay/Source/MassSpawner/Private/MassEntityTemplate.cpp
@@ -5,13 +5,15 @@
 #include "MassDebugger.h"
 #include "MassSpawnerTypes.h"
 #include "StructUtilsTypes.h"
+#include "Algo/Find.h"
 
 //----------------------------------------------------------------------//
 //  FMassEntityTemplateID
 //----------------------------------------------------------------------//
 FString FMassEntityTemplateID::ToString() const
 {
-	return FString::Printf(TEXT("[%d]"), Hash);
+	return IsValid() ? FString::Printf(TEXT("[Invalid]"))
+		: FString::Printf(TEXT("[%s:%d]"), *ConfigGuid.ToString(EGuidFormats::DigitsLower), FlavorHash);
 }
 
 //----------------------------------------------------------------------//
@@ -23,17 +25,51 @@ FMassEntityTemplateData::FMassEntityTemplateData(const FMassEntityTemplate& InFi
 
 }
 
+bool FMassEntityTemplateData::SlowIsEquivalent(const FMassEntityTemplateData& Other) const
+{
+	if (Composition.IsEquivalent(Other.GetCompositionDescriptor()) == false)
+	{
+		return false;
+	}
+	else if (SharedFragmentValues.IsEquivalent(Other.GetSharedFragmentValues()) == false)
+	{
+		return false;
+	}
+	
+	TConstArrayView<FInstancedStruct> OtherInitialFragmentValues = Other.GetInitialFragmentValues();
+	
+	if (OtherInitialFragmentValues.Num() != InitialFragmentValues.Num())
+	{
+		return false;
+	}
+
+	for (const FInstancedStruct& InitialValue : InitialFragmentValues)
+	{
+		const FInstancedStruct* FoundElement = Algo::FindByPredicate(OtherInitialFragmentValues, [&InitialValue](const FInstancedStruct& Element)
+		{
+			return InitialValue == Element;
+		});
+
+		if (FoundElement == nullptr)
+		{
+			return false;
+		}
+	}
+	return true;
+}
+
 uint32 GetTypeHash(const FMassEntityTemplateData& Template)
 {
 	// @todo: using the template name is temporary solution to allow to tell two templates apart based on something 
 	// else than composition. Ideally we would hash the fragment values instead.
 	const int32 NameHash = GetTypeHash(Template.GetTemplateName());
-	const uint32 CompositionHash = Template.Composition.CalculateHash();
-	const uint32 SharedFragmentValuesHash = GetTypeHash(Template.SharedFragmentValues);
+	const uint32 CompositionHash = Template.GetCompositionDescriptor().CalculateHash();
+	// @todo shared fragments hash is based on pointers - this needs to be changed as well
+	const uint32 SharedFragmentValuesHash = GetTypeHash(Template.GetSharedFragmentValues());
 	
 	uint32 InitialValuesHash = 0;
 	// Initial fragment values, this is not part of the archetype as it is the spawner job to set them.
-	for (const FInstancedStruct& Struct : Template.InitialFragmentValues)
+	for (const FInstancedStruct& Struct : Template.GetInitialFragmentValues())
 	{
 		InitialValuesHash = UE::StructUtils::GetStructCrc32(Struct, InitialValuesHash);
 	}
@@ -120,13 +156,12 @@ FString FMassEntityTemplate::DebugGetDescription(FMassEntityManager* EntityManag
 //-----------------------------------------------------------------------------
 // FMassEntityTemplateIDFactory
 //-----------------------------------------------------------------------------
-FMassEntityTemplateID FMassEntityTemplateIDFactory::Make(const FMassEntityTemplateData& TemplateData)
+FMassEntityTemplateID FMassEntityTemplateIDFactory::Make(const FGuid& ConfigGuid)
 {
-	return FMassEntityTemplateID(GetTypeHash(TemplateData));
+	return FMassEntityTemplateID(ConfigGuid);
 }
 
-FMassEntityTemplateID FMassEntityTemplateIDFactory::Make(TConstArrayView<UMassEntityTraitBase*> Traits)
+FMassEntityTemplateID FMassEntityTemplateIDFactory::MakeFlavor(const FMassEntityTemplateID& SourceTemplateID, const int32 Flavor)
 {
-	const uint32 HashOut = UE::MassSpawner::HashTraits(Traits);
-	return FMassEntityTemplateID(HashOut);
+	return FMassEntityTemplateID(SourceTemplateID.ConfigGuid, Flavor);
 }
diff --git a/Engine/Plugins/Runtime/MassGameplay/Source/MassSpawner/Private/MassEntityTemplateRegistry.cpp b/Engine/Plugins/Runtime/MassGameplay/Source/MassSpawner/Private/MassEntityTemplateRegistry.cpp
--- a/Engine/Plugins/Runtime/MassGameplay/Source/MassSpawner/Private/MassEntityTemplateRegistry.cpp
+++ b/Engine/Plugins/Runtime/MassGameplay/Source/MassSpawner/Private/MassEntityTemplateRegistry.cpp
@@ -76,21 +76,6 @@ const TSharedRef<FMassEntityTemplate>& FMassEntityTemplateRegistry::FindOrAddTem
 	return TemplateIDToTemplateMap.Add(TemplateID, FMassEntityTemplate::MakeFinalTemplate(*EntityManager, MoveTemp(TemplateData), TemplateID));
 }
 
-const TSharedRef<FMassEntityTemplate>& FMassEntityTemplateRegistry::FindOrAddTemplate(FMassEntityTemplateData&& TemplateData)
-{
-	const FMassEntityTemplateID NewTemplateID = FMassEntityTemplateIDFactory::Make(TemplateData);
-
-	check(EntityManager);
-	const TSharedRef<FMassEntityTemplate>* ExistingTemplate = FindTemplateFromTemplateID(NewTemplateID);
-
-	if (ExistingTemplate)
-	{
-		return *ExistingTemplate;
-	}
-
-	return TemplateIDToTemplateMap.Add(NewTemplateID, FMassEntityTemplate::MakeFinalTemplate(*EntityManager, MoveTemp(TemplateData), NewTemplateID));
-}
-
 void FMassEntityTemplateRegistry::DestroyTemplate(FMassEntityTemplateID TemplateID)
 {
 	TemplateIDToTemplateMap.Remove(TemplateID);
diff --git a/Engine/Plugins/Runtime/MassGameplay/Source/MassSpawner/Public/MassEntityConfigAsset.h b/Engine/Plugins/Runtime/MassGameplay/Source/MassSpawner/Public/MassEntityConfigAsset.h
--- a/Engine/Plugins/Runtime/MassGameplay/Source/MassSpawner/Public/MassEntityConfigAsset.h
+++ b/Engine/Plugins/Runtime/MassGameplay/Source/MassSpawner/Public/MassEntityConfigAsset.h
@@ -71,9 +71,11 @@ struct MASSSPAWNER_API FMassEntityConfig
 
 	bool IsEmpty() const { return Parent == nullptr && Traits.Num() == 0; }
 
+	const FGuid& GetGuid() const { return ConfigGuid; }
+	
 protected:
 	/** Combines traits based on the config hierarchy and returns list of unique traits */
-	void GetCombinedTraits(TArray<UMassEntityTraitBase*>& OutTraits, TArray<const UObject*>& Visited) const;
+	void GetCombinedTraits(TArray<UMassEntityTraitBase*>& OutTraits) const;
 
 	/** Combines traits based on the config hierarchy and returns list of unique traits */
 	UE_DEPRECATED(5.3, "This flavor of GetCombinedTraits is deprecated. Use the one without the ConfigOwner parameter (now a property of the FMassEntityConfig itself)")
@@ -91,7 +93,12 @@ protected:
 	TObjectPtr<UObject> ConfigOwner = nullptr;
 
 private:
-	const FMassEntityTemplate* GetEntityTemplateInternal(const UWorld& World, FMassEntityTemplateID& TemplateIDOut, TArray<UMassEntityTraitBase*>& CombinedTraitsOut) const;
+	UPROPERTY(VisibleAnywhere, Category="Mass")
+	FGuid ConfigGuid;
+
+private:
+	void GetCombinedTraitsInternal(TArray<UMassEntityTraitBase*>& OutTraits, TArray<const UObject*>& Visited) const;
+	const FMassEntityTemplate* GetEntityTemplateInternal(const UWorld& World, FMassEntityTemplateID& TemplateIDOut) const;
 };
 
 /**
diff --git a/Engine/Plugins/Runtime/MassGameplay/Source/MassSpawner/Public/MassEntityTemplate.h b/Engine/Plugins/Runtime/MassGameplay/Source/MassSpawner/Public/MassEntityTemplate.h
--- a/Engine/Plugins/Runtime/MassGameplay/Source/MassSpawner/Public/MassEntityTemplate.h
+++ b/Engine/Plugins/Runtime/MassGameplay/Source/MassSpawner/Public/MassEntityTemplate.h
@@ -2,12 +2,13 @@
 
 #pragma once
 
-#include "CoreMinimal.h"
 #include "MassProcessingTypes.h"
 #include "MassEntityManager.h"
 #include "MassCommonTypes.h"
 #include "MassTranslator.h"
 #include "Templates/SharedPointer.h"
+#include "Misc/Guid.h"
+#include "Hash/CityHash.h"
 #include "MassEntityTemplate.generated.h"
 
 
@@ -23,52 +24,68 @@ struct MASSSPAWNER_API FMassEntityTemplateID
 {
 	GENERATED_BODY()
 
-	FMassEntityTemplateID()
-		: bIsSet(false)
+	FMassEntityTemplateID() 
+		: FlavorHash(0), TotalHash(InvalidHash)
 	{}
 
 private:
 	friend FMassEntityTemplateIDFactory;
 	// use FMassEntityTemplateIDFactory to access this constructor flavor
-	explicit FMassEntityTemplateID(uint32 InHash)
-		: Hash(InHash), bIsSet(true)
-	{}
+	explicit FMassEntityTemplateID(const FGuid& InGuid, const int32 InFlavorHash = 0)
+		: ConfigGuid(InGuid), FlavorHash(InFlavorHash)
+	{
+		 const uint64 GuidHash = CityHash64((char*)&ConfigGuid, sizeof(FGuid));
+		 TotalHash = CityHash128to64({GuidHash, (uint64)InFlavorHash});
+	}
 
 public:
-	uint32 GetHash() const 
+	uint64 GetHash64() const 
 	{
 		checkSlow(bIsSet);
-		return Hash; 
+		return TotalHash;
 	}
 	
-	void Invalidate(uint32 InHash)
+	void Invalidate()
 	{
-		// the exact value we set here doesn't really matter, but just to keep the possible states consistent we set it 
-		// to the default value;
-		Hash = 0;
-		bIsSet = false;
+		TotalHash = InvalidHash;
 	}
 
 	bool operator==(const FMassEntityTemplateID& Other) const
 	{
-		return (Hash == Other.Hash) && IsValid() == Other.IsValid();
+		return (TotalHash == Other.TotalHash);
+	}
+	
+	bool operator!=(const FMassEntityTemplateID& Other) const
+	{
+		return !(*this == Other);
 	}
 
+	/** 
+	 * Note that since the function is 32-hashing a 64-bit value it's not guaranteed to produce globally unique values.
+	 * But also note that it's still fine to use FMassEntityTemplateID as a TMap key type, since TMap is using 32bit hash
+	 * to assign buckets rather than identify individual values.
+	 */
 	friend uint32 GetTypeHash(const FMassEntityTemplateID& TemplateID)
 	{
-		return HashCombine(TemplateID.Hash, uint32(TemplateID.bIsSet));
+		return GetTypeHash(TemplateID.TotalHash);
 	}
 
-	bool IsValid() const { return bIsSet; }
+	bool IsValid() const { return (TotalHash != InvalidHash); }
 
 	FString ToString() const;
 
 protected:
+	UPROPERTY(VisibleAnywhere, Category="Mass")
+	FGuid ConfigGuid;
+
 	UPROPERTY()
-	uint32 Hash = 0;
+	uint32 FlavorHash;
 
 	UPROPERTY()
-	uint8 bIsSet : 1;
+	uint64 TotalHash;
+
+private:
+	static constexpr uint64 InvalidHash = 0;
 };
 
 
@@ -252,13 +269,15 @@ struct MASSSPAWNER_API FMassEntityTemplateData
 		return Composition.SharedFragments.Contains(ScriptStruct);
 	}
 
-	friend uint32 GetTypeHash(const FMassEntityTemplateData& Template);
-
 	void Sort()
 	{
 		SharedFragmentValues.Sort();
 	}
 
+	/** Compares contents of two archetypes (this and Other). Returns whether both are equivalent.
+	 *  @Note that the function can be slow, depending on how elaborate the template is. This function is meant for debugging purposes. */
+	bool SlowIsEquivalent(const FMassEntityTemplateData& Other) const;
+
 protected:
 	FMassArchetypeCompositionDescriptor Composition;
 	FMassArchetypeSharedFragmentValues SharedFragmentValues;
@@ -317,8 +336,8 @@ private:
 };
 
 
-struct FMassEntityTemplateIDFactory
+struct MASSSPAWNER_API FMassEntityTemplateIDFactory
 {
-	static FMassEntityTemplateID Make(const FMassEntityTemplateData& TemplateData);
-	static FMassEntityTemplateID Make(TConstArrayView<UMassEntityTraitBase*> Traits);
+	static FMassEntityTemplateID Make(const FGuid& ConfigGuid);
+	static FMassEntityTemplateID MakeFlavor(const FMassEntityTemplateID& SourceTemplateID, const int32 Flavor);
 };
diff --git a/Engine/Plugins/Runtime/MassGameplay/Source/MassSpawner/Public/MassEntityTemplateRegistry.h b/Engine/Plugins/Runtime/MassGameplay/Source/MassSpawner/Public/MassEntityTemplateRegistry.h
--- a/Engine/Plugins/Runtime/MassGameplay/Source/MassSpawner/Public/MassEntityTemplateRegistry.h
+++ b/Engine/Plugins/Runtime/MassGameplay/Source/MassSpawner/Public/MassEntityTemplateRegistry.h
@@ -242,8 +242,6 @@ struct MASSSPAWNER_API FMassEntityTemplateRegistry
 	 */
 	const TSharedRef<FMassEntityTemplate>& FindOrAddTemplate(FMassEntityTemplateID TemplateID, FMassEntityTemplateData&& TemplateData);
 
-	const TSharedRef<FMassEntityTemplate>& FindOrAddTemplate(FMassEntityTemplateData&& TemplateData);
-
 	UE_DEPRECATED(5.3, "We no longer support fething mutable templates from the TemplateRegistry. Stored templates are considered const.")
 	FMassEntityTemplate* FindMutableTemplateFromTemplateID(FMassEntityTemplateID TemplateID);
 
