Created two flavors of MassVisualizationTrait - "Stationary" and "Movable" - and changed MassVisualizationTrait's display name to discourage people from using it, along with MassRepresentationProcessor's targeted support for stationary-representation entities.

Also:
* made FMassISMCSharedData not be a UStruct anymore and change its reference-tracking mechanic internals.
* added FMassStaticMeshInstanceVisualizationMeshDesc.ISMComponentClass that controls which UInstancedStaticMeshComponent-based class will be used for visualization purposes
* made MassUpdateISMProcessor ignore static-tagged entities
* added MassStationaryISMSwitcherProcessor responsible for adding and removing individual entitities to ISMComponents (via FMassInstancedStaticMeshInfoArrayView)

diff --git a/Engine/Plugins/Runtime/MassGameplay/Source/MassRepresentation/Private/MassInstancedStaticMeshComponent.cpp b/Engine/Plugins/Runtime/MassGameplay/Source/MassRepresentation/Private/MassInstancedStaticMeshComponent.cpp
--- a/Engine/Plugins/Runtime/MassGameplay/Source/MassRepresentation/Private/MassInstancedStaticMeshComponent.cpp
+++ b/Engine/Plugins/Runtime/MassGameplay/Source/MassRepresentation/Private/MassInstancedStaticMeshComponent.cpp
@@ -61,6 +61,9 @@ void UMassInstancedStaticMeshComponent::AddInstancesWithIds(TConstArrayView<int3
 
 	for (int32 i = 0; i < InstanceIds.Num(); ++i)
 	{
+		checkfSlow(InstanceIdToInstanceIndexMap.Find(InstanceIds[i]) == nullptr
+			, TEXT("This occuring signals trouble. None of the InstanceIds is expected to have already been added to this MassISM component instance."));
+
 		InstanceIdToInstanceIndexMap.Add(InstanceIds[i], NewIndices[i]);
 		PerInstanceIds[NewIndices[i]] = InstanceIds[i];
 	}
@@ -217,6 +220,11 @@ bool UMassInstancedStaticMeshComponent::RemoveInstanceWithIds(TConstArrayView<in
 		}
 	}
 
+	if (InstanceIndices.Num() == 0)
+	{
+		return false;
+	}
+
 	InstanceIndices.Sort(TGreater<int32>());
 
 	if (!PerInstanceSMData.IsValidIndex(InstanceIndices[0]) || !PerInstanceSMData.IsValidIndex(InstanceIndices.Last()))
diff --git a/Engine/Plugins/Runtime/MassGameplay/Source/MassRepresentation/Private/MassMovableVisualizationTrait.cpp b/Engine/Plugins/Runtime/MassGameplay/Source/MassRepresentation/Private/MassMovableVisualizationTrait.cpp
new file mode 100644
--- /dev/null
+++ b/Engine/Plugins/Runtime/MassGameplay/Source/MassRepresentation/Private/MassMovableVisualizationTrait.cpp
@@ -0,0 +1,20 @@
+// Copyright Epic Games, Inc. All Rights Reserved.
+
+#include "MassMovableVisualizationTrait.h"
+#include "MassEntityTemplateRegistry.h"
+
+
+void UMassMovableVisualizationTrait::BuildTemplate(FMassEntityTemplateBuildContext& BuildContext, const UWorld& World) const
+{
+	if (World.IsNetMode(NM_DedicatedServer))
+	{
+		return;
+	}
+
+	for (FMassStaticMeshInstanceVisualizationMeshDesc& MeshDesc : StaticMeshInstanceDesc.Meshes)
+	{
+		MeshDesc.Mobility = EComponentMobility::Movable;
+	}
+
+	Super::BuildTemplate(BuildContext, World);
+}
diff --git a/Engine/Plugins/Runtime/MassGameplay/Source/MassRepresentation/Private/MassStationaryISMSwitcherProcessor.cpp b/Engine/Plugins/Runtime/MassGameplay/Source/MassRepresentation/Private/MassStationaryISMSwitcherProcessor.cpp
new file mode 100644
--- /dev/null
+++ b/Engine/Plugins/Runtime/MassGameplay/Source/MassRepresentation/Private/MassStationaryISMSwitcherProcessor.cpp
@@ -0,0 +1,116 @@
+// Copyright Epic Games, Inc. All Rights Reserved.
+
+#include "MassStationaryISMSwitcherProcessor.h"
+#include "MassRepresentationSubsystem.h"
+#include "MassCommonFragments.h"
+#include "MassRepresentationProcessor.h"
+
+
+UMassStationaryISMSwitcherProcessor::UMassStationaryISMSwitcherProcessor(const FObjectInitializer& ObjectInitializer)
+	: EntityQuery(*this)
+{
+	ExecutionOrder.ExecuteAfter.Add(UMassVisualizationProcessor::StaticClass()->GetFName());
+	bAutoRegisterWithProcessingPhases = true;
+}
+
+void UMassStationaryISMSwitcherProcessor::ConfigureQueries()
+{
+	EntityQuery.AddRequirement<FMassRepresentationLODFragment>(EMassFragmentAccess::ReadOnly);
+	EntityQuery.AddRequirement<FTransformFragment>(EMassFragmentAccess::ReadOnly);
+	EntityQuery.AddRequirement<FMassRepresentationFragment>(EMassFragmentAccess::ReadWrite);
+	EntityQuery.AddSharedRequirement<FMassRepresentationSubsystemSharedFragment>(EMassFragmentAccess::ReadWrite);
+	EntityQuery.AddTagRequirement<FMassStaticRepresentationTag>(EMassFragmentPresence::All);
+}
+
+void UMassStationaryISMSwitcherProcessor::Execute(FMassEntityManager& EntityManager, FMassExecutionContext& Context)
+{
+	EntityQuery.ForEachEntityChunk(EntityManager, Context, [this](FMassExecutionContext& Context)
+	{
+		UMassRepresentationSubsystem* RepresentationSubsystem = Context.GetMutableSharedFragment<FMassRepresentationSubsystemSharedFragment>().RepresentationSubsystem;
+		check(RepresentationSubsystem);
+		FMassInstancedStaticMeshInfoArrayView ISMInfosView = RepresentationSubsystem->GetMutableInstancedStaticMeshInfos();
+
+		const TConstArrayView<FMassRepresentationLODFragment> RepresentationLODList = Context.GetFragmentView<FMassRepresentationLODFragment>();
+		const TConstArrayView<FTransformFragment> TransformList = Context.GetFragmentView<FTransformFragment>();
+		const TArrayView<FMassRepresentationFragment> RepresentationList = Context.GetMutableFragmentView<FMassRepresentationFragment>();
+
+		const int32 NumEntities = Context.GetNumEntities();
+		for (int32 EntityIdx = 0; EntityIdx < NumEntities; EntityIdx++)
+		{
+			const FMassEntityHandle EntityHandle = Context.GetEntity(EntityIdx);
+			const FMassRepresentationLODFragment& RepresentationLOD = RepresentationLODList[EntityIdx];
+			const FTransformFragment& TransformFragment = TransformList[EntityIdx];
+			FMassRepresentationFragment& Representation = RepresentationList[EntityIdx];
+
+			if (Representation.PrevRepresentation != EMassRepresentationType::StaticMeshInstance
+				&& Representation.CurrentRepresentation != EMassRepresentationType::StaticMeshInstance)
+			{
+				// nothing to do here
+				continue;
+			}
+
+			FMassInstancedStaticMeshInfo& ISMInfo = ISMInfosView[Representation.StaticMeshDescIndex];
+
+			if (Representation.PrevRepresentation == EMassRepresentationType::StaticMeshInstance
+				&& Representation.CurrentRepresentation != EMassRepresentationType::StaticMeshInstance)
+			{
+				const int32 EntityId = GetTypeHash(EntityHandle);
+
+				// note that we're using the PrevLODSignificance here, and the reason for it is that the Prev value matches the 
+				// PrevRepresentation - thus we need to remove from the "previously" used LODSignificance range.
+				ISMInfo.RemoveInstance(EntityId, Representation.PrevLODSignificance);
+
+				// consume "prev" data
+				Representation.PrevRepresentation = Representation.CurrentRepresentation;
+			}
+			else if (Representation.PrevRepresentation != EMassRepresentationType::StaticMeshInstance
+				&& Representation.CurrentRepresentation == EMassRepresentationType::StaticMeshInstance)
+			{
+				const int32 EntityId = GetTypeHash(EntityHandle);
+				const FTransform& Transform = TransformFragment.GetTransform();
+				const FTransform& PrevTransform = Representation.PrevTransform;
+				const float LODSignificance = RepresentationLOD.LODSignificance;
+				const float PrevLODSignificance = Representation.PrevLODSignificance;
+				
+				if (FMassLODSignificanceRange* NewRange = ISMInfo.GetLODSignificanceRange(RepresentationLOD.LODSignificance))
+				{
+					if (ISMInfo.ShouldUseTransformOffset())
+					{
+						const FTransform& TransformOffset = ISMInfo.GetTransformOffset();
+						const FTransform SMTransform = TransformOffset * Transform;
+						NewRange->AddInstance(EntityId, SMTransform);
+					}
+					else
+					{
+						NewRange->AddInstance(EntityId, Transform);
+					}
+				}
+
+				// consume "prev" data
+				Representation.PrevRepresentation = Representation.CurrentRepresentation;
+			}
+			else if (ISMInfo.GetLODSignificanceRangesNum() > 1 && Representation.PrevLODSignificance != RepresentationLOD.LODSignificance)
+			{
+				// we remain in ISM land, but LODSignificance changed and we have multiple LODSignificance ranges for this entity
+				FMassLODSignificanceRange* OldRange = ISMInfo.GetLODSignificanceRange(Representation.PrevLODSignificance);
+				FMassLODSignificanceRange* NewRange = ISMInfo.GetLODSignificanceRange(RepresentationLOD.LODSignificance);
+				if (OldRange != NewRange)
+				{
+					const int32 EntityId = GetTypeHash(EntityHandle);
+					if (OldRange)
+					{
+						OldRange->RemoveInstance(EntityId);
+					}
+					if (NewRange)
+					{
+						const FTransform& Transform = TransformFragment.GetTransform();
+						NewRange->AddInstance(EntityId, Transform);
+					}
+				}
+			}
+
+			// consume "prev" data
+			Representation.PrevLODSignificance = RepresentationLOD.LODSignificance;
+		}
+	});
+}
diff --git a/Engine/Plugins/Runtime/MassGameplay/Source/MassRepresentation/Private/MassStationaryVisualizationTrait.cpp b/Engine/Plugins/Runtime/MassGameplay/Source/MassRepresentation/Private/MassStationaryVisualizationTrait.cpp
new file mode 100644
--- /dev/null
+++ b/Engine/Plugins/Runtime/MassGameplay/Source/MassRepresentation/Private/MassStationaryVisualizationTrait.cpp
@@ -0,0 +1,41 @@
+// Copyright Epic Games, Inc. All Rights Reserved.
+
+#include "MassStationaryVisualizationTrait.h"
+#include "MassEntityTemplateRegistry.h"
+#include "MassInstancedStaticMeshComponent.h"
+
+
+UMassStationaryVisualizationTrait::UMassStationaryVisualizationTrait(const FObjectInitializer& ObjectInitializer)
+{
+}
+
+void UMassStationaryVisualizationTrait::BuildTemplate(FMassEntityTemplateBuildContext& BuildContext, const UWorld& World) const
+{
+	for (FMassStaticMeshInstanceVisualizationMeshDesc& MeshDesc : StaticMeshInstanceDesc.Meshes)
+	{
+		MeshDesc.Mobility = EComponentMobility::Stationary;
+		MeshDesc.ISMComponentClass = UMassInstancedStaticMeshComponent::StaticClass();
+	}
+
+	Super::BuildTemplate(BuildContext, World);
+
+	BuildContext.AddTag<FMassStaticRepresentationTag>();
+}
+
+#if WITH_EDITOR
+void UMassStationaryVisualizationTrait::PostEditChangeProperty(struct FPropertyChangedEvent& PropertyChangedEvent)
+{
+	static const FName StaticMeshInstanceDescName = GET_MEMBER_NAME_CHECKED(UMassStationaryVisualizationTrait, StaticMeshInstanceDesc);
+
+	Super::PostEditChangeProperty(PropertyChangedEvent);
+
+	if (PropertyChangedEvent.MemberProperty && PropertyChangedEvent.MemberProperty->GetFName() == StaticMeshInstanceDescName)
+	{
+		for (FMassStaticMeshInstanceVisualizationMeshDesc& MeshDesc : StaticMeshInstanceDesc.Meshes)
+		{
+			MeshDesc.Mobility = EComponentMobility::Stationary;
+			MeshDesc.ISMComponentClass = UMassInstancedStaticMeshComponent::StaticClass();
+		}
+	}
+}
+#endif // WITH_EDITOR
diff --git a/Engine/Plugins/Runtime/MassGameplay/Source/MassRepresentation/Private/MassUpdateISMProcessor.cpp b/Engine/Plugins/Runtime/MassGameplay/Source/MassRepresentation/Private/MassUpdateISMProcessor.cpp
--- a/Engine/Plugins/Runtime/MassGameplay/Source/MassRepresentation/Private/MassUpdateISMProcessor.cpp
+++ b/Engine/Plugins/Runtime/MassGameplay/Source/MassRepresentation/Private/MassUpdateISMProcessor.cpp
@@ -27,6 +27,9 @@ void UMassUpdateISMProcessor::ConfigureQueries()
 	EntityQuery.AddChunkRequirement<FMassVisualizationChunkFragment>(EMassFragmentAccess::ReadWrite);
 	EntityQuery.SetChunkFilter(&FMassVisualizationChunkFragment::AreAnyEntitiesVisibleInChunk);
 	EntityQuery.AddSharedRequirement<FMassRepresentationSubsystemSharedFragment>(EMassFragmentAccess::ReadWrite);
+
+	// ignore entitites configured to have their representation static (@todo maybe just check if there's not movement fragment?)
+	EntityQuery.AddTagRequirement<FMassStaticRepresentationTag>(EMassFragmentPresence::None);
 }
 
 void UMassUpdateISMProcessor::Execute(FMassEntityManager& EntityManager, FMassExecutionContext& Context)
diff --git a/Engine/Plugins/Runtime/MassGameplay/Source/MassRepresentation/Private/MassVisualizationComponent.cpp b/Engine/Plugins/Runtime/MassGameplay/Source/MassRepresentation/Private/MassVisualizationComponent.cpp
--- a/Engine/Plugins/Runtime/MassGameplay/Source/MassRepresentation/Private/MassVisualizationComponent.cpp
+++ b/Engine/Plugins/Runtime/MassGameplay/Source/MassRepresentation/Private/MassVisualizationComponent.cpp
@@ -9,9 +9,9 @@
 #include "Engine/CollisionProfile.h"
 #include "RenderUtils.h"
 #include "SceneInterface.h"
-
 #include "Components/HierarchicalInstancedStaticMeshComponent.h"
 #include "MassInstancedStaticMeshComponent.h"
+#include "VisualLogger/VisualLogger.h"
 
 //---------------------------------------------------------------
 // UMassVisualizationComponent
@@ -82,6 +82,8 @@ void UMassVisualizationComponent::ConstructStaticMeshComponents()
 			{
 				ISMC = NewObject<UInstancedStaticMeshComponent>(ActorOwner, MeshDesc.ISMComponentClass);
 				CA_ASSUME(ISMC);
+				REDIRECT_OBJECT_TO_VLOG(ISMC, this);
+
 				ISMC->SetStaticMesh(MeshDesc.Mesh);
 				for (int32 ElementIndex = 0; ElementIndex < MeshDesc.MaterialOverrides.Num(); ++ElementIndex)
 				{
@@ -361,6 +363,18 @@ void FMassLODSignificanceRange::AddBatchedCustomDataFloats(const TArray<float>&
 	}
 }
 
+void FMassLODSignificanceRange::AddInstance(const int32 InstanceId, const FTransform& Transform)
+{
+	check(ISMCSharedDataPtr);
+	for (int i = 0; i < StaticMeshRefs.Num(); i++)
+	{
+		FMassISMCSharedData& SharedData = (*ISMCSharedDataPtr)[StaticMeshRefs[i]];
+		SharedData.UpdateInstanceIds.Add(InstanceId);
+		SharedData.StaticMeshInstanceTransforms.Add(Transform);
+		SharedData.StaticMeshInstancePrevTransforms.Add(Transform);
+	}
+}
+
 void FMassLODSignificanceRange::RemoveInstance(const int32 InstanceId)
 {
 	check(ISMCSharedDataPtr);
diff --git a/Engine/Plugins/Runtime/MassGameplay/Source/MassRepresentation/Public/MassMovableVisualizationTrait.h b/Engine/Plugins/Runtime/MassGameplay/Source/MassRepresentation/Public/MassMovableVisualizationTrait.h
new file mode 100644
--- /dev/null
+++ b/Engine/Plugins/Runtime/MassGameplay/Source/MassRepresentation/Public/MassMovableVisualizationTrait.h
@@ -0,0 +1,16 @@
+// Copyright Epic Games, Inc. All Rights Reserved.
+
+#pragma once 
+
+#include "MassVisualizationTrait.h"
+#include "MassMovableVisualizationTrait.generated.h"
+
+
+UCLASS()
+class MASSREPRESENTATION_API UMassMovableVisualizationTrait : public UMassVisualizationTrait
+{
+	GENERATED_BODY()
+public:
+	virtual void BuildTemplate(FMassEntityTemplateBuildContext& BuildContext, const UWorld& World) const override;
+};
+	
\ No newline at end of file
diff --git a/Engine/Plugins/Runtime/MassGameplay/Source/MassRepresentation/Public/MassRepresentationFragments.h b/Engine/Plugins/Runtime/MassGameplay/Source/MassRepresentation/Public/MassRepresentationFragments.h
--- a/Engine/Plugins/Runtime/MassGameplay/Source/MassRepresentation/Public/MassRepresentationFragments.h
+++ b/Engine/Plugins/Runtime/MassGameplay/Source/MassRepresentation/Public/MassRepresentationFragments.h
@@ -15,6 +15,12 @@
 class UMassRepresentationSubsystem;
 class UMassRepresentationActorManagement;
 
+USTRUCT()
+struct MASSREPRESENTATION_API FMassStaticRepresentationTag : public FMassTag
+{
+	GENERATED_BODY(); 
+};
+
 USTRUCT()
 struct MASSREPRESENTATION_API FMassRepresentationLODFragment : public FMassFragment
 {
diff --git a/Engine/Plugins/Runtime/MassGameplay/Source/MassRepresentation/Public/MassRepresentationTypes.h b/Engine/Plugins/Runtime/MassGameplay/Source/MassRepresentation/Public/MassRepresentationTypes.h
--- a/Engine/Plugins/Runtime/MassGameplay/Source/MassRepresentation/Public/MassRepresentationTypes.h
+++ b/Engine/Plugins/Runtime/MassGameplay/Source/MassRepresentation/Public/MassRepresentationTypes.h
@@ -226,6 +226,9 @@ public:
 
 	void AddBatchedCustomDataFloats(const TArray<float>& CustomFloats, const TArray<uint32>& ExcludeStaticMeshRefs);
 
+	/** Single-instance version of AddBatchedCustomData when called to add entities (as opposed to modify existing ones).*/
+	void AddInstance(const int32 InstanceId, const FTransform& Transform);
+
 	void RemoveInstance(const int32 InstanceId);
 
 	void WriteCustomDataFloatsAtStartIndex(int32 StaticMeshIndex, const TArrayView<float>& CustomFloats, const int32 FloatsPerInstance, const int32 StartIndex, const TArray<uint32>& ExcludeStaticMeshRefs);
@@ -364,6 +367,8 @@ public:
 		}
 	}
 
+	int32 GetLODSignificanceRangesNum() const { return LODSignificanceRanges.Num(); }
+
 protected:
 
 	/** Destroy the visual instance */
diff --git a/Engine/Plugins/Runtime/MassGameplay/Source/MassRepresentation/Public/MassStationaryISMSwitcherProcessor.h b/Engine/Plugins/Runtime/MassGameplay/Source/MassRepresentation/Public/MassStationaryISMSwitcherProcessor.h
new file mode 100644
--- /dev/null
+++ b/Engine/Plugins/Runtime/MassGameplay/Source/MassRepresentation/Public/MassStationaryISMSwitcherProcessor.h
@@ -0,0 +1,32 @@
+// Copyright Epic Games, Inc. All Rights Reserved.
+
+#pragma once 
+
+#include "MassProcessor.h"
+#include "MassStationaryISMSwitcherProcessor.generated.h"
+
+
+/** 
+ * This processor's sole responsibility is to process all entities tagged with FMassStaticRepresentationTag
+ * and check if they've switched to or away from EMassRepresentationType::StaticMeshInstance; and acordingly add or remove 
+ * the entity from the appropriate FMassInstancedStaticMeshInfoArrayView.
+ */
+UCLASS()
+class MASSREPRESENTATION_API UMassStationaryISMSwitcherProcessor : public UMassProcessor
+{
+	GENERATED_BODY()
+
+public:
+	UMassStationaryISMSwitcherProcessor(const FObjectInitializer& ObjectInitializer = FObjectInitializer::Get());
+
+protected:
+	// UMassProcessor overrides begin
+	virtual void ConfigureQueries() override;
+	// UMassProcessor overrides end
+	// 
+	// UMassRepresentationProcessor overrides begin
+	virtual void Execute(FMassEntityManager& EntityManager, FMassExecutionContext& Context) override;
+	// UMassRepresentationProcessor overrides end
+
+	FMassEntityQuery EntityQuery;
+};
diff --git a/Engine/Plugins/Runtime/MassGameplay/Source/MassRepresentation/Public/MassStationaryVisualizationTrait.h b/Engine/Plugins/Runtime/MassGameplay/Source/MassRepresentation/Public/MassStationaryVisualizationTrait.h
new file mode 100644
--- /dev/null
+++ b/Engine/Plugins/Runtime/MassGameplay/Source/MassRepresentation/Public/MassStationaryVisualizationTrait.h
@@ -0,0 +1,22 @@
+// Copyright Epic Games, Inc. All Rights Reserved.
+
+#pragma once 
+
+#include "MassVisualizationTrait.h"
+#include "MassStationaryVisualizationTrait.generated.h"
+
+
+UCLASS()
+class MASSREPRESENTATION_API UMassStationaryVisualizationTrait : public UMassVisualizationTrait
+{
+	GENERATED_BODY()
+public:
+	UMassStationaryVisualizationTrait(const FObjectInitializer& ObjectInitializer = FObjectInitializer::Get());
+
+	virtual void BuildTemplate(FMassEntityTemplateBuildContext& BuildContext, const UWorld& World) const override;
+
+protected:
+#if WITH_EDITOR
+	virtual void PostEditChangeProperty(struct FPropertyChangedEvent& PropertyChangedEvent) override;
+#endif // WITH_EDITOR
+};
diff --git a/Engine/Plugins/Runtime/MassGameplay/Source/MassRepresentation/Public/MassVisualizationTrait.h b/Engine/Plugins/Runtime/MassGameplay/Source/MassRepresentation/Public/MassVisualizationTrait.h
--- a/Engine/Plugins/Runtime/MassGameplay/Source/MassRepresentation/Public/MassVisualizationTrait.h
+++ b/Engine/Plugins/Runtime/MassGameplay/Source/MassRepresentation/Public/MassVisualizationTrait.h
@@ -13,7 +13,8 @@ class UMassRepresentationSubsystem;
 class UMassRepresentationActorManagement;
 class UMassProcessor;
 
-UCLASS(meta=(DisplayName="Visualization"))
+/** This class has been soft-deprecated. Use MassStationaryVisualizationTrait or MassMovableVisualizationTrait */
+UCLASS(meta=(DisplayName="DEPRECATED Visualization"))
 class MASSREPRESENTATION_API UMassVisualizationTrait : public UMassEntityTraitBase
 {
 	GENERATED_BODY()
@@ -24,7 +25,7 @@ public:
 
 	/** Instanced static mesh information for this agent */
 	UPROPERTY(EditAnywhere, Category = "Mass|Visual")
-	FStaticMeshInstanceVisualizationDesc StaticMeshInstanceDesc;
+	mutable FStaticMeshInstanceVisualizationDesc StaticMeshInstanceDesc;
 
 	/** Actor class of this agent when spawned in high resolution*/
 	UPROPERTY(EditAnywhere, Category = "Mass|Visual")
