Removed unneeded forward declaration of AMassCrowdClientBubbleInfo from MassReplicationSubsystem.h

diff --git a/Engine/Plugins/Runtime/MassGameplay/Source/MassReplication/Public/MassReplicationSubsystem.h b/Engine/Plugins/Runtime/MassGameplay/Source/MassReplication/Public/MassReplicationSubsystem.h
--- a/Engine/Plugins/Runtime/MassGameplay/Source/MassReplication/Public/MassReplicationSubsystem.h
+++ b/Engine/Plugins/Runtime/MassGameplay/Source/MassReplication/Public/MassReplicationSubsystem.h
@@ -10,7 +10,6 @@
 #include "MassReplicationSubsystem.generated.h"
 
 struct FMassEntityManager;
-class AMassCrowdClientBubbleInfo;
 class UMassLODSubsystem;
 class AMassClientBubbleInfoBase;
 
