Changed FMassEntityTemplateRegistry::FindOrAddTemplate so that it doesn't complain when a template has actually been found.

diff --git a/Engine/Plugins/Runtime/MassGameplay/Source/MassSpawner/Private/MassEntityTemplateRegistry.cpp b/Engine/Plugins/Runtime/MassGameplay/Source/MassSpawner/Private/MassEntityTemplateRegistry.cpp
--- a/Engine/Plugins/Runtime/MassGameplay/Source/MassSpawner/Private/MassEntityTemplateRegistry.cpp
+++ b/Engine/Plugins/Runtime/MassGameplay/Source/MassSpawner/Private/MassEntityTemplateRegistry.cpp
@@ -68,7 +68,7 @@ const TSharedRef<FMassEntityTemplate>& FMassEntityTemplateRegistry::FindOrAddTem
 {
 	check(EntityManager);
 	const TSharedRef<FMassEntityTemplate>* ExistingTemplate = FindTemplateFromTemplateID(TemplateID);
-	if (!ensureMsgf(ExistingTemplate == nullptr, TEXT("")))
+	if (ExistingTemplate != nullptr)
 	{
 		return *ExistingTemplate;
 	}
