Removed unused local variable from TClientBubbleHandlerBase<AgentArrayItem>::PostReplicatedAddHelper

diff --git a/Engine/Plugins/Runtime/MassGameplay/Source/MassReplication/Public/MassClientBubbleHandler.h b/Engine/Plugins/Runtime/MassGameplay/Source/MassReplication/Public/MassClientBubbleHandler.h
--- a/Engine/Plugins/Runtime/MassGameplay/Source/MassReplication/Public/MassClientBubbleHandler.h
+++ b/Engine/Plugins/Runtime/MassGameplay/Source/MassReplication/Public/MassClientBubbleHandler.h
@@ -419,8 +419,6 @@ template<typename AgentArrayItem>
 void TClientBubbleHandlerBase<AgentArrayItem>::PostReplicatedAddHelper(const TArrayView<int32> AddedIndices, FAddRequirementsForSpawnQueryFunction AddRequirementsForSpawnQuery
 	, FCacheFragmentViewsForSpawnQueryFunction CacheFragmentViewsForSpawnQuery, FSetSpawnedEntityDataFunction SetSpawnedEntityData, FSetModifiedEntityDataFunction SetModifiedEntityData)
 {
-	TArray<FMassEntityHandle> EntitiesDestroy;
-
 	FMassEntityManager& EntityManager = Serializer->GetEntityManagerChecked();
 
 	UMassReplicationSubsystem* ReplicationSubsystem = Serializer->GetReplicationSubsystem();
