Made it possible to make FMassRepresentationParameters:::RepresentationActorManagementClass and UMassVisualizationTrait::RepresentationSubsystemClass  read only in the UI

diff --git a/Engine/Plugins/Runtime/MassGameplay/Source/MassRepresentation/Public/MassRepresentationFragments.h b/Engine/Plugins/Runtime/MassGameplay/Source/MassRepresentation/Public/MassRepresentationFragments.h
--- a/Engine/Plugins/Runtime/MassGameplay/Source/MassRepresentation/Public/MassRepresentationFragments.h
+++ b/Engine/Plugins/Runtime/MassGameplay/Source/MassRepresentation/Public/MassRepresentationFragments.h
@@ -88,7 +88,7 @@ struct FMassRepresentationParameters : public FMassSharedFragment
 	FMassRepresentationParameters() = default;
 
 	/** Allow subclasses to override the representation actor management behavior */
-	UPROPERTY(EditAnywhere, Category = "Mass|Visual")
+	UPROPERTY(EditAnywhere, Category = "Mass|Visual", meta = (EditCondition = "bCanModifyRepresentationActorManagementClass"))
 	TSubclassOf<UMassRepresentationActorManagement> RepresentationActorManagementClass;
 
 	/** What should be the representation of this entity for each specific LOD */
@@ -123,6 +123,12 @@ struct FMassRepresentationParameters : public FMassSharedFragment
 
 	UPROPERTY(Transient)
 	mutable TObjectPtr<UMassRepresentationActorManagement> CachedRepresentationActorManagement = nullptr;
+
+#if WITH_EDITORONLY_DATA
+	/** the property is marked like this to ensure it won't show up in UI */
+	UPROPERTY(EditDefaultsOnly, Category = "Mass|Visual")
+	bool bCanModifyRepresentationActorManagementClass = true;
+#endif // WITH_EDITORONLY_DATA
 };
 
 template<>
diff --git a/Engine/Plugins/Runtime/MassGameplay/Source/MassRepresentation/Public/MassVisualizationTrait.h b/Engine/Plugins/Runtime/MassGameplay/Source/MassRepresentation/Public/MassVisualizationTrait.h
--- a/Engine/Plugins/Runtime/MassGameplay/Source/MassRepresentation/Public/MassVisualizationTrait.h
+++ b/Engine/Plugins/Runtime/MassGameplay/Source/MassRepresentation/Public/MassVisualizationTrait.h
@@ -36,7 +36,7 @@ public:
 	TSubclassOf<AActor> LowResTemplateActor;
 
 	/** Allow subclasses to override the representation subsystem to use */
-	UPROPERTY(EditAnywhere, Category = "Mass|Visual")
+	UPROPERTY(EditAnywhere, Category = "Mass|Visual", meta = (EditCondition = "bCanModifyRepresentationSubsystemClass"))
 	TSubclassOf<UMassRepresentationSubsystem> RepresentationSubsystemClass;
 
 	/** Configuration parameters for the representation processor */
@@ -46,4 +46,10 @@ public:
 	/** Configuration parameters for the visualization LOD processor */
 	UPROPERTY(EditAnywhere, Category = "Mass|Visual")
 	FMassVisualizationLODParameters LODParams;
+
+#if WITH_EDITORONLY_DATA
+	/** the property is marked like this to ensure it won't show up in UI */
+	UPROPERTY(EditDefaultsOnly, Category = "Mass|Visual")
+	bool bCanModifyRepresentationSubsystemClass = true;
+#endif // WITH_EDITORONLY_DATA
 };
