Made FMassActorFragment.Actor a UProperty and marked in a way to have it show up in Mass's gameplay debugger when viewing fragment details

diff --git a/Engine/Plugins/Runtime/MassGameplay/Source/MassActors/Public/MassActorSubsystem.h b/Engine/Plugins/Runtime/MassGameplay/Source/MassActors/Public/MassActorSubsystem.h
--- a/Engine/Plugins/Runtime/MassGameplay/Source/MassActors/Public/MassActorSubsystem.h
+++ b/Engine/Plugins/Runtime/MassGameplay/Source/MassActors/Public/MassActorSubsystem.h
@@ -75,6 +75,8 @@ struct MASSACTORS_API FMassActorFragment : public FObjectWrapperFragment
 	FORCEINLINE bool IsValid() const { return Actor.IsValid(); }
 
 private:
+	// made visible for debugging purposes. It will show up in Mass's gameplay debugger category when viewing fragment details
+	UPROPERTY(VisibleAnywhere, Category="Mass", Transient)
 	TWeakObjectPtr<AActor> Actor;
 
 	/** Ownership of the actor */
