Made having a radius fragment not required to be processed by MassReplicationGridProcessor

diff --git a/Engine/Plugins/Runtime/MassGameplay/Source/MassReplication/Private/MassReplicationGridProcessor.cpp b/Engine/Plugins/Runtime/MassGameplay/Source/MassReplication/Private/MassReplicationGridProcessor.cpp
--- a/Engine/Plugins/Runtime/MassGameplay/Source/MassReplication/Private/MassReplicationGridProcessor.cpp
+++ b/Engine/Plugins/Runtime/MassGameplay/Source/MassReplication/Private/MassReplicationGridProcessor.cpp
@@ -28,17 +28,23 @@ UMassReplicationGridProcessor::UMassReplicationGridProcessor()
 void UMassReplicationGridProcessor::ConfigureQueries()
 {
 	AddToGridEntityQuery.AddRequirement<FTransformFragment>(EMassFragmentAccess::ReadOnly);
-	AddToGridEntityQuery.AddRequirement<FAgentRadiusFragment>(EMassFragmentAccess::ReadOnly);
 	AddToGridEntityQuery.AddRequirement<FMassReplicationGridCellLocationFragment>(EMassFragmentAccess::ReadWrite);
 	AddToGridEntityQuery.AddSubsystemRequirement<UMassReplicationSubsystem>(EMassFragmentAccess::ReadWrite);
 
-	UpdateGridEntityQuery = AddToGridEntityQuery;
+	// copying AddToGridEntityQuery to RemoveFromGridEntityQuery now because RemoveFromGridEntityQuery doesn't utilize 
+	// the other fragments AddToGridEntityQuery relies on
 	RemoveFromGridEntityQuery = AddToGridEntityQuery;
 
+	// FAgentRadiusFragment is optional since it's not strictly required for the provided functionality 
+	AddToGridEntityQuery.AddRequirement<FAgentRadiusFragment>(EMassFragmentAccess::ReadOnly, EMassFragmentPresence::Optional);
+	// we don't care about "off-lod" entities
 	AddToGridEntityQuery.AddTagRequirement<FMassOffLODTag>(EMassFragmentPresence::None);
-	AddToGridEntityQuery.AddTagRequirement<FMassInReplicationGridTag>(EMassFragmentPresence::None);
 
-	UpdateGridEntityQuery.AddTagRequirement<FMassOffLODTag>(EMassFragmentPresence::None);
+	// storing the state in UpdateGridEntityQuery, after that both queries diverge in terms of requirements
+	UpdateGridEntityQuery = AddToGridEntityQuery;
+		
+	AddToGridEntityQuery.AddTagRequirement<FMassInReplicationGridTag>(EMassFragmentPresence::None);
+		
 	UpdateGridEntityQuery.AddTagRequirement<FMassInReplicationGridTag>(EMassFragmentPresence::All);
 
 	RemoveFromGridEntityQuery.AddTagRequirement<FMassOffLODTag>(EMassFragmentPresence::All);
@@ -61,7 +67,8 @@ void UMassReplicationGridProcessor::Execute(FMassEntityManager& EntityManager, F
 		{
 			// Add to the grid
 			const FVector NewPos = LocationList[EntityIndex].GetTransform().GetLocation();
-			const float Radius = RadiusList[EntityIndex].Radius;
+			// note that 0-radius is fine, the underlying THierarchicalHashGrid2D supports that just fine
+			const float Radius = RadiusList.IsEmpty() ? 0.f : RadiusList[EntityIndex].Radius;
 
 			const FMassEntityHandle EntityHandle = Context.GetEntity(EntityIndex);
 			const FBox NewBounds(NewPos - FVector(Radius, Radius, 0.f), NewPos + FVector(Radius, Radius, 0.f));
@@ -85,7 +92,8 @@ void UMassReplicationGridProcessor::Execute(FMassEntityManager& EntityManager, F
 		{
 			// Update position in grid
 			const FVector NewPos = LocationList[EntityIndex].GetTransform().GetLocation();
-			const float Radius = RadiusList[EntityIndex].Radius;
+			// note that 0-radius is fine, the underlying THierarchicalHashGrid2D supports that just fine
+			const float Radius = RadiusList.IsEmpty() ? 0.f : RadiusList[EntityIndex].Radius;
 			const FMassEntityHandle EntityHandle = Context.GetEntity(EntityIndex);
 			const FBox NewBounds(NewPos - FVector(Radius, Radius, 0.f), NewPos + FVector(Radius, Radius, 0.f));
 			ReplicationCellLocationList[EntityIndex].CellLoc = ReplicationGrid.Move(EntityHandle, ReplicationCellLocationList[EntityIndex].CellLoc, NewBounds);
diff --git a/Engine/Source/Runtime/AIModule/Public/HierarchicalHashGrid2D.h b/Engine/Source/Runtime/AIModule/Public/HierarchicalHashGrid2D.h
--- a/Engine/Source/Runtime/AIModule/Public/HierarchicalHashGrid2D.h
+++ b/Engine/Source/Runtime/AIModule/Public/HierarchicalHashGrid2D.h
@@ -510,6 +510,7 @@ public:
 		for (Location.Level = 0; Location.Level < NumLevels; Location.Level++)
 		{
 			const int32 DiameterCells = IntCastChecked<int32>(FMath::CeilToInt(Diameter * InvCellSize[Location.Level]));
+			// note that it's fine for DiameterCells to equal 0 - that would happen for 0-sized items (valid location, no extent).
 			if (DiameterCells <= 1)
 			{
 				Location.X = IntCastChecked<int32>(FMath::FloorToInt(Center.X * InvCellSize[Location.Level]));
