MassLWI plugin added. It aims to host all LWI-related Mass bits.

diff --git a/Engine/Plugins/Runtime/MassGameplay/MassGameplay.uplugin b/Engine/Plugins/Runtime/MassGameplay/MassGameplay.uplugin
--- a/Engine/Plugins/Runtime/MassGameplay/MassGameplay.uplugin
+++ b/Engine/Plugins/Runtime/MassGameplay/MassGameplay.uplugin
@@ -73,6 +73,10 @@
 			"Name": "MassGameplayExternalTraits",
 			"Type": "Runtime"
 		},
+		{
+			"Name": "MassLWI",
+			"Type": "Runtime"
+		},
 		{
 			"Name": "MassGameplayTestSuite",
 			"Type": "UncookedOnly",
diff --git a/Engine/Plugins/Runtime/MassGameplay/Source/MassLWI/MassLWI.Build.cs b/Engine/Plugins/Runtime/MassGameplay/Source/MassLWI/MassLWI.Build.cs
new file mode 100644
--- /dev/null
+++ b/Engine/Plugins/Runtime/MassGameplay/Source/MassLWI/MassLWI.Build.cs
@@ -0,0 +1,34 @@
+// Copyright Epic Games, Inc. All Rights Reserved.
+
+namespace UnrealBuildTool.Rules
+{
+	public class MassLWI : ModuleRules
+	{
+		public MassLWI(ReadOnlyTargetRules Target) : base(Target)
+		{
+			PCHUsage = ModuleRules.PCHUsageMode.UseExplicitOrSharedPCHs;
+
+			PublicDependencyModuleNames.AddRange(
+				new string[] {
+					"Core",
+					"CoreUObject",
+					"Engine",
+					"MassEntity",
+					"MassCommon",
+					"MassReplication",
+					"MassRepresentation",
+					"MassSpawner",
+					"StructUtils",
+					"AIModule",
+					"NetCore",
+				}
+			);
+
+			if (Target.bBuildEditor == true)
+			{
+				PrivateDependencyModuleNames.Add("EditorFramework");
+				PrivateDependencyModuleNames.Add("UnrealEd");
+			}
+		}
+	}
+}
diff --git a/Engine/Plugins/Runtime/MassGameplay/Source/MassLWI/Private/MassLWIConfigActor.cpp b/Engine/Plugins/Runtime/MassGameplay/Source/MassLWI/Private/MassLWIConfigActor.cpp
new file mode 100644
--- /dev/null
+++ b/Engine/Plugins/Runtime/MassGameplay/Source/MassLWI/Private/MassLWIConfigActor.cpp
@@ -0,0 +1,17 @@
+// Copyright Epic Games, Inc. All Rights Reserved.
+
+#include "MassLWIConfigActor.h"
+#include "MassLWISubsystem.h"
+
+//-----------------------------------------------------------------------------
+// ALWIxMassConfigActor
+//-----------------------------------------------------------------------------
+void AMassLWIConfigActor::PostLoad()
+{
+	if (UMassLWISubsystem* MassLWI = UWorld::GetSubsystem<UMassLWISubsystem>(GetWorld()))
+	{
+		MassLWI->RegisterConfigActor(*this);
+	}
+	Super::PostLoad();
+}
+
diff --git a/Engine/Plugins/Runtime/MassGameplay/Source/MassLWI/Private/MassLWIModule.cpp b/Engine/Plugins/Runtime/MassGameplay/Source/MassLWI/Private/MassLWIModule.cpp
new file mode 100644
--- /dev/null
+++ b/Engine/Plugins/Runtime/MassGameplay/Source/MassLWI/Private/MassLWIModule.cpp
@@ -0,0 +1,33 @@
+// Copyright Epic Games, Inc. All Rights Reserved.
+
+#include "MassLWITypes.h"
+#include "IMassLWIModule.h"
+
+DEFINE_LOG_CATEGORY(LogMassLWI)
+
+
+class FMassLWIModule : public IMassLWIModule
+{
+	/** IModuleInterface implementation */
+	virtual void StartupModule() override;
+	virtual void ShutdownModule() override;
+};
+
+IMPLEMENT_MODULE(FMassLWIModule, MassLWI)
+
+
+
+void FMassLWIModule::StartupModule()
+{
+	// This code will execute after your module is loaded into memory (but after global variables are initialized, of course.)
+}
+
+
+void FMassLWIModule::ShutdownModule()
+{
+	// This function may be called during shutdown to clean up your module.  For modules that support dynamic reloading,
+	// we call this function before unloading the module.
+}
+
+
+
diff --git a/Engine/Plugins/Runtime/MassGameplay/Source/MassLWI/Private/MassLWIRepresentationActorManagement.cpp b/Engine/Plugins/Runtime/MassGameplay/Source/MassLWI/Private/MassLWIRepresentationActorManagement.cpp
new file mode 100644
--- /dev/null
+++ b/Engine/Plugins/Runtime/MassGameplay/Source/MassLWI/Private/MassLWIRepresentationActorManagement.cpp
@@ -0,0 +1,22 @@
+// Copyright Epic Games, Inc. All Rights Reserved.
+
+#include "MassLWIRepresentationActorManagement.h"
+#include "MassLWIRepresentationSubsystem.h"
+
+
+//-----------------------------------------------------------------------------
+// UMassLWIRepresentationActorManagement
+//-----------------------------------------------------------------------------
+AActor* UMassLWIRepresentationActorManagement::GetOrSpawnActor(UMassRepresentationSubsystem& RepresentationSubsystem
+	, FMassEntityManager& EntityManager, const FMassEntityHandle MassAgent, FMassActorFragment& ActorInfo
+	, const FTransform& Transform, const int16 TemplateActorIndex, FMassActorSpawnRequestHandle& SpawnRequestHandle, const float Priority) const
+{
+	if (UMassLWIRepresentationSubsystem* LWIRepresentationSubsystem = Cast<UMassLWIRepresentationSubsystem>(&RepresentationSubsystem))
+	{
+		return LWIRepresentationSubsystem->GetOrSpawnActorFromTemplate(MassAgent, Transform, TemplateActorIndex, SpawnRequestHandle, Priority,
+			FMassActorPreSpawnDelegate::CreateUObject(this, &UMassRepresentationActorManagement::OnPreActorSpawn, &EntityManager),
+			FMassActorPostSpawnDelegate::CreateUObject(this, &UMassRepresentationActorManagement::OnPostActorSpawn, &EntityManager));
+	}
+
+	return Super::GetOrSpawnActor(RepresentationSubsystem, EntityManager, MassAgent, ActorInfo, Transform, TemplateActorIndex, SpawnRequestHandle, Priority);
+}
\ No newline at end of file
diff --git a/Engine/Plugins/Runtime/MassGameplay/Source/MassLWI/Private/MassLWIRepresentationSubsystem.cpp b/Engine/Plugins/Runtime/MassGameplay/Source/MassLWI/Private/MassLWIRepresentationSubsystem.cpp
new file mode 100644
--- /dev/null
+++ b/Engine/Plugins/Runtime/MassGameplay/Source/MassLWI/Private/MassLWIRepresentationSubsystem.cpp
@@ -0,0 +1,31 @@
+// Copyright Epic Games, Inc. All Rights Reserved.
+
+#include "MassLWIRepresentationSubsystem.h"
+#include "Engine/Note.h"
+#include "Engine/World.h"
+
+
+//-----------------------------------------------------------------------------
+// UMassLWIRepresentationSubsystem
+//-----------------------------------------------------------------------------
+AActor* UMassLWIRepresentationSubsystem::GetOrSpawnActorFromTemplate(const FMassEntityHandle MassAgent, const FTransform& Transform, const int16 TemplateActorIndex
+	, FMassActorSpawnRequestHandle& SpawnRequestHandle, float Priority, FMassActorPreSpawnDelegate ActorPreSpawnDelegate
+	, FMassActorPostSpawnDelegate ActorPostSpawnDelegate)
+{
+	UWorld* World = GetWorld();
+	if (World == nullptr)
+	{
+		return nullptr;
+	}
+	else if (World->IsNetMode(NM_Client))
+	{
+		// this is where we'd sent a request to the server to spawn the actor
+		//return SpawnedActor;
+	}
+
+	// @todo notify OnActorSpawnedDelegate
+
+	return Super::GetOrSpawnActorFromTemplate(MassAgent, Transform, TemplateActorIndex
+		, SpawnRequestHandle, Priority, ActorPreSpawnDelegate
+		, ActorPostSpawnDelegate);
+}
diff --git a/Engine/Plugins/Runtime/MassGameplay/Source/MassLWI/Private/MassLWIStaticMeshManager.cpp b/Engine/Plugins/Runtime/MassGameplay/Source/MassLWI/Private/MassLWIStaticMeshManager.cpp
new file mode 100644
--- /dev/null
+++ b/Engine/Plugins/Runtime/MassGameplay/Source/MassLWI/Private/MassLWIStaticMeshManager.cpp
@@ -0,0 +1,214 @@
+// Copyright Epic Games, Inc. All Rights Reserved.
+
+#include "MassLWIStaticMeshManager.h"
+#include "MassLWISubsystem.h"
+#include "Components/HierarchicalInstancedStaticMeshComponent.h"
+#include "MassCommonFragments.h"
+#include "MassRepresentationSubsystem.h"
+#include "MassSpawnerSubsystem.h"
+#include "MassSpawnerTypes.h"
+#include "MassSpawnLocationProcessor.h"
+#include "MassVisualizationTrait.h"
+#include "VisualLogger/VisualLogger.h"
+
+
+void AMassLWIStaticMeshManager::PostLoad()
+{
+	Super::PostLoad();
+
+	if (UWorld* World = GetWorld())
+	{
+		if (UMassLWISubsystem* MassLWI = World->GetSubsystem<UMassLWISubsystem>())
+		{
+			MassLWI->RegisterLWIManager(*this);
+		}
+	}
+}
+
+void AMassLWIStaticMeshManager::BeginPlay()
+{
+	Super::BeginPlay();
+
+	UWorld* World = GetWorld();
+	if (World != nullptr && IsRegisteredWithMass() == false)
+	{
+		if (UMassLWISubsystem* LWIxMass = World->GetSubsystem<UMassLWISubsystem>())
+		{
+			LWIxMass->RegisterLWIManager(*this);
+		}
+	}
+}
+
+void AMassLWIStaticMeshManager::EndPlay(const EEndPlayReason::Type EndPlayReason)
+{
+	UWorld* World = GetWorld();
+	if (World && IsRegisteredWithMass())
+	{
+		if (UMassLWISubsystem* LWIxMass = World->GetSubsystem<UMassLWISubsystem>())
+		{
+			LWIxMass->UnregisterLWIManager(*this);
+		}
+	}
+
+	Super::EndPlay(EndPlayReason);
+}
+
+void AMassLWIStaticMeshManager::TransferDataToMass(FMassEntityManager& EntityManager)
+{
+	if (InstanceTransforms.Num() == 0)
+	{
+		return;
+	}
+
+	if (MassTemplateID.IsValid() == false && RepresentedClass)
+	{
+		CreateMassTemplate(EntityManager);
+	}
+
+	if (MassTemplateID.IsValid() && InstancedStaticMeshComponent && FinalizedTemplate)
+	{
+		const FVector& ManagerLocation = GetActorLocation();
+
+		UMassSpawnerSubsystem* SpawnerSystem = UWorld::GetSubsystem<UMassSpawnerSubsystem>(GetWorld());
+		check(SpawnerSystem);
+		const int32 NumEntities = InstanceTransforms.Num();
+		
+		if (ManagerLocation.IsNearlyZero())
+		{
+			// a tiny bit hacky conversion of InstanceTransforms to FMassTransformsSpawnData. Static assert here should catch if the type doesn't match
+			static_assert(sizeof(FMassTransformsSpawnData::FTransformsContainerType) == sizeof(InstanceTransforms)
+				// @todo use some "TTypeOf" template 
+				, "the InstanceTransforms-to-FMassTransformsSpawnData conversion relies on FMassTransformsSpawnData having only a single property that matches InstanceTransforms' type exactly");
+			FConstStructView SpawnLocationDataView(FMassTransformsSpawnData::StaticStruct(), reinterpret_cast<const uint8*>(&InstanceTransforms));
+
+			SpawnerSystem->SpawnEntities(FinalizedTemplate->GetTemplateID(), NumEntities, SpawnLocationDataView, UMassSpawnLocationProcessor::StaticClass(), Entities);
+		}
+		else
+		{
+			FMassTransformsSpawnData SpawnLocationData;
+			SpawnLocationData.Transforms = InstanceTransforms;
+			for (FTransform& EntityTransform : SpawnLocationData.Transforms)
+			{
+				EntityTransform.AddToTranslation(ManagerLocation);
+			}
+
+			FConstStructView SpawnLocationDataView(FMassTransformsSpawnData::StaticStruct(), reinterpret_cast<const uint8*>(&SpawnLocationData));
+			SpawnerSystem->SpawnEntities(FinalizedTemplate->GetTemplateID(), NumEntities, SpawnLocationDataView, UMassSpawnLocationProcessor::StaticClass(), Entities);
+		}
+
+		// destroy actors that have already been created
+		for (const TPair<int32, TObjectPtr<AActor>>& ActorPair : Actors)
+		{
+			if (ActorPair.Value)
+			{
+				ActorPair.Value->OnDestroyed.RemoveAll(this);
+				ActorPair.Value->Destroy();
+			}
+		}
+		Actors.Reset();
+		
+		InstancedStaticMeshComponent->UnregisterComponent();
+		InstancedStaticMeshComponent->DestroyComponent();
+	}
+}
+
+void AMassLWIStaticMeshManager::CreateMassTemplate(FMassEntityManager& EntityManager)
+{
+	UWorld* World = GetWorld();
+	if (World == nullptr)
+	{
+		return;
+	}
+
+	UMassSpawnerSubsystem* SpawnerSystem = World->GetSubsystem<UMassSpawnerSubsystem>(World);
+	if (!ensure(SpawnerSystem))
+	{
+		//+complain
+		return;
+	}
+
+	UMassLWISubsystem* LWIxMass = World->GetSubsystem<UMassLWISubsystem>();
+	if (!ensure(LWIxMass))
+	{
+		return;
+	}
+
+	const FMassEntityConfig* ClassConfig = LWIxMass->GetConfigForClass(RepresentedClass);
+	if (!ClassConfig || ClassConfig->IsEmpty())
+	{
+		UE_VLOG_UELOG(this, LogMass, Log, TEXT("%s failed to find a calid entity config for %s class. This LWI manager won't transfer data to Mass.")
+			, *UObjectBaseUtility::GetName(), *GetNameSafe(RepresentedClass));
+		return;
+	}
+
+	const UMassVisualizationTrait* VisTrait = Cast<const UMassVisualizationTrait>(ClassConfig->FindTrait(UMassVisualizationTrait::StaticClass()));
+	if (!ensureMsgf(VisTrait, TEXT("The config used doesn't contain a VisualizationTrait, which is required for LWIxMass to function")))
+	{
+		return;
+	}
+	if (!ensureMsgf(VisTrait->StaticMeshInstanceDesc.Meshes.Num(), TEXT("The VisualizationTrait.StaticMeshInstanceDesc.Meshes being used needs to have at least one entry")))
+	{
+		return;
+	}
+	ensureMsgf(VisTrait->StaticMeshInstanceDesc.Meshes.Num() == 1
+		, TEXT("It's recommended for the VisualizationTrait.StaticMeshInstanceDesc.Meshes being used to have exactly one entry, since only the first entry will be used"));
+
+	UMassRepresentationSubsystem* RepresentationSubsystem = nullptr;
+	
+	const FMassEntityTemplate& SourceTemplate = ClassConfig->GetOrCreateEntityTemplate(*World);
+	// what we want to do now is to modify the config to point at the specific static mesh and actor class
+	FMassEntityTemplateData NewTemplate(SourceTemplate);
+
+		
+	const FMassArchetypeSharedFragmentValues& SharedFragmentValues = SourceTemplate.GetSharedFragmentValues();
+	const TArray<FSharedStruct>& SharedFragments = SharedFragmentValues.GetSharedFragments();
+	for (FSharedStruct SharedFragment : SharedFragments)
+	{
+		if (FMassRepresentationSubsystemSharedFragment* AsRepresentationSubsystemSharedFragment = SharedFragment.GetPtr<FMassRepresentationSubsystemSharedFragment>())
+		{
+			RepresentationSubsystem = AsRepresentationSubsystemSharedFragment->RepresentationSubsystem;
+			break;
+		}
+	}
+	
+	FStaticMeshInstanceVisualizationDesc StaticMeshInstanceDesc = VisTrait->StaticMeshInstanceDesc;
+	// @todo make sure a nullptr here won't break stuff
+	FMassStaticMeshInstanceVisualizationMeshDesc& MeshDesc = StaticMeshInstanceDesc.Meshes[0];
+	MeshDesc.Mesh = StaticMesh.Get();
+	if (InstancedStaticMeshComponent)
+	{
+		MeshDesc.MaterialOverrides = InstancedStaticMeshComponent->OverrideMaterials;
+	}
+
+	check(MeshDesc.Mesh);
+	NewTemplate.SetTemplateName(MeshDesc.Mesh->GetName());
+
+	if (RepresentationSubsystem)
+	{
+		FMassRepresentationFragment& RepresentationFragment = NewTemplate.AddFragment_GetRef<FMassRepresentationFragment>();
+		RepresentationFragment.StaticMeshDescIndex = RepresentationSubsystem->FindOrAddStaticMeshDesc(StaticMeshInstanceDesc);
+		const int32 TemplateActorIndex = RepresentationSubsystem->FindOrAddTemplateActor(RepresentedClass);
+		RepresentationFragment.HighResTemplateActorIndex = TemplateActorIndex;
+		//RepresentationFragment.LowResTemplateActorIndex = TemplateActorIndex;
+	}
+
+	FMassEntityTemplateRegistry& TemplateRegistry = SpawnerSystem->GetMutableTemplateRegistryInstance();
+	const uint32 FlavorHash = GetTypeHash(GetNameSafe(RepresentedClass));
+	MassTemplateID = FMassEntityTemplateIDFactory::MakeFlavor(SourceTemplate.GetTemplateID(), FlavorHash);
+
+	const TSharedRef<FMassEntityTemplate>& NewFinalizedTemplate = TemplateRegistry.FindOrAddTemplate(MassTemplateID, MoveTemp(NewTemplate));
+	MassTemplateID = NewFinalizedTemplate->GetTemplateID();
+	FinalizedTemplate = NewFinalizedTemplate;
+
+}
+
+void AMassLWIStaticMeshManager::MarkRegisteredWithMass(const FMassLWIManagerRegistrationHandle RegistrationIndex)
+{
+	check(MassRegistrationHandle.IsValid() == false && RegistrationIndex.IsValid() == true);
+	new (&MassRegistrationHandle) FMassLWIManagerRegistrationHandle(RegistrationIndex);
+}
+
+void AMassLWIStaticMeshManager::MarkUnregisteredWithMass()
+{
+	new (&MassRegistrationHandle) FMassLWIManagerRegistrationHandle();
+}
diff --git a/Engine/Plugins/Runtime/MassGameplay/Source/MassLWI/Private/MassLWISubsystem.cpp b/Engine/Plugins/Runtime/MassGameplay/Source/MassLWI/Private/MassLWISubsystem.cpp
new file mode 100644
--- /dev/null
+++ b/Engine/Plugins/Runtime/MassGameplay/Source/MassLWI/Private/MassLWISubsystem.cpp
@@ -0,0 +1,131 @@
+// Copyright Epic Games, Inc. All Rights Reserved.
+
+#include "MassLWISubsystem.h"
+#include "EngineUtils.h"
+#include "MassLWITypes.h"
+#include "MassLWIStaticMeshManager.h"
+#include "MassLWIConfigActor.h"
+#include "MassEntitySubsystem.h"
+#include "VisualLogger/VisualLogger.h"
+
+
+#include UE_INLINE_GENERATED_CPP_BY_NAME(MassLWISubsystem)
+
+//-----------------------------------------------------------------------------
+// UMassLWISubsystem
+//-----------------------------------------------------------------------------
+UMassLWISubsystem::UMassLWISubsystem()
+{
+
+}
+
+void UMassLWISubsystem::Initialize(FSubsystemCollectionBase& Collection)
+{
+	Super::Initialize(Collection);
+
+	UMassEntitySubsystem* EntitySubsystem = Collection.InitializeDependency<UMassEntitySubsystem>();
+	check(EntitySubsystem);
+	EntityManager = EntitySubsystem->GetMutableEntityManager().AsShared();
+}
+
+void UMassLWISubsystem::PostInitialize()
+{
+	Super::PostInitialize();
+
+	RegisterExistingConfigActors();
+}
+
+void UMassLWISubsystem::Deinitialize()
+{
+	EntityManager = nullptr;
+	Super::Deinitialize();
+}
+
+bool UMassLWISubsystem::ShouldCreateSubsystem(UObject* Outer) const
+{
+	UWorld* World = Outer ? Outer->GetWorld() : nullptr;
+	return World && World->IsGameWorld() && Super::ShouldCreateSubsystem(Outer);
+}
+
+void UMassLWISubsystem::RegisterLWIManager(AMassLWIStaticMeshManager& Manager)
+{
+	check(EntityManager);
+	if (Manager.IsRegisteredWithMass())
+	{
+		return;
+	}
+	checkSlow(RegisteredManagers.Find(&Manager) == INDEX_NONE);
+
+	REDIRECT_OBJECT_TO_VLOG(&Manager, this);
+	
+	int32 RegistrationIndex = INDEX_NONE;
+	if (FreeIndices.Num() > 0)
+	{
+		RegistrationIndex = FreeIndices.Pop(/*bAllowShrinking = */true);
+	}
+	else
+	{
+		RegistrationIndex = RegisteredManagers.Add(&Manager);
+	}
+
+	Manager.TransferDataToMass(*EntityManager);
+	Manager.MarkRegisteredWithMass(FMassLWIManagerRegistrationHandle(RegistrationIndex));
+}
+
+void UMassLWISubsystem::UnregisterLWIManager(AMassLWIStaticMeshManager& Manager)
+{
+	if (Manager.IsRegisteredWithMass() == false)
+	{
+		return;
+	}
+
+	FMassLWIManagerRegistrationHandle IndexHandle = Manager.GetMassRegistrationHandle();
+	if (RegisteredManagers.IsValidIndex(IndexHandle))
+	{
+		FreeIndices.Add(IndexHandle);
+		RegisteredManagers[IndexHandle] = nullptr;
+	}
+
+	Manager.MarkUnregisteredWithMass();
+}
+
+void UMassLWISubsystem::RegisterConfigActor(AMassLWIConfigActor& ConfigActor)
+{
+	if (ConfigActor.IsRegisteredWithSubsystem())
+	{
+		return;
+	}
+
+	if (DefaultConfig.IsEmpty() && !ConfigActor.DefaultConfig.IsEmpty())
+	{
+		DefaultConfig = ConfigActor.DefaultConfig;
+	}
+
+}
+
+void UMassLWISubsystem::RegisterExistingConfigActors()
+{
+	UWorld* World = GetWorld();
+	if (World == nullptr)
+	{
+		return;
+	}
+
+	for (TActorIterator<AMassLWIConfigActor> It(World); It; ++It)
+	{
+		RegisterConfigActor(**It);
+	}
+}
+
+const FMassEntityConfig* UMassLWISubsystem::GetConfigForClass(TSubclassOf<AActor> ActorClass) const
+{
+	if (const FMassEntityConfig* ConfigFound = ClassToConfigMap.Find(ActorClass))
+	{
+		return ConfigFound;
+	}
+	else if (DefaultConfig.IsEmpty() == false)
+	{
+		return &DefaultConfig;
+	}
+	return nullptr;
+}
diff --git a/Engine/Plugins/Runtime/MassGameplay/Source/MassLWI/Private/MassLWITypes.cpp b/Engine/Plugins/Runtime/MassGameplay/Source/MassLWI/Private/MassLWITypes.cpp
new file mode 100644
--- /dev/null
+++ b/Engine/Plugins/Runtime/MassGameplay/Source/MassLWI/Private/MassLWITypes.cpp
@@ -0,0 +1,3 @@
+// Copyright Epic Games, Inc. All Rights Reserved.
+
+#include "MassLWITypes.h"
diff --git a/Engine/Plugins/Runtime/MassGameplay/Source/MassLWI/Private/MassLWIVisualizationTrait.cpp b/Engine/Plugins/Runtime/MassGameplay/Source/MassLWI/Private/MassLWIVisualizationTrait.cpp
new file mode 100644
--- /dev/null
+++ b/Engine/Plugins/Runtime/MassGameplay/Source/MassLWI/Private/MassLWIVisualizationTrait.cpp
@@ -0,0 +1,17 @@
+// Copyright Epic Games, Inc. All Rights Reserved.
+
+#include "MassLWIVisualizationTrait.h"
+#include "MassEntityTemplateRegistry.h"
+#include "MassLWIRepresentationSubsystem.h"
+#include "MassLWIRepresentationActorManagement.h"
+//#include "MassInstancedStaticMeshComponent.h"
+
+
+UMassLWIVisualizationTrait::UMassLWIVisualizationTrait(const FObjectInitializer& ObjectInitializer)
+{
+	bAllowServerSideVisualization = true;
+	bCanModifyRepresentationSubsystemClass = false;
+	RepresentationSubsystemClass = UMassLWIRepresentationSubsystem::StaticClass();
+	Params.bCanModifyRepresentationActorManagementClass = false;
+	Params.RepresentationActorManagementClass = UMassLWIRepresentationActorManagement::StaticClass();
+}
diff --git a/Engine/Plugins/Runtime/MassGameplay/Source/MassLWI/Public/IMassLWIModule.h b/Engine/Plugins/Runtime/MassGameplay/Source/MassLWI/Public/IMassLWIModule.h
new file mode 100644
--- /dev/null
+++ b/Engine/Plugins/Runtime/MassGameplay/Source/MassLWI/Public/IMassLWIModule.h
@@ -0,0 +1,39 @@
+// Copyright Epic Games, Inc. All Rights Reserved.
+
+#pragma once
+
+#include "Modules/ModuleInterface.h"
+#include "Modules/ModuleManager.h"
+
+
+/**
+ * The public interface to this module.  In most cases, this interface is only public to sibling modules 
+ * within this plugin.
+ */
+class IMassLWIModule : public IModuleInterface
+{
+
+public:
+
+	/**
+	 * Singleton-like access to this module's interface.  This is just for convenience!
+	 * Beware of calling this during the shutdown phase, though.  Your module might have been unloaded already.
+	 *
+	 * @return Returns singleton instance, loading the module on demand if needed
+	 */
+	static inline IMassLWIModule& Get()
+	{
+		return FModuleManager::LoadModuleChecked<IMassLWIModule>("MassLWI");
+	}
+
+	/**
+	 * Checks to see if this module is loaded and ready.  It is only valid to call Get() if IsAvailable() returns true.
+	 *
+	 * @return True if the module is loaded and ready to use
+	 */
+	static inline bool IsAvailable()
+	{
+		return FModuleManager::Get().IsModuleLoaded("MassLWI");
+	}
+};
+
diff --git a/Engine/Plugins/Runtime/MassGameplay/Source/MassLWI/Public/MassLWIConfigActor.h b/Engine/Plugins/Runtime/MassGameplay/Source/MassLWI/Public/MassLWIConfigActor.h
new file mode 100644
--- /dev/null
+++ b/Engine/Plugins/Runtime/MassGameplay/Source/MassLWI/Public/MassLWIConfigActor.h
@@ -0,0 +1,31 @@
+// Copyright Epic Games, Inc. All Rights Reserved.
+
+#pragma once
+
+#include "GameFramework/LightWeightInstanceStaticMeshManager.h"
+#include "MassEntityConfigAsset.h"
+#include "MassLWIConfigActor.generated.h"
+
+
+class UMassLWISubsystem;
+
+UCLASS(hidecategories = (Actor, Input, Collision, Rendering, Replication, Partition, HLOD, Cooking))
+class AMassLWIConfigActor : public AActor
+{
+	GENERATED_BODY()
+public:
+	virtual void PostLoad() override;
+
+	bool IsRegisteredWithSubsystem() const { return bRegisteredWithSubsystem; }
+
+protected:
+	UPROPERTY(Category = "Derived Traits", EditAnywhere)
+	FMassEntityConfig DefaultConfig;
+
+	UPROPERTY(Category = "Derived Traits", EditAnywhere)
+	TMap<TObjectPtr<UClass>, FMassEntityConfig> PerClassConfigs;
+
+	uint8 bRegisteredWithSubsystem : 1;
+
+	friend UMassLWISubsystem;
+};
diff --git a/Engine/Plugins/Runtime/MassGameplay/Source/MassLWI/Public/MassLWIRepresentationActorManagement.h b/Engine/Plugins/Runtime/MassGameplay/Source/MassLWI/Public/MassLWIRepresentationActorManagement.h
new file mode 100644
--- /dev/null
+++ b/Engine/Plugins/Runtime/MassGameplay/Source/MassLWI/Public/MassLWIRepresentationActorManagement.h
@@ -0,0 +1,20 @@
+// Copyright Epic Games, Inc. All Rights Reserved.
+
+#pragma once
+
+#include "MassRepresentationActorManagement.h"
+#include "MassLWIRepresentationActorManagement.generated.h"
+
+
+UCLASS()
+class MASSLWI_API UMassLWIRepresentationActorManagement : public UMassRepresentationActorManagement
+{
+	GENERATED_BODY()
+
+public:
+
+protected:
+	virtual AActor* GetOrSpawnActor(UMassRepresentationSubsystem& RepresentationSubsystem, FMassEntityManager& EntitySubsystem
+		, const FMassEntityHandle MassAgent, FMassActorFragment& ActorInfo, const FTransform& Transform, const int16 TemplateActorIndex
+		, FMassActorSpawnRequestHandle& SpawnRequestHandle, const float Priority) const override;
+};
diff --git a/Engine/Plugins/Runtime/MassGameplay/Source/MassLWI/Public/MassLWIRepresentationSubsystem.h b/Engine/Plugins/Runtime/MassGameplay/Source/MassLWI/Public/MassLWIRepresentationSubsystem.h
new file mode 100644
--- /dev/null
+++ b/Engine/Plugins/Runtime/MassGameplay/Source/MassLWI/Public/MassLWIRepresentationSubsystem.h
@@ -0,0 +1,35 @@
+// Copyright Epic Games, Inc. All Rights Reserved.
+
+#pragma once
+
+#include "MassRepresentationSubsystem.h"
+#include "MassExternalSubsystemTraits.h"
+#include "MassLWIRepresentationSubsystem.generated.h"
+
+
+UCLASS()
+class MASSLWI_API UMassLWIRepresentationSubsystem : public UMassRepresentationSubsystem
+{
+	GENERATED_BODY()
+
+public:
+	DECLARE_MULTICAST_DELEGATE_OneParam(FOnActorSpawned, AActor&);
+
+	AActor* GetOrSpawnActorFromTemplate(const FMassEntityHandle MassAgent, const FTransform& Transform, const int16 TemplateActorIndex
+		, FMassActorSpawnRequestHandle& SpawnRequestHandle, float Priority, FMassActorPreSpawnDelegate ActorPreSpawnDelegate
+		, FMassActorPostSpawnDelegate ActorPostSpawnDelegate);
+
+	FOnActorSpawned& GetOnActorSpawnedDelegate() { return OnActorSpawnedDelegate; }
+
+protected:
+	FOnActorSpawned OnActorSpawnedDelegate;
+};
+
+template<>
+struct TMassExternalSubsystemTraits<UMassLWIRepresentationSubsystem> final
+{
+	enum
+	{
+		GameThreadOnly = true
+	};
+};
diff --git a/Engine/Plugins/Runtime/MassGameplay/Source/MassLWI/Public/MassLWIStaticMeshManager.h b/Engine/Plugins/Runtime/MassGameplay/Source/MassLWI/Public/MassLWIStaticMeshManager.h
new file mode 100644
--- /dev/null
+++ b/Engine/Plugins/Runtime/MassGameplay/Source/MassLWI/Public/MassLWIStaticMeshManager.h
@@ -0,0 +1,56 @@
+// Copyright Epic Games, Inc. All Rights Reserved.
+
+#pragma once
+
+#include "GameFramework/LightWeightInstanceStaticMeshManager.h"
+#include "MassEntityConfigAsset.h"
+#include "MassEntityTemplate.h"
+#include "MassLWIStaticMeshManager.generated.h"
+
+
+struct MASSLWI_API FMassLWIManagerRegistrationHandle
+{
+	explicit FMassLWIManagerRegistrationHandle(const int32 InRegisteredIndex = INDEX_NONE)
+		: RegisteredIndex(InRegisteredIndex)
+	{}
+
+	operator int32() const { return RegisteredIndex; }
+	bool IsValid() const { return RegisteredIndex != INDEX_NONE; }
+
+private:
+	const int32 RegisteredIndex = INDEX_NONE;
+};
+
+UCLASS()
+class MASSLWI_API AMassLWIStaticMeshManager : public ALightWeightInstanceStaticMeshManager
+{
+	GENERATED_BODY()
+
+public:
+	bool IsRegisteredWithMass() const { return MassRegistrationHandle.IsValid(); }
+
+	void TransferDataToMass(FMassEntityManager& EntityManager);
+	FMassLWIManagerRegistrationHandle GetMassRegistrationHandle() const { return MassRegistrationHandle; }
+	void MarkRegisteredWithMass(const FMassLWIManagerRegistrationHandle RegistrationIndex);
+	void MarkUnregisteredWithMass();
+
+protected:
+	// UObject API
+	virtual void PostLoad() override;
+	// End UObject API
+
+	// AActor API
+	virtual void BeginPlay() override;
+	virtual void EndPlay(const EEndPlayReason::Type EndPlayReason) override;
+	// End AActor API
+
+	void CreateMassTemplate(FMassEntityManager& EntityManager);
+
+
+	FMassLWIManagerRegistrationHandle MassRegistrationHandle;
+	FMassEntityTemplateID MassTemplateID;
+	TSharedPtr<FMassEntityTemplate> FinalizedTemplate;
+	FMassArchetypeHandle TargetArchetype;
+	TArray<FMassEntityHandle> Entities;
+};
+\
\ No newline at end of file
diff --git a/Engine/Plugins/Runtime/MassGameplay/Source/MassLWI/Public/MassLWISubsystem.h b/Engine/Plugins/Runtime/MassGameplay/Source/MassLWI/Public/MassLWISubsystem.h
new file mode 100644
--- /dev/null
+++ b/Engine/Plugins/Runtime/MassGameplay/Source/MassLWI/Public/MassLWISubsystem.h
@@ -0,0 +1,64 @@
+// Copyright Epic Games, Inc. All Rights Reserved.
+
+#pragma once
+
+#include "MassEntityManager.h"
+#include "Subsystems/WorldSubsystem.h"
+#include "MassEntityConfigAsset.h"
+#include "MassRepresentationSubsystem.h"
+#include "MassLWISubsystem.generated.h"
+
+
+struct FActorInstanceHandle;
+class AMassLWIStaticMeshManager;
+class AMassLWIConfigActor;
+struct FMassEntityManager;
+
+
+UCLASS()
+class MASSLWI_API UMassLWISubsystem : public UWorldSubsystem
+{
+	GENERATED_BODY()
+
+	//~USubsystem interface
+	virtual void Initialize(FSubsystemCollectionBase& Collection) override;
+	virtual void PostInitialize() override;
+	virtual void Deinitialize() override;
+	virtual bool ShouldCreateSubsystem(UObject* Outer) const override;
+	//~End of USubsystem interface
+
+public:
+	UMassLWISubsystem();
+
+	void RegisterLWIManager(AMassLWIStaticMeshManager& Manager);
+	void UnregisterLWIManager(AMassLWIStaticMeshManager& Manager);
+
+	const FMassEntityConfig* GetConfigForClass(TSubclassOf<AActor> ActorClass) const;
+
+	void RegisterConfigActor(AMassLWIConfigActor& ConfigActor);
+
+protected:
+	void RegisterExistingConfigActors();
+
+	UPROPERTY()
+	TArray<TObjectPtr<AMassLWIStaticMeshManager>> RegisteredManagers;
+
+	TArray<int32> FreeIndices;
+
+	UPROPERTY()
+	FMassEntityConfig DefaultConfig;
+
+	TMap<UClass*, FMassEntityConfig> ClassToConfigMap;
+
+private:
+	TSharedPtr<FMassEntityManager> EntityManager;
+};
+
+template<>
+struct TMassExternalSubsystemTraits<UMassLWISubsystem> final
+{
+	enum
+	{
+		GameThreadOnly = true
+	};
+};
diff --git a/Engine/Plugins/Runtime/MassGameplay/Source/MassLWI/Public/MassLWITypes.h b/Engine/Plugins/Runtime/MassGameplay/Source/MassLWI/Public/MassLWITypes.h
new file mode 100644
--- /dev/null
+++ b/Engine/Plugins/Runtime/MassGameplay/Source/MassLWI/Public/MassLWITypes.h
@@ -0,0 +1,8 @@
+// Copyright Epic Games, Inc. All Rights Reserved.
+
+#pragma once
+
+#include "Logging/LogMacros.h"
+
+
+DECLARE_LOG_CATEGORY_EXTERN(LogMassLWI, Log, All)
diff --git a/Engine/Plugins/Runtime/MassGameplay/Source/MassLWI/Public/MassLWIVisualizationTrait.h b/Engine/Plugins/Runtime/MassGameplay/Source/MassLWI/Public/MassLWIVisualizationTrait.h
new file mode 100644
--- /dev/null
+++ b/Engine/Plugins/Runtime/MassGameplay/Source/MassLWI/Public/MassLWIVisualizationTrait.h
@@ -0,0 +1,15 @@
+// Copyright Epic Games, Inc. All Rights Reserved.
+
+#pragma once 
+
+#include "MassStationaryVisualizationTrait.h"
+#include "MassLWIVisualizationTrait.generated.h"
+
+
+UCLASS()
+class MASSLWI_API UMassLWIVisualizationTrait : public UMassStationaryVisualizationTrait
+{
+	GENERATED_BODY()
+public:
+	UMassLWIVisualizationTrait(const FObjectInitializer& ObjectInitializer = FObjectInitializer::Get());
+};
