Fixed FMassEntityConfig's default constructor not setting it's FGuid member to a valid guid.

diff --git a/Engine/Plugins/Runtime/MassGameplay/Source/MassSpawner/Private/MassEntityConfigAsset.cpp b/Engine/Plugins/Runtime/MassGameplay/Source/MassSpawner/Private/MassEntityConfigAsset.cpp
--- a/Engine/Plugins/Runtime/MassGameplay/Source/MassSpawner/Private/MassEntityConfigAsset.cpp
+++ b/Engine/Plugins/Runtime/MassGameplay/Source/MassSpawner/Private/MassEntityConfigAsset.cpp
@@ -21,6 +21,11 @@
 //-----------------------------------------------------------------------------
 // FMassEntityConfig
 //-----------------------------------------------------------------------------
+FMassEntityConfig::FMassEntityConfig()
+{
+	ConfigGuid = FGuid::NewGuid();
+}
+
 FMassEntityConfig::FMassEntityConfig(UObject& InOwner)
 	: ConfigOwner(&InOwner)
 {
diff --git a/Engine/Plugins/Runtime/MassGameplay/Source/MassSpawner/Public/MassEntityConfigAsset.h b/Engine/Plugins/Runtime/MassGameplay/Source/MassSpawner/Public/MassEntityConfigAsset.h
--- a/Engine/Plugins/Runtime/MassGameplay/Source/MassSpawner/Public/MassEntityConfigAsset.h
+++ b/Engine/Plugins/Runtime/MassGameplay/Source/MassSpawner/Public/MassEntityConfigAsset.h
@@ -23,7 +23,7 @@ struct MASSSPAWNER_API FMassEntityConfig
 
 	GENERATED_BODY()
 
-	FMassEntityConfig() = default;
+	FMassEntityConfig();
 	FMassEntityConfig(UObject& InOwner);
 
 	/** Create entity template based on the features included in this config.
