CIS fix

diff --git a/Engine/Plugins/Runtime/MassGameplay/Source/MassLWI/Private/MassLWIVisualizationTrait.cpp b/Engine/Plugins/Runtime/MassGameplay/Source/MassLWI/Private/MassLWIVisualizationTrait.cpp
--- a/Engine/Plugins/Runtime/MassGameplay/Source/MassLWI/Private/MassLWIVisualizationTrait.cpp
+++ b/Engine/Plugins/Runtime/MassGameplay/Source/MassLWI/Private/MassLWIVisualizationTrait.cpp
@@ -10,8 +10,10 @@
 UMassLWIVisualizationTrait::UMassLWIVisualizationTrait(const FObjectInitializer& ObjectInitializer)
 {
 	bAllowServerSideVisualization = true;
+#if WITH_EDITORONLY_DATA
 	bCanModifyRepresentationSubsystemClass = false;
 	RepresentationSubsystemClass = UMassLWIRepresentationSubsystem::StaticClass();
 	Params.bCanModifyRepresentationActorManagementClass = false;
 	Params.RepresentationActorManagementClass = UMassLWIRepresentationActorManagement::StaticClass();
+#endif // WITH_EDITORONLY_DATA
 }
diff --git a/Engine/Plugins/Runtime/MassGameplay/Source/MassRepresentation/Private/MassVisualizationTrait.cpp b/Engine/Plugins/Runtime/MassGameplay/Source/MassRepresentation/Private/MassVisualizationTrait.cpp
--- a/Engine/Plugins/Runtime/MassGameplay/Source/MassRepresentation/Private/MassVisualizationTrait.cpp
+++ b/Engine/Plugins/Runtime/MassGameplay/Source/MassRepresentation/Private/MassVisualizationTrait.cpp
@@ -40,12 +40,14 @@ UMassVisualizationTrait::UMassVisualizationTrait()
 	LODParams.BufferHysteresisOnDistancePercentage = 10.0f;
 	LODParams.DistanceToFrustum = 0.0f;
 	LODParams.DistanceToFrustumHysteresis = 0.0f;
+
+	bAllowServerSideVisualization = false;
 }
 
 void UMassVisualizationTrait::BuildTemplate(FMassEntityTemplateBuildContext& BuildContext, const UWorld& World) const
 {
 	// This should not be ran on NM_Server network mode
-	if (World.IsNetMode(NM_DedicatedServer))
+	if (World.IsNetMode(NM_DedicatedServer) && !bAllowServerSideVisualization)
 	{
 		return;
 	}
diff --git a/Engine/Plugins/Runtime/MassGameplay/Source/MassRepresentation/Public/MassVisualizationTrait.h b/Engine/Plugins/Runtime/MassGameplay/Source/MassRepresentation/Public/MassVisualizationTrait.h
--- a/Engine/Plugins/Runtime/MassGameplay/Source/MassRepresentation/Public/MassVisualizationTrait.h
+++ b/Engine/Plugins/Runtime/MassGameplay/Source/MassRepresentation/Public/MassVisualizationTrait.h
@@ -47,6 +47,11 @@ public:
 	UPROPERTY(EditAnywhere, Category = "Mass|Visual")
 	FMassVisualizationLODParameters LODParams;
 
+	/** If set to true will result in the visualization-related fragments being added to server-size entities as well.
+	 *  By default only the clients require visualization fragments */
+	UPROPERTY(EditAnywhere, Category = "Mass|Visual")
+	bool bAllowServerSideVisualization = false;
+
 #if WITH_EDITORONLY_DATA
 	/** the property is marked like this to ensure it won't show up in UI */
 	UPROPERTY(EditDefaultsOnly, Category = "Mass|Visual")
diff --git a/Engine/Plugins/Runtime/MassGameplay/Source/MassSpawner/Public/MassSpawnerTypes.h b/Engine/Plugins/Runtime/MassGameplay/Source/MassSpawner/Public/MassSpawnerTypes.h
--- a/Engine/Plugins/Runtime/MassGameplay/Source/MassSpawner/Public/MassSpawnerTypes.h
+++ b/Engine/Plugins/Runtime/MassGameplay/Source/MassSpawner/Public/MassSpawnerTypes.h
@@ -25,7 +25,9 @@ struct FMassTransformsSpawnData
 {
 	GENERATED_BODY()
 
-	TArray<FTransform> Transforms;
+	// declaring the type used to be able to statically test it against other types 
+	using FTransformsContainerType = TArray<FTransform>;
+	FTransformsContainerType Transforms;
 };
 
 /**
