CIS fix

diff --git a/Engine/Plugins/Runtime/MassGameplay/Source/MassLWI/Public/MassLWIStaticMeshManager.h b/Engine/Plugins/Runtime/MassGameplay/Source/MassLWI/Public/MassLWIStaticMeshManager.h
--- a/Engine/Plugins/Runtime/MassGameplay/Source/MassLWI/Public/MassLWIStaticMeshManager.h
+++ b/Engine/Plugins/Runtime/MassGameplay/Source/MassLWI/Public/MassLWIStaticMeshManager.h
@@ -53,4 +53,3 @@ protected:
 	FMassArchetypeHandle TargetArchetype;
 	TArray<FMassEntityHandle> Entities;
 };
-\
\ No newline at end of file
