Mass long-code-lines styling

diff --git a/Engine/Plugins/Runtime/MassGameplay/Source/MassRepresentation/Private/MassRepresentationProcessor.cpp b/Engine/Plugins/Runtime/MassGameplay/Source/MassRepresentation/Private/MassRepresentationProcessor.cpp
--- a/Engine/Plugins/Runtime/MassGameplay/Source/MassRepresentation/Private/MassRepresentationProcessor.cpp
+++ b/Engine/Plugins/Runtime/MassGameplay/Source/MassRepresentation/Private/MassRepresentationProcessor.cpp
@@ -64,7 +64,7 @@ void UMassRepresentationProcessor::UpdateRepresentation(FMassExecutionContext& C
 	const int32 NumEntities = Context.GetNumEntities();
 	for (int32 EntityIdx = 0; EntityIdx < NumEntities; EntityIdx++)
 	{
-		const FMassEntityHandle MassAgent = Context.GetEntity(EntityIdx);
+		const FMassEntityHandle MassAgentEntityHandle = Context.GetEntity(EntityIdx);
 		const FTransformFragment& TransformFragment = TransformList[EntityIdx];
 		const FMassRepresentationLODFragment& RepresentationLOD = RepresentationLODList[EntityIdx];
 		FMassRepresentationFragment& Representation = RepresentationList[EntityIdx];
@@ -92,12 +92,12 @@ void UMassRepresentationProcessor::UpdateRepresentation(FMassExecutionContext& C
 				if (Representation.HighResTemplateActorIndex != Representation.LowResTemplateActorIndex || !RepresentationParams.bKeepLowResActors)
 				{
 					// Try releasing the high actor or any high res spawning request
-					if (ReleaseActorOrCancelSpawning(*RepresentationSubsystem, MassActorSubsystem, MassAgent, ActorInfo, Representation.HighResTemplateActorIndex, Representation.ActorSpawnRequestHandle, Context.Defer()))
+					if (ReleaseActorOrCancelSpawning(*RepresentationSubsystem, MassActorSubsystem, MassAgentEntityHandle, ActorInfo, Representation.HighResTemplateActorIndex, Representation.ActorSpawnRequestHandle, Context.Defer()))
 					{
 						Actor = ActorInfo.GetOwnedByMassMutable();
 					}
 					// Do not do the same with low res if indicated so
-					if (!RepresentationParams.bKeepLowResActors && ReleaseActorOrCancelSpawning(*RepresentationSubsystem, MassActorSubsystem, MassAgent, ActorInfo, Representation.LowResTemplateActorIndex, Representation.ActorSpawnRequestHandle, Context.Defer()))
+					if (!RepresentationParams.bKeepLowResActors && ReleaseActorOrCancelSpawning(*RepresentationSubsystem, MassActorSubsystem, MassAgentEntityHandle, ActorInfo, Representation.LowResTemplateActorIndex, Representation.ActorSpawnRequestHandle, Context.Defer()))
 					{
 						Actor = ActorInfo.GetOwnedByMassMutable();
 					}
@@ -105,7 +105,7 @@ void UMassRepresentationProcessor::UpdateRepresentation(FMassExecutionContext& C
 				// If we already queued spawn request but have changed our mind, continue with it but once we get the actor back, disable it immediately
 				if (Representation.ActorSpawnRequestHandle.IsValid())
 				{
-					Actor = RepresentationActorManagement->GetOrSpawnActor(*RepresentationSubsystem, CachedEntityManager, MassAgent, ActorInfo, TransformFragment.GetTransform(), Representation.LowResTemplateActorIndex, Representation.ActorSpawnRequestHandle, RepresentationActorManagement->GetSpawnPriority(RepresentationLOD));
+					Actor = RepresentationActorManagement->GetOrSpawnActor(*RepresentationSubsystem, CachedEntityManager, MassAgentEntityHandle, ActorInfo, TransformFragment.GetTransform(), Representation.LowResTemplateActorIndex, Representation.ActorSpawnRequestHandle, RepresentationActorManagement->GetSpawnPriority(RepresentationLOD));
 				}
 			}
 			if (Actor != nullptr)
@@ -140,7 +140,7 @@ void UMassRepresentationProcessor::UpdateRepresentation(FMassExecutionContext& C
 						// If the low res is different than the high res, cancel any pending spawn request that is the opposite of what is needed.
 						if (Representation.LowResTemplateActorIndex != Representation.HighResTemplateActorIndex)
 						{
-							ReleaseActorOrCancelSpawning(*RepresentationSubsystem, MassActorSubsystem, MassAgent, ActorInfo, bHighResActor ? Representation.LowResTemplateActorIndex : Representation.HighResTemplateActorIndex, Representation.ActorSpawnRequestHandle, Context.Defer(), /*bCancelSpawningOnly*/true);
+							ReleaseActorOrCancelSpawning(*RepresentationSubsystem, MassActorSubsystem, MassAgentEntityHandle, ActorInfo, bHighResActor ? Representation.LowResTemplateActorIndex : Representation.HighResTemplateActorIndex, Representation.ActorSpawnRequestHandle, Context.Defer(), /*bCancelSpawningOnly*/true);
 							Actor = ActorInfo.GetOwnedByMassMutable();
 						}
 
@@ -152,7 +152,7 @@ void UMassRepresentationProcessor::UpdateRepresentation(FMassExecutionContext& C
 							!RepresentationSubsystem->DoesActorMatchTemplate(*Actor, WantedTemplateActorIndex) ||
 							Representation.ActorSpawnRequestHandle.IsValid())
 						{
-							NewActor = RepresentationActorManagement->GetOrSpawnActor(*RepresentationSubsystem, CachedEntityManager, MassAgent, ActorInfo, TransformFragment.GetTransform(), WantedTemplateActorIndex, Representation.ActorSpawnRequestHandle, RepresentationActorManagement->GetSpawnPriority(RepresentationLOD));
+							NewActor = RepresentationActorManagement->GetOrSpawnActor(*RepresentationSubsystem, CachedEntityManager, MassAgentEntityHandle, ActorInfo, TransformFragment.GetTransform(), WantedTemplateActorIndex, Representation.ActorSpawnRequestHandle, RepresentationActorManagement->GetSpawnPriority(RepresentationLOD));
 						}
 						else
 						{
@@ -195,8 +195,8 @@ void UMassRepresentationProcessor::UpdateRepresentation(FMassExecutionContext& C
 					if (!Actor || ActorInfo.IsOwnedByMass())
 					{
 						// Try releasing both, could have an high res spawned actor and a spawning request for a low res one
-						ReleaseActorOrCancelSpawning(*RepresentationSubsystem, MassActorSubsystem, MassAgent, ActorInfo, Representation.LowResTemplateActorIndex, Representation.ActorSpawnRequestHandle, Context.Defer());
-						ReleaseActorOrCancelSpawning(*RepresentationSubsystem, MassActorSubsystem, MassAgent, ActorInfo, Representation.HighResTemplateActorIndex, Representation.ActorSpawnRequestHandle, Context.Defer());
+						ReleaseActorOrCancelSpawning(*RepresentationSubsystem, MassActorSubsystem, MassAgentEntityHandle, ActorInfo, Representation.LowResTemplateActorIndex, Representation.ActorSpawnRequestHandle, Context.Defer());
+						ReleaseActorOrCancelSpawning(*RepresentationSubsystem, MassActorSubsystem, MassAgentEntityHandle, ActorInfo, Representation.HighResTemplateActorIndex, Representation.ActorSpawnRequestHandle, Context.Defer());
 					}
 					else
 					{
@@ -228,7 +228,9 @@ void UMassRepresentationProcessor::Execute(FMassEntityManager& InEntityManager,
 	});
 }
 
-bool UMassRepresentationProcessor::ReleaseActorOrCancelSpawning(UMassRepresentationSubsystem& RepresentationSubsystem, UMassActorSubsystem* MassActorSubsystem, const FMassEntityHandle MassAgent, FMassActorFragment& ActorInfo, const int16 TemplateActorIndex, FMassActorSpawnRequestHandle& SpawnRequestHandle, FMassCommandBuffer& CommandBuffer, bool bCancelSpawningOnly /*= false*/)
+bool UMassRepresentationProcessor::ReleaseActorOrCancelSpawning(UMassRepresentationSubsystem& RepresentationSubsystem, UMassActorSubsystem* MassActorSubsystem
+	, const FMassEntityHandle MassAgent, FMassActorFragment& ActorInfo, const int16 TemplateActorIndex, FMassActorSpawnRequestHandle& SpawnRequestHandle
+	, FMassCommandBuffer& CommandBuffer, bool bCancelSpawningOnly /*= false*/)
 {
 	if (TemplateActorIndex == INDEX_NONE)
 	{
@@ -397,9 +399,9 @@ void UMassRepresentationFragmentDestructor::Execute(FMassEntityManager& EntityMa
 			FMassRepresentationFragment& Representation = RepresentationList[i];
 			FMassActorFragment& ActorInfo = ActorList[i];
 
-			const FMassEntityHandle MassAgent = Context.GetEntity(i);
+			const FMassEntityHandle MassAgentEntityHandle = Context.GetEntity(i);
 
-			UMassRepresentationActorManagement::ReleaseAnyActorOrCancelAnySpawning(*RepresentationSubsystem, MassAgent, ActorInfo, Representation);
+			UMassRepresentationActorManagement::ReleaseAnyActorOrCancelAnySpawning(*RepresentationSubsystem, MassAgentEntityHandle, ActorInfo, Representation);
 		}
 	});
-}
\ No newline at end of file
+}
diff --git a/Engine/Plugins/Runtime/MassGameplay/Source/MassRepresentation/Private/MassRepresentationSubsystem.cpp b/Engine/Plugins/Runtime/MassGameplay/Source/MassRepresentation/Private/MassRepresentationSubsystem.cpp
--- a/Engine/Plugins/Runtime/MassGameplay/Source/MassRepresentation/Private/MassRepresentationSubsystem.cpp
+++ b/Engine/Plugins/Runtime/MassGameplay/Source/MassRepresentation/Private/MassRepresentationSubsystem.cpp
@@ -51,8 +51,9 @@ int16 UMassRepresentationSubsystem::FindOrAddTemplateActor(const TSubclassOf<AAc
 	return (int16)VisualIndex;
 }
 
-AActor* UMassRepresentationSubsystem::GetOrSpawnActorFromTemplate(const FMassEntityHandle MassAgent, const FTransform& Transform, const int16 TemplateActorIndex, FMassActorSpawnRequestHandle& SpawnRequestHandle, float Priority /*= MAX_FLT*/, 
-	FMassActorPreSpawnDelegate ActorPreSpawnDelegate /*= FMassActorPreSpawnDelegate()*/, FMassActorPostSpawnDelegate ActorPostSpawnDelegate /*= FMassActorPostSpawnDelegate()*/)
+AActor* UMassRepresentationSubsystem::GetOrSpawnActorFromTemplate(const FMassEntityHandle MassAgent, const FTransform& Transform
+	, const int16 TemplateActorIndex, FMassActorSpawnRequestHandle& SpawnRequestHandle, float Priority
+	, FMassActorPreSpawnDelegate ActorPreSpawnDelegate, FMassActorPostSpawnDelegate ActorPostSpawnDelegate)
 {
 	UE_MT_SCOPED_READ_ACCESS(TemplateActorsMTAccessDetector);
 	if (!TemplateActors.IsValidIndex(TemplateActorIndex))
diff --git a/Engine/Plugins/Runtime/MassGameplay/Source/MassRepresentation/Public/MassRepresentationActorManagement.h b/Engine/Plugins/Runtime/MassGameplay/Source/MassRepresentation/Public/MassRepresentationActorManagement.h
--- a/Engine/Plugins/Runtime/MassGameplay/Source/MassRepresentation/Public/MassRepresentationActorManagement.h
+++ b/Engine/Plugins/Runtime/MassGameplay/Source/MassRepresentation/Public/MassRepresentationActorManagement.h
@@ -9,7 +9,6 @@ struct FMassCommandBuffer;
 struct FMassEntityManager;
 enum class EMassActorSpawnRequestAction : uint8;
 enum class EMassActorEnabledType : uint8;
-struct FMassActorSpawnRequestHandle;
 struct FConstStructView;
 struct FMassEntityHandle;
 struct FMassActorFragment;
@@ -42,7 +41,9 @@ public:
 	 * @param Priority of this spawn request in comparison with the others, lower value means higher priority
 	 * @return the actor spawned
 	 */
-	virtual AActor* GetOrSpawnActor(UMassRepresentationSubsystem& RepresentationSubsystem, FMassEntityManager& EntitySubsystem, const FMassEntityHandle MassAgent, FMassActorFragment& ActorInfo, const FTransform& Transform, const int16 TemplateActorIndex, FMassActorSpawnRequestHandle& SpawnRequestHandle, const float Priority) const;
+	virtual AActor* GetOrSpawnActor(UMassRepresentationSubsystem& RepresentationSubsystem, FMassEntityManager& EntitySubsystem
+		, const FMassEntityHandle MassAgent, FMassActorFragment& ActorInfo, const FTransform& Transform, const int16 TemplateActorIndex
+		, FMassActorSpawnRequestHandle& SpawnRequestHandle, const float Priority) const;
 
 	/**
 	 * Enable/disable a spawned actor for a mass entity
