The change prevents failing checks when fetching an entity for a given actor, when the entity has already been destroyed.

diff --git a/Engine/Plugins/Runtime/MassGameplay/Source/MassActors/Private/MassActorSubsystem.cpp b/Engine/Plugins/Runtime/MassGameplay/Source/MassActors/Private/MassActorSubsystem.cpp
--- a/Engine/Plugins/Runtime/MassGameplay/Source/MassActors/Private/MassActorSubsystem.cpp
+++ b/Engine/Plugins/Runtime/MassGameplay/Source/MassActors/Private/MassActorSubsystem.cpp
@@ -107,9 +107,11 @@ FMassActorManager::FMassActorManager(const TSharedPtr<FMassEntityManager>& InEnt
 
 FMassEntityHandle FMassActorManager::GetEntityHandleFromActor(const TObjectKey<const AActor> Actor)
 {
+	checkSlow(EntityManager);
+
 	UE_MT_SCOPED_READ_ACCESS(ActorHandleMapDetector);
 	FMassEntityHandle* Entity = ActorHandleMap.Find(Actor);
-	if (!Entity)
+	if (Entity == nullptr || EntityManager->IsEntityValid(*Entity) == false)
 	{
 		return FMassEntityManager::InvalidEntity;
 	}
