Modified Mass actor spawning to allow MassActorSpawnerSubsystem::SpawnActor to result in "pending" status, that in turn lets us perform latent spawning (like requiesting it on client, spawning on the server, and waiting for the server-spawned actor to replicate over).

diff --git a/Engine/Plugins/Runtime/MassGameplay/Source/MassActors/Private/MassActorSpawnerSubsystem.cpp b/Engine/Plugins/Runtime/MassGameplay/Source/MassActors/Private/MassActorSpawnerSubsystem.cpp
--- a/Engine/Plugins/Runtime/MassGameplay/Source/MassActors/Private/MassActorSpawnerSubsystem.cpp
+++ b/Engine/Plugins/Runtime/MassGameplay/Source/MassActors/Private/MassActorSpawnerSubsystem.cpp
@@ -178,7 +178,7 @@ FMassActorSpawnRequestHandle UMassActorSpawnerSubsystem::GetNextRequestToSpawn()
 	return BestSpawnRequestHandle;
 }
 
-AActor* UMassActorSpawnerSubsystem::SpawnOrRetrieveFromPool(FConstStructView SpawnRequestView)
+ESpawnRequestStatus UMassActorSpawnerSubsystem::SpawnOrRetrieveFromPool(FConstStructView SpawnRequestView, TObjectPtr<AActor>& OutSpawnedActor)
 {
 	if (UE::MassActors::bUseActorPooling != 0 && bActorPoolingEnabled)
 	{
@@ -201,14 +201,15 @@ AActor* UMassActorSpawnerSubsystem::SpawnOrRetrieveFromPool(FConstStructView Spa
 				AgentComp->SetPuppetHandle(SpawnRequest.MassAgent);
 			}
 
-			return PooledActor;
+			OutSpawnedActor = PooledActor;
+			return ESpawnRequestStatus::Succeeded;
 		}
 	}
 
-	return SpawnActor(SpawnRequestView);
+	return SpawnActor(SpawnRequestView, OutSpawnedActor);
 }
 
-AActor* UMassActorSpawnerSubsystem::SpawnActor(FConstStructView SpawnRequestView) const
+ ESpawnRequestStatus UMassActorSpawnerSubsystem::SpawnActor(FConstStructView SpawnRequestView, TObjectPtr<AActor>& OutSpawnedActor) const
 {
 	TRACE_CPUPROFILER_EVENT_SCOPE(UMassActorSpawnerSubsystem::SpawnActor);
 
@@ -229,7 +230,8 @@ AActor* UMassActorSpawnerSubsystem::SpawnActor(FConstStructView SpawnRequestView
 			{
 				AgentComp->SetPuppetHandle(SpawnRequest.MassAgent);
 			}
-			return SpawnedActor;
+			OutSpawnedActor = SpawnedActor;
+			return ESpawnRequestStatus::Succeeded;
 		}
 	}
 
@@ -240,7 +242,8 @@ AActor* UMassActorSpawnerSubsystem::SpawnActor(FConstStructView SpawnRequestView
 					SpawnRequest.Transform.GetRotation(),
 					FColor::Red,
 					TEXT("Unable to spawn actor for Mass entity [%s]"), *SpawnRequest.MassAgent.DebugGetDescription());
-	return nullptr;
+
+	return ESpawnRequestStatus::Failed;
 }
 
 void UMassActorSpawnerSubsystem::ProcessPendingSpawningRequest(const double MaxTimeSlicePerTick)
@@ -277,17 +280,18 @@ void UMassActorSpawnerSubsystem::ProcessPendingSpawningRequest(const double MaxT
 			SpawnRequest.ActorPreSpawnDelegate.Execute(SpawnRequestHandle, SpawnRequestView);
 		}
 
-		SpawnRequest.SpawnedActor = SpawnOrRetrieveFromPool(SpawnRequestView);
-
-		SpawnRequest.SpawnStatus = SpawnRequest.SpawnedActor ? ESpawnRequestStatus::Succeeded : ESpawnRequestStatus::Failed;
+		SpawnRequest.SpawnStatus = SpawnOrRetrieveFromPool(SpawnRequestView, SpawnRequest.SpawnedActor);
 
-		// Call the post spawn delegate on the spawn request
-		if (SpawnRequest.ActorPostSpawnDelegate.IsBound())
+		if (SpawnRequest.IsFinished())
 		{
-			if (SpawnRequest.ActorPostSpawnDelegate.Execute(SpawnRequestHandle, SpawnRequestView) == EMassActorSpawnRequestAction::Remove)
+			// Call the post spawn delegate on the spawn request
+			if (SpawnRequest.ActorPostSpawnDelegate.IsBound())
 			{
-				// If notified, remove the spawning request
-				ensureMsgf(SpawnRequestHandleManager.RemoveHandle(SpawnRequestHandle), TEXT("When providing a delegate, the spawn request gets automatically removed, no need to remove it on your side"));
+				if (SpawnRequest.ActorPostSpawnDelegate.Execute(SpawnRequestHandle, SpawnRequestView) == EMassActorSpawnRequestAction::Remove)
+				{
+					// If notified, remove the spawning request
+					ensureMsgf(SpawnRequestHandleManager.RemoveHandle(SpawnRequestHandle), TEXT("When providing a delegate, the spawn request gets automatically removed, no need to remove it on your side"));
+				}
 			}
 		}
 	}
diff --git a/Engine/Plugins/Runtime/MassGameplay/Source/MassActors/Public/MassActorSpawnerSubsystem.h b/Engine/Plugins/Runtime/MassGameplay/Source/MassActors/Public/MassActorSpawnerSubsystem.h
--- a/Engine/Plugins/Runtime/MassGameplay/Source/MassActors/Public/MassActorSpawnerSubsystem.h
+++ b/Engine/Plugins/Runtime/MassGameplay/Source/MassActors/Public/MassActorSpawnerSubsystem.h
@@ -101,6 +101,8 @@ public:
 		SerialNumber = 0;
 		RequestedTime = 0.0f;
 	}
+
+	bool IsFinished() const { return SpawnStatus == ESpawnRequestStatus::Failed || SpawnStatus == ESpawnRequestStatus::Succeeded; }
 };
 
 /**
@@ -194,11 +196,11 @@ protected:
 	 *  @return the next best handle to spawn. */
 	virtual FMassActorSpawnRequestHandle GetNextRequestToSpawn() const;
 
-	virtual AActor* SpawnOrRetrieveFromPool(FConstStructView SpawnRequest);
+	virtual ESpawnRequestStatus SpawnOrRetrieveFromPool(FConstStructView SpawnRequestView, TObjectPtr<AActor>& OutSpawnedActor);
 
 	/** Actual code that will spawn the actor, overridable by subclass if need to be.
 	 *  @return spawned actor if succeeded. */
-	virtual AActor* SpawnActor(FConstStructView SpawnRequest) const;
+	virtual ESpawnRequestStatus SpawnActor(FConstStructView SpawnRequestView, TObjectPtr<AActor>& OutSpawnedActor) const;
 
 	/** Go through the spawning request and spawn them until we reach the budget 
 	 * @param MaxTimeSlicePerTick is the budget in seconds allowed to do spawning */
