MassLWI work in progress.

* added stationary ISM fragment destructor
* added MassLWIClientActorSpawnerSubsystem for handling client-side actor spawning requests by waiting for server-side actor to spawn over
* added destryction of created entities when the source LWI manager is removed (via EndPlay)

diff --git a/Engine/Plugins/Runtime/MassGameplay/Source/MassLWI/MassLWI.Build.cs b/Engine/Plugins/Runtime/MassGameplay/Source/MassLWI/MassLWI.Build.cs
--- a/Engine/Plugins/Runtime/MassGameplay/Source/MassLWI/MassLWI.Build.cs
+++ b/Engine/Plugins/Runtime/MassGameplay/Source/MassLWI/MassLWI.Build.cs
@@ -18,6 +18,7 @@ namespace UnrealBuildTool.Rules
 					"MassReplication",
 					"MassRepresentation",
 					"MassSpawner",
+					"MassActors",
 					"StructUtils",
 					"AIModule",
 					"NetCore",
diff --git a/Engine/Plugins/Runtime/MassGameplay/Source/MassLWI/Private/MassLWIClientActorSpawnerSubsystem.cpp b/Engine/Plugins/Runtime/MassGameplay/Source/MassLWI/Private/MassLWIClientActorSpawnerSubsystem.cpp
new file mode 100644
--- /dev/null
+++ b/Engine/Plugins/Runtime/MassGameplay/Source/MassLWI/Private/MassLWIClientActorSpawnerSubsystem.cpp
@@ -0,0 +1,63 @@
+// Copyright Epic Games, Inc. All Rights Reserved.
+
+#include "MassLWIClientActorSpawnerSubsystem.h"
+#include "Engine/World.h"
+
+
+//-----------------------------------------------------------------------------
+// UMassLWIClientActorSpawnerSubsystem
+//-----------------------------------------------------------------------------
+void UMassLWIClientActorSpawnerSubsystem::Initialize(FSubsystemCollectionBase& Collection)
+{
+	Super::Initialize(Collection);
+
+	UWorld* World = GetWorld();
+	check(World);
+	check(World->GetNetMode() == NM_Client);
+	
+	WorldOnActorSpawnedHandle = World->AddOnActorSpawnedHandler(FOnActorSpawned::FDelegate::CreateUObject(this, &UMassLWIClientActorSpawnerSubsystem::OnActorSpawned));
+}
+
+bool UMassLWIClientActorSpawnerSubsystem::ShouldCreateSubsystem(UObject* Outer) const
+{
+	UWorld* World = Cast<UWorld>(Outer);
+	return (World != nullptr && World->GetNetMode() == NM_Client);
+}
+
+void UMassLWIClientActorSpawnerSubsystem::OnActorSpawned(AActor* InActor)
+{
+	check(InActor);
+	// check if it's an actor we're waiting for
+	if (FMassStoredActorsContainer* StoredActors = PendingActors.Find(InActor->GetClass()))
+	{
+		// getting here means we're waiting for this class of actor
+		StoredActors->Container.Add(InActor);
+	}
+}
+
+ESpawnRequestStatus UMassLWIClientActorSpawnerSubsystem::SpawnActor(FConstStructView SpawnRequestView, TObjectPtr<AActor>& OutSpawnedActor) const
+{
+	constexpr double ComparisonPrecission = 1.0;
+	const FMassActorSpawnRequest& SpawnRequest = SpawnRequestView.Get<const FMassActorSpawnRequest>();
+	TArray<TObjectPtr<AActor>>& StoredActors = PendingActors.FindOrAdd(SpawnRequest.Template).Container;
+
+	for (auto It = StoredActors.CreateIterator(); It; ++It)
+	{
+		if (*It && (*It)->GetActorTransform().Equals(SpawnRequest.Transform, ComparisonPrecission))
+		{
+			OutSpawnedActor = *It;
+			It.RemoveCurrentSwap();
+			return ESpawnRequestStatus::Succeeded;
+		}
+	}
+
+	return ESpawnRequestStatus::Pending;
+}
+
+void UMassLWIClientActorSpawnerSubsystem::RegisterRepresentedClass(UClass* RepresentedClass)
+{
+	if (TSubclassOf<AActor> ActorClass = RepresentedClass)
+	{
+		PendingActors.Add(MoveTemp(ActorClass));
+	}
+}
diff --git a/Engine/Plugins/Runtime/MassGameplay/Source/MassLWI/Private/MassLWIRepresentationActorManagement.cpp b/Engine/Plugins/Runtime/MassGameplay/Source/MassLWI/Private/MassLWIRepresentationActorManagement.cpp
deleted file mode 100644
diff --git a/Engine/Plugins/Runtime/MassGameplay/Source/MassLWI/Private/MassLWIRepresentationSubsystem.cpp b/Engine/Plugins/Runtime/MassGameplay/Source/MassLWI/Private/MassLWIRepresentationSubsystem.cpp
--- a/Engine/Plugins/Runtime/MassGameplay/Source/MassLWI/Private/MassLWIRepresentationSubsystem.cpp
+++ b/Engine/Plugins/Runtime/MassGameplay/Source/MassLWI/Private/MassLWIRepresentationSubsystem.cpp
@@ -1,31 +1,23 @@
 // Copyright Epic Games, Inc. All Rights Reserved.
 
 #include "MassLWIRepresentationSubsystem.h"
-#include "Engine/Note.h"
 #include "Engine/World.h"
+#include "MassLWIClientActorSpawnerSubsystem.h"
 
 
 //-----------------------------------------------------------------------------
 // UMassLWIRepresentationSubsystem
 //-----------------------------------------------------------------------------
-AActor* UMassLWIRepresentationSubsystem::GetOrSpawnActorFromTemplate(const FMassEntityHandle MassAgent, const FTransform& Transform, const int16 TemplateActorIndex
-	, FMassActorSpawnRequestHandle& SpawnRequestHandle, float Priority, FMassActorPreSpawnDelegate ActorPreSpawnDelegate
-	, FMassActorPostSpawnDelegate ActorPostSpawnDelegate)
+void UMassLWIRepresentationSubsystem::Initialize(FSubsystemCollectionBase& Collection)
 {
+	Super::Initialize(Collection);
+	
 	UWorld* World = GetWorld();
-	if (World == nullptr)
-	{
-		return nullptr;
-	}
-	else if (World->IsNetMode(NM_Client))
+	check(World);
+
+	if (World->GetNetMode() == NM_Client)
 	{
-		// this is where we'd sent a request to the server to spawn the actor
-		//return SpawnedActor;
+		ActorSpawnerSubsystem = World->GetSubsystem<UMassLWIClientActorSpawnerSubsystem>();
 	}
-
-	// @todo notify OnActorSpawnedDelegate
-
-	return Super::GetOrSpawnActorFromTemplate(MassAgent, Transform, TemplateActorIndex
-		, SpawnRequestHandle, Priority, ActorPreSpawnDelegate
-		, ActorPostSpawnDelegate);
 }
+
diff --git a/Engine/Plugins/Runtime/MassGameplay/Source/MassLWI/Private/MassLWIStaticMeshManager.cpp b/Engine/Plugins/Runtime/MassGameplay/Source/MassLWI/Private/MassLWIStaticMeshManager.cpp
--- a/Engine/Plugins/Runtime/MassGameplay/Source/MassLWI/Private/MassLWIStaticMeshManager.cpp
+++ b/Engine/Plugins/Runtime/MassGameplay/Source/MassLWI/Private/MassLWIStaticMeshManager.cpp
@@ -4,6 +4,7 @@
 #include "MassLWISubsystem.h"
 #include "Components/HierarchicalInstancedStaticMeshComponent.h"
 #include "MassCommonFragments.h"
+#include "MassEntitySubsystem.h"
 #include "MassRepresentationSubsystem.h"
 #include "MassSpawnerSubsystem.h"
 #include "MassSpawnerTypes.h"
@@ -12,6 +13,18 @@
 #include "VisualLogger/VisualLogger.h"
 
 
+namespace UE::Mass::Tweakables
+{
+	bool bDestroyEntitiesOnEndPlay = false;
+
+	FAutoConsoleVariableRef CLWIVars[] = {
+		{TEXT("mass.LWI.DestroyEntitiesOnEndPlay"), bDestroyEntitiesOnEndPlay, TEXT("Whether we should destroy LWI-sources entities when the original LWI manager ends play")},
+	};
+}
+
+//-----------------------------------------------------------------------------
+// AMassLWIStaticMeshManager
+//-----------------------------------------------------------------------------
 void AMassLWIStaticMeshManager::PostLoad()
 {
 	Super::PostLoad();
@@ -44,6 +57,23 @@ void AMassLWIStaticMeshManager::EndPlay(const EEndPlayReason::Type EndPlayReason
 	UWorld* World = GetWorld();
 	if (World && IsRegisteredWithMass())
 	{
+		if (UE::Mass::Tweakables::bDestroyEntitiesOnEndPlay && Entities.Num())
+		{
+			if (UMassEntitySubsystem* EntitySubsystem = UWorld::GetSubsystem<UMassEntitySubsystem>(GetWorld()))
+			{
+				FMassEntityManager& EntityManager = EntitySubsystem->GetMutableEntityManager();
+
+				TArray<FMassArchetypeEntityCollection> EntityCollectionsToDestroy;
+				UE::Mass::Utils::CreateEntityCollections(EntityManager, Entities, FMassArchetypeEntityCollection::FoldDuplicates, EntityCollectionsToDestroy);
+				for (FMassArchetypeEntityCollection& Collection : EntityCollectionsToDestroy)
+				{
+					EntityManager.BatchDestroyEntityChunks(Collection);
+				}
+
+				Entities.Reset();
+			}
+		}
+
 		if (UMassLWISubsystem* LWIxMass = World->GetSubsystem<UMassLWISubsystem>())
 		{
 			LWIxMass->UnregisterLWIManager(*this);
@@ -106,6 +136,10 @@ void AMassLWIStaticMeshManager::TransferDataToMass(FMassEntityManager& EntityMan
 			}
 		}
 		Actors.Reset();
+
+		InstanceTransforms.Reset();
+		ValidIndices.Reset();
+		FreeIndices.Reset();
 		
 		InstancedStaticMeshComponent->UnregisterComponent();
 		InstancedStaticMeshComponent->DestroyComponent();
@@ -175,6 +209,17 @@ void AMassLWIStaticMeshManager::CreateMassTemplate(FMassEntityManager& EntityMan
 	// @todo make sure a nullptr here won't break stuff
 	FMassStaticMeshInstanceVisualizationMeshDesc& MeshDesc = StaticMeshInstanceDesc.Meshes[0];
 	MeshDesc.Mesh = StaticMesh.Get();
+#if WITH_EDITOR
+	if (!ensure(MeshDesc.Mesh) && GIsEditor)
+	{
+		MeshDesc.Mesh = StaticMesh.LoadSynchronous();
+		if (!ensure(MeshDesc.Mesh))
+		{
+			return;
+		}
+	}
+#endif // WITH_EDITOR
+
 	if (InstancedStaticMeshComponent)
 	{
 		MeshDesc.MaterialOverrides = InstancedStaticMeshComponent->OverrideMaterials;
diff --git a/Engine/Plugins/Runtime/MassGameplay/Source/MassLWI/Private/MassLWISubsystem.cpp b/Engine/Plugins/Runtime/MassGameplay/Source/MassLWI/Private/MassLWISubsystem.cpp
--- a/Engine/Plugins/Runtime/MassGameplay/Source/MassLWI/Private/MassLWISubsystem.cpp
+++ b/Engine/Plugins/Runtime/MassGameplay/Source/MassLWI/Private/MassLWISubsystem.cpp
@@ -6,6 +6,7 @@
 #include "MassLWIStaticMeshManager.h"
 #include "MassLWIConfigActor.h"
 #include "MassEntitySubsystem.h"
+#include "MassLWIClientActorSpawnerSubsystem.h"
 #include "VisualLogger/VisualLogger.h"
 
 
@@ -26,6 +27,10 @@ void UMassLWISubsystem::Initialize(FSubsystemCollectionBase& Collection)
 	UMassEntitySubsystem* EntitySubsystem = Collection.InitializeDependency<UMassEntitySubsystem>();
 	check(EntitySubsystem);
 	EntityManager = EntitySubsystem->GetMutableEntityManager().AsShared();
+
+	// note that it's ok for LWISpawnerSubsystem to be null since in principle it's only available on clients,
+	// while UMassLWISubsystem can be created on both.
+	LWISpawnerSubsystem = Collection.InitializeDependency<UMassLWIClientActorSpawnerSubsystem>();
 }
 
 void UMassLWISubsystem::PostInitialize()
@@ -62,12 +67,18 @@ void UMassLWISubsystem::RegisterLWIManager(AMassLWIStaticMeshManager& Manager)
 	if (FreeIndices.Num() > 0)
 	{
 		RegistrationIndex = FreeIndices.Pop(/*bAllowShrinking = */true);
+		RegisteredManagers[RegistrationIndex] = &Manager;
 	}
 	else
 	{
 		RegistrationIndex = RegisteredManagers.Add(&Manager);
 	}
 
+	if (LWISpawnerSubsystem)
+	{
+		LWISpawnerSubsystem->RegisterRepresentedClass(Manager.GetRepresentedClass());
+	}
+
 	Manager.TransferDataToMass(*EntityManager);
 	Manager.MarkRegisteredWithMass(FMassLWIManagerRegistrationHandle(RegistrationIndex));
 }
diff --git a/Engine/Plugins/Runtime/MassGameplay/Source/MassLWI/Private/MassLWIVisualizationTrait.cpp b/Engine/Plugins/Runtime/MassGameplay/Source/MassLWI/Private/MassLWIVisualizationTrait.cpp
--- a/Engine/Plugins/Runtime/MassGameplay/Source/MassLWI/Private/MassLWIVisualizationTrait.cpp
+++ b/Engine/Plugins/Runtime/MassGameplay/Source/MassLWI/Private/MassLWIVisualizationTrait.cpp
@@ -3,8 +3,6 @@
 #include "MassLWIVisualizationTrait.h"
 #include "MassEntityTemplateRegistry.h"
 #include "MassLWIRepresentationSubsystem.h"
-#include "MassLWIRepresentationActorManagement.h"
-//#include "MassInstancedStaticMeshComponent.h"
 
 
 UMassLWIVisualizationTrait::UMassLWIVisualizationTrait(const FObjectInitializer& ObjectInitializer)
@@ -13,7 +11,5 @@ UMassLWIVisualizationTrait::UMassLWIVisualizationTrait(const FObjectInitializer&
 #if WITH_EDITORONLY_DATA
 	bCanModifyRepresentationSubsystemClass = false;
 	RepresentationSubsystemClass = UMassLWIRepresentationSubsystem::StaticClass();
-	Params.bCanModifyRepresentationActorManagementClass = false;
-	Params.RepresentationActorManagementClass = UMassLWIRepresentationActorManagement::StaticClass();
 #endif // WITH_EDITORONLY_DATA
 }
diff --git a/Engine/Plugins/Runtime/MassGameplay/Source/MassLWI/Public/MassLWIClientActorSpawnerSubsystem.h b/Engine/Plugins/Runtime/MassGameplay/Source/MassLWI/Public/MassLWIClientActorSpawnerSubsystem.h
new file mode 100644
--- /dev/null
+++ b/Engine/Plugins/Runtime/MassGameplay/Source/MassLWI/Public/MassLWIClientActorSpawnerSubsystem.h
@@ -0,0 +1,53 @@
+// Copyright Epic Games, Inc. All Rights Reserved.
+
+#pragma once
+
+#include "MassExternalSubsystemTraits.h"
+#include "MassActorSpawnerSubsystem.h"
+#include "MassLWIClientActorSpawnerSubsystem.generated.h"
+
+
+USTRUCT()
+struct MASSLWI_API FMassStoredActorsContainer
+{
+	GENERATED_BODY()
+
+	UPROPERTY()
+	TArray<TObjectPtr<AActor>> Container;
+};
+
+UCLASS()
+class MASSLWI_API UMassLWIClientActorSpawnerSubsystem : public UMassActorSpawnerSubsystem
+{
+	GENERATED_BODY()
+
+public:
+	/** checks if RepresentedClass represents an actor class and if so creates an entry in PendingActors for the type */
+	void RegisterRepresentedClass(UClass* RepresentedClass);
+
+protected:
+	// USubsystem API begin
+	virtual bool ShouldCreateSubsystem(UObject* Outer) const override;
+	virtual void Initialize(FSubsystemCollectionBase& Collection) override;
+	// USubsystem API end
+
+	// UMassActorSpawnerSubsystem API begin
+	virtual ESpawnRequestStatus SpawnActor(FConstStructView SpawnRequestView, TObjectPtr<AActor>& OutSpawnedActor) const override;
+	// UMassActorSpawnerSubsystem API end
+
+	void OnActorSpawned(AActor* InActor);
+
+	FDelegateHandle WorldOnActorSpawnedHandle;
+
+	UPROPERTY()
+	mutable TMap<TSubclassOf<AActor>, FMassStoredActorsContainer> PendingActors;
+};
+
+template<>
+struct TMassExternalSubsystemTraits<UMassLWIClientActorSpawnerSubsystem> final
+{
+	enum
+	{
+		GameThreadOnly = true
+	};
+};
\ No newline at end of file
diff --git a/Engine/Plugins/Runtime/MassGameplay/Source/MassLWI/Public/MassLWIRepresentationActorManagement.h b/Engine/Plugins/Runtime/MassGameplay/Source/MassLWI/Public/MassLWIRepresentationActorManagement.h
deleted file mode 100644
diff --git a/Engine/Plugins/Runtime/MassGameplay/Source/MassLWI/Public/MassLWIRepresentationSubsystem.h b/Engine/Plugins/Runtime/MassGameplay/Source/MassLWI/Public/MassLWIRepresentationSubsystem.h
--- a/Engine/Plugins/Runtime/MassGameplay/Source/MassLWI/Public/MassLWIRepresentationSubsystem.h
+++ b/Engine/Plugins/Runtime/MassGameplay/Source/MassLWI/Public/MassLWIRepresentationSubsystem.h
@@ -12,17 +12,10 @@ class MASSLWI_API UMassLWIRepresentationSubsystem : public UMassRepresentationSu
 {
 	GENERATED_BODY()
 
-public:
-	DECLARE_MULTICAST_DELEGATE_OneParam(FOnActorSpawned, AActor&);
-
-	AActor* GetOrSpawnActorFromTemplate(const FMassEntityHandle MassAgent, const FTransform& Transform, const int16 TemplateActorIndex
-		, FMassActorSpawnRequestHandle& SpawnRequestHandle, float Priority, FMassActorPreSpawnDelegate ActorPreSpawnDelegate
-		, FMassActorPostSpawnDelegate ActorPostSpawnDelegate);
-
-	FOnActorSpawned& GetOnActorSpawnedDelegate() { return OnActorSpawnedDelegate; }
-
 protected:
-	FOnActorSpawned OnActorSpawnedDelegate;
+	// USubsystem API begin
+	virtual void Initialize(FSubsystemCollectionBase& Collection) override;
+	// USubsystem API end
 };
 
 template<>
diff --git a/Engine/Plugins/Runtime/MassGameplay/Source/MassLWI/Public/MassLWIStaticMeshManager.h b/Engine/Plugins/Runtime/MassGameplay/Source/MassLWI/Public/MassLWIStaticMeshManager.h
--- a/Engine/Plugins/Runtime/MassGameplay/Source/MassLWI/Public/MassLWIStaticMeshManager.h
+++ b/Engine/Plugins/Runtime/MassGameplay/Source/MassLWI/Public/MassLWIStaticMeshManager.h
@@ -44,7 +44,7 @@ protected:
 	virtual void EndPlay(const EEndPlayReason::Type EndPlayReason) override;
 	// End AActor API
 
-	void CreateMassTemplate(FMassEntityManager& EntityManager);
+	virtual void CreateMassTemplate(FMassEntityManager& EntityManager);
 
 
 	FMassLWIManagerRegistrationHandle MassRegistrationHandle;
diff --git a/Engine/Plugins/Runtime/MassGameplay/Source/MassLWI/Public/MassLWISubsystem.h b/Engine/Plugins/Runtime/MassGameplay/Source/MassLWI/Public/MassLWISubsystem.h
--- a/Engine/Plugins/Runtime/MassGameplay/Source/MassLWI/Public/MassLWISubsystem.h
+++ b/Engine/Plugins/Runtime/MassGameplay/Source/MassLWI/Public/MassLWISubsystem.h
@@ -12,6 +12,7 @@
 struct FActorInstanceHandle;
 class AMassLWIStaticMeshManager;
 class AMassLWIConfigActor;
+class UMassLWIClientActorSpawnerSubsystem;
 struct FMassEntityManager;
 
 
@@ -50,6 +51,9 @@ protected:
 
 	TMap<UClass*, FMassEntityConfig> ClassToConfigMap;
 
+	UPROPERTY()
+	TObjectPtr<UMassLWIClientActorSpawnerSubsystem> LWISpawnerSubsystem;
+
 private:
 	TSharedPtr<FMassEntityManager> EntityManager;
 };
diff --git a/Engine/Plugins/Runtime/MassGameplay/Source/MassRepresentation/Private/MassStationaryISMRepresentationFragmentDestructor.cpp b/Engine/Plugins/Runtime/MassGameplay/Source/MassRepresentation/Private/MassStationaryISMRepresentationFragmentDestructor.cpp
new file mode 100644
--- /dev/null
+++ b/Engine/Plugins/Runtime/MassGameplay/Source/MassRepresentation/Private/MassStationaryISMRepresentationFragmentDestructor.cpp
@@ -0,0 +1,58 @@
+// Copyright Epic Games, Inc. All Rights Reserved.
+
+#include "MassStationaryISMRepresentationFragmentDestructor.h"
+#include "MassRepresentationSubsystem.h"
+#include "MassCommonFragments.h"
+#include "MassRepresentationProcessor.h"
+
+//-----------------------------------------------------------------------------
+// UMassStationaryISMRepresentationFragmentDestructor
+//-----------------------------------------------------------------------------
+UMassStationaryISMRepresentationFragmentDestructor::UMassStationaryISMRepresentationFragmentDestructor()
+	: EntityQuery(*this)
+{
+	ObservedType = FMassRepresentationFragment::StaticStruct();
+	Operation = EMassObservedOperation::Remove;
+	ExecutionFlags = (int32)(EProcessorExecutionFlags::All);
+	bRequiresGameThreadExecution = true; // not sure about this
+}
+
+void UMassStationaryISMRepresentationFragmentDestructor::ConfigureQueries()
+{
+	EntityQuery.AddRequirement<FMassRepresentationFragment>(EMassFragmentAccess::ReadOnly);
+	EntityQuery.AddConstSharedRequirement<FMassRepresentationParameters>();
+	EntityQuery.AddSharedRequirement<FMassRepresentationSubsystemSharedFragment>(EMassFragmentAccess::ReadWrite);
+	EntityQuery.AddTagRequirement<FMassStaticRepresentationTag>(EMassFragmentPresence::All);
+}
+
+void UMassStationaryISMRepresentationFragmentDestructor::Execute(FMassEntityManager& EntityManager, FMassExecutionContext& Context)
+{
+	EntityQuery.ForEachEntityChunk(EntityManager, Context, [this](FMassExecutionContext& Context)
+	{
+		UMassRepresentationSubsystem* RepresentationSubsystem = Context.GetMutableSharedFragment<FMassRepresentationSubsystemSharedFragment>().RepresentationSubsystem;
+		check(RepresentationSubsystem);
+		FMassInstancedStaticMeshInfoArrayView ISMInfosView = RepresentationSubsystem->GetMutableInstancedStaticMeshInfos();
+
+		const TArrayView<FMassRepresentationFragment> RepresentationList = Context.GetMutableFragmentView<FMassRepresentationFragment>();
+
+		const int32 NumEntities = Context.GetNumEntities();
+		for (int32 EntityIdx = 0; EntityIdx < NumEntities; ++EntityIdx)
+		{
+			FMassRepresentationFragment& Representation = RepresentationList[EntityIdx];
+			if (Representation.CurrentRepresentation == EMassRepresentationType::StaticMeshInstance)
+			{
+				FMassInstancedStaticMeshInfo& ISMInfo = ISMInfosView[Representation.StaticMeshDescIndex];
+				if (FMassLODSignificanceRange* OldRange = ISMInfo.GetLODSignificanceRange(Representation.PrevLODSignificance))
+				{
+					const FMassEntityHandle EntityHandle = Context.GetEntity(EntityIdx);
+					const int32 EntityId = GetTypeHash(EntityHandle);
+					if (OldRange)
+					{
+						OldRange->RemoveInstance(EntityId);
+					}
+				}
+				Representation.CurrentRepresentation = EMassRepresentationType::None;
+			}
+		}
+	});
+}
diff --git a/Engine/Plugins/Runtime/MassGameplay/Source/MassRepresentation/Private/MassStationaryVisualizationTrait.cpp b/Engine/Plugins/Runtime/MassGameplay/Source/MassRepresentation/Private/MassStationaryVisualizationTrait.cpp
--- a/Engine/Plugins/Runtime/MassGameplay/Source/MassRepresentation/Private/MassStationaryVisualizationTrait.cpp
+++ b/Engine/Plugins/Runtime/MassGameplay/Source/MassRepresentation/Private/MassStationaryVisualizationTrait.cpp
@@ -7,6 +7,7 @@
 
 UMassStationaryVisualizationTrait::UMassStationaryVisualizationTrait(const FObjectInitializer& ObjectInitializer)
 {
+	bAllowServerSideVisualization = true;
 }
 
 void UMassStationaryVisualizationTrait::BuildTemplate(FMassEntityTemplateBuildContext& BuildContext, const UWorld& World) const
diff --git a/Engine/Plugins/Runtime/MassGameplay/Source/MassRepresentation/Public/MassRepresentationSubsystem.h b/Engine/Plugins/Runtime/MassGameplay/Source/MassRepresentation/Public/MassRepresentationSubsystem.h
--- a/Engine/Plugins/Runtime/MassGameplay/Source/MassRepresentation/Public/MassRepresentationSubsystem.h
+++ b/Engine/Plugins/Runtime/MassGameplay/Source/MassRepresentation/Public/MassRepresentationSubsystem.h
@@ -163,3 +163,11 @@ protected:
 	TMap<FMassEntityHandle, int32> HandledMassAgents;
 };
 
+template<>
+struct TMassExternalSubsystemTraits<UMassRepresentationSubsystem> final
+{
+	enum
+	{
+		GameThreadOnly = true
+	};
+};
diff --git a/Engine/Plugins/Runtime/MassGameplay/Source/MassRepresentation/Public/MassStationaryISMRepresentationFragmentDestructor.h b/Engine/Plugins/Runtime/MassGameplay/Source/MassRepresentation/Public/MassStationaryISMRepresentationFragmentDestructor.h
new file mode 100644
--- /dev/null
+++ b/Engine/Plugins/Runtime/MassGameplay/Source/MassRepresentation/Public/MassStationaryISMRepresentationFragmentDestructor.h
@@ -0,0 +1,24 @@
+// Copyright Epic Games, Inc. All Rights Reserved.
+
+#pragma once 
+
+#include "MassObserverProcessor.h"
+#include "MassStationaryISMRepresentationFragmentDestructor.generated.h"
+
+
+/** 
+ * This class is responsible for cleaning up ISM instances visualizing stationary entities
+ */
+UCLASS()
+class MASSREPRESENTATION_API UMassStationaryISMRepresentationFragmentDestructor : public UMassObserverProcessor
+{
+	GENERATED_BODY()
+public:
+	UMassStationaryISMRepresentationFragmentDestructor();
+
+protected:
+	virtual void ConfigureQueries() override;
+	virtual void Execute(FMassEntityManager& EntityManager, FMassExecutionContext& Context) override;
+
+	FMassEntityQuery EntityQuery;
+};
