(re)storing state on building prop (de)hydration with MassLWI

diff --git a/Engine/Plugins/Runtime/MassGameplay/Source/MassRepresentation/MassRepresentation.Build.cs b/Engine/Plugins/Runtime/MassGameplay/Source/MassRepresentation/MassRepresentation.Build.cs
--- a/Engine/Plugins/Runtime/MassGameplay/Source/MassRepresentation/MassRepresentation.Build.cs
+++ b/Engine/Plugins/Runtime/MassGameplay/Source/MassRepresentation/MassRepresentation.Build.cs
@@ -25,6 +25,7 @@ namespace UnrealBuildTool.Rules
 					"RenderCore",
 					"RHI",
 					"MassMovement",
+					"MassSignals",
 				}
 			);
 
diff --git a/Engine/Plugins/Runtime/MassGameplay/Source/MassRepresentation/Private/MassStationaryISMSwitcherProcessor.cpp b/Engine/Plugins/Runtime/MassGameplay/Source/MassRepresentation/Private/MassStationaryISMSwitcherProcessor.cpp
--- a/Engine/Plugins/Runtime/MassGameplay/Source/MassRepresentation/Private/MassStationaryISMSwitcherProcessor.cpp
+++ b/Engine/Plugins/Runtime/MassGameplay/Source/MassRepresentation/Private/MassStationaryISMSwitcherProcessor.cpp
@@ -5,6 +5,8 @@
 #include "MassCommonFragments.h"
 #include "MassRepresentationProcessor.h"
 
+#include "MassSignalSubsystem.h"
+
 
 UMassStationaryISMSwitcherProcessor::UMassStationaryISMSwitcherProcessor(const FObjectInitializer& ObjectInitializer)
 	: EntityQuery(*this)
@@ -20,12 +22,15 @@ void UMassStationaryISMSwitcherProcessor::ConfigureQueries()
 	EntityQuery.AddRequirement<FMassRepresentationFragment>(EMassFragmentAccess::ReadWrite);
 	EntityQuery.AddSharedRequirement<FMassRepresentationSubsystemSharedFragment>(EMassFragmentAccess::ReadWrite);
 	EntityQuery.AddTagRequirement<FMassStaticRepresentationTag>(EMassFragmentPresence::All);
+	EntityQuery.AddSubsystemRequirement<UMassSignalSubsystem>(EMassFragmentAccess::ReadWrite);
 }
 
 void UMassStationaryISMSwitcherProcessor::Execute(FMassEntityManager& EntityManager, FMassExecutionContext& Context)
 {
 	EntityQuery.ForEachEntityChunk(EntityManager, Context, [this](FMassExecutionContext& Context)
 	{
+		UMassSignalSubsystem& SignalSubsystem = Context.GetMutableSubsystemChecked<UMassSignalSubsystem>();
+
 		UMassRepresentationSubsystem* RepresentationSubsystem = Context.GetMutableSharedFragment<FMassRepresentationSubsystemSharedFragment>().RepresentationSubsystem;
 		check(RepresentationSubsystem);
 		FMassInstancedStaticMeshInfoArrayView ISMInfosView = RepresentationSubsystem->GetMutableInstancedStaticMeshInfos();
@@ -62,6 +67,11 @@ void UMassStationaryISMSwitcherProcessor::Execute(FMassEntityManager& EntityMana
 
 				// consume "prev" data
 				Representation.PrevRepresentation = Representation.CurrentRepresentation;
+
+				if (Representation.PrevRepresentation != EMassRepresentationType::None)
+				{
+					SignalSubsystem.SignalEntity(UE::Mass::Signals::SwitchedToActor, EntityHandle);
+				}
 			}
 			else if (Representation.PrevRepresentation != EMassRepresentationType::StaticMeshInstance
 				&& Representation.CurrentRepresentation == EMassRepresentationType::StaticMeshInstance)
@@ -88,6 +98,8 @@ void UMassStationaryISMSwitcherProcessor::Execute(FMassEntityManager& EntityMana
 
 				// consume "prev" data
 				Representation.PrevRepresentation = Representation.CurrentRepresentation;
+
+				SignalSubsystem.SignalEntity(UE::Mass::Signals::SwitchedToISM, EntityHandle);
 			}
 			else if (ISMInfo.GetLODSignificanceRangesNum() > 1 && Representation.PrevLODSignificance != RepresentationLOD.LODSignificance)
 			{
diff --git a/Engine/Plugins/Runtime/MassGameplay/Source/MassRepresentation/Public/MassStationaryISMSwitcherProcessor.h b/Engine/Plugins/Runtime/MassGameplay/Source/MassRepresentation/Public/MassStationaryISMSwitcherProcessor.h
--- a/Engine/Plugins/Runtime/MassGameplay/Source/MassRepresentation/Public/MassStationaryISMSwitcherProcessor.h
+++ b/Engine/Plugins/Runtime/MassGameplay/Source/MassRepresentation/Public/MassStationaryISMSwitcherProcessor.h
@@ -5,6 +5,11 @@
 #include "MassProcessor.h"
 #include "MassStationaryISMSwitcherProcessor.generated.h"
 
+namespace UE::Mass::Signals
+{
+	const FName SwitchedToActor = FName(TEXT("SwitchedToActor"));
+	const FName SwitchedToISM = FName(TEXT("SwitchedToISM"));
+}
 
 /** 
  * This processor's sole responsibility is to process all entities tagged with FMassStaticRepresentationTag
