Forcing sync-load of StaticMesh when creating a template for MassLWI-represented actors.

diff --git a/Engine/Plugins/Runtime/MassGameplay/Source/MassLWI/Private/MassLWIStaticMeshManager.cpp b/Engine/Plugins/Runtime/MassGameplay/Source/MassLWI/Private/MassLWIStaticMeshManager.cpp
--- a/Engine/Plugins/Runtime/MassGameplay/Source/MassLWI/Private/MassLWIStaticMeshManager.cpp
+++ b/Engine/Plugins/Runtime/MassGameplay/Source/MassLWI/Private/MassLWIStaticMeshManager.cpp
@@ -1,6 +1,7 @@
 // Copyright Epic Games, Inc. All Rights Reserved.
 
 #include "MassLWIStaticMeshManager.h"
+#include "MassLWITypes.h"
 #include "MassLWISubsystem.h"
 #include "Components/HierarchicalInstancedStaticMeshComponent.h"
 #include "MassCommonFragments.h"
@@ -209,16 +210,16 @@ void AMassLWIStaticMeshManager::CreateMassTemplate(FMassEntityManager& EntityMan
 	// @todo make sure a nullptr here won't break stuff
 	FMassStaticMeshInstanceVisualizationMeshDesc& MeshDesc = StaticMeshInstanceDesc.Meshes[0];
 	MeshDesc.Mesh = StaticMesh.Get();
-#if WITH_EDITOR
-	if (!ensure(MeshDesc.Mesh) && GIsEditor)
+
+	if (!ensure(MeshDesc.Mesh))
 	{
 		MeshDesc.Mesh = StaticMesh.LoadSynchronous();
 		if (!ensure(MeshDesc.Mesh))
 		{
+			UE_VLOG_UELOG(this, LogMassLWI, Error, TEXT("%s: Unable to load mesh asset %s"), *UObject::GetName(), *StaticMesh.GetLongPackageName());
 			return;
 		}
 	}
-#endif // WITH_EDITOR
 
 	if (InstancedStaticMeshComponent)
 	{
