Added bRandomize property to FMassTransformsSpawnData to control whether UMassSpawnLocationProcessor assigns transforms randomly (defaults to true to maintain current behavior) or if false, assigns transforms in order.

diff --git a/Engine/Plugins/Runtime/MassGameplay/Source/MassSpawner/Private/MassSpawnLocationProcessor.cpp b/Engine/Plugins/Runtime/MassGameplay/Source/MassSpawner/Private/MassSpawnLocationProcessor.cpp
--- a/Engine/Plugins/Runtime/MassGameplay/Source/MassSpawner/Private/MassSpawnLocationProcessor.cpp
+++ b/Engine/Plugins/Runtime/MassGameplay/Source/MassSpawner/Private/MassSpawnLocationProcessor.cpp
@@ -65,15 +65,31 @@ void UMassSpawnLocationProcessor::Execute(FMassEntityManager& EntityManager, FMa
 		}
 	}
 
-	EntityQuery.ForEachEntityChunk(EntityManager, Context, [&Transforms, this](FMassExecutionContext& Context)
-		{
-			const TArrayView<FTransformFragment> LocationList = Context.GetMutableFragmentView<FTransformFragment>();
-			const int32 NumEntities = Context.GetNumEntities();
-			for (int32 i = 0; i < NumEntities; ++i)
+	if (AuxData.bRandomize)
+	{
+		EntityQuery.ForEachEntityChunk(EntityManager, Context, [&Transforms, this](FMassExecutionContext& Context)
 			{
-				const int32 AuxIndex = FMath::RandRange(0, Transforms.Num() - 1);
-				LocationList[i].GetMutableTransform() = Transforms[AuxIndex];
-				Transforms.RemoveAtSwap(AuxIndex, 1, /*bAllowShrinking=*/false);
-			}
-		});
+				const TArrayView<FTransformFragment> LocationList = Context.GetMutableFragmentView<FTransformFragment>();
+				const int32 NumEntities = Context.GetNumEntities();
+				for (int32 i = 0; i < NumEntities; ++i)
+				{
+					const int32 AuxIndex = FMath::RandRange(0, Transforms.Num() - 1);
+					LocationList[i].GetMutableTransform() = Transforms[AuxIndex];
+					Transforms.RemoveAtSwap(AuxIndex, 1, /*bAllowShrinking=*/false);
+				}
+			});
+	}
+	else
+	{
+		int32 NextTransformIndex = 0;
+		EntityQuery.ForEachEntityChunk(EntityManager, Context, [&Transforms, &NextTransformIndex, this](FMassExecutionContext& Context)
+			{
+				const int32 NumEntities = Context.GetNumEntities();
+				TArrayView<FTransformFragment> LocationList = Context.GetMutableFragmentView<FTransformFragment>();
+				check(NextTransformIndex + NumEntities <= Transforms.Num());
+	
+				FMemory::Memcpy(LocationList.GetData(), &Transforms[NextTransformIndex], NumEntities * LocationList.GetTypeSize());
+				NextTransformIndex += NumEntities;
+			});
+	}
 }
diff --git a/Engine/Plugins/Runtime/MassGameplay/Source/MassSpawner/Public/MassSpawnerTypes.h b/Engine/Plugins/Runtime/MassGameplay/Source/MassSpawner/Public/MassSpawnerTypes.h
--- a/Engine/Plugins/Runtime/MassGameplay/Source/MassSpawner/Public/MassSpawnerTypes.h
+++ b/Engine/Plugins/Runtime/MassGameplay/Source/MassSpawner/Public/MassSpawnerTypes.h
@@ -25,9 +25,21 @@ struct FMassTransformsSpawnData
 {
 	GENERATED_BODY()
 
+	FMassTransformsSpawnData()
+	: bRandomize(true)
+	{}
+
 	// declaring the type used to be able to statically test it against other types 
 	using FTransformsContainerType = TArray<FTransform>;
 	FTransformsContainerType Transforms;
+	
+	// When true, Transforms will be assigned to entities as if pre-shuffled. 
+	// If Transforms >= Entities, this provides the best chance for a good spread of entity transforms 
+	// and transforms choices will be guaranteed unique. If Transforms < Entities, the same 
+	// transform will be used for multiple entities and a warning logged.
+	//
+	// When false, Transforms will be assigned in order to spawned entities.
+	bool bRandomize : 1;
 };
 
 /**
