Added FMassLWIManagerSharedFragment shared fragment to MassLWI spawned entities to record the AMassLWIStaticMeshManager they were spawned from.

Also fixed an issue where LWIM stored transforms were being assigned to spawned entities in random order.

#rnx

diff --git a/Engine/Plugins/Runtime/MassGameplay/Source/MassLWI/Private/MassLWIStaticMeshManager.cpp b/Engine/Plugins/Runtime/MassGameplay/Source/MassLWI/Private/MassLWIStaticMeshManager.cpp
--- a/Engine/Plugins/Runtime/MassGameplay/Source/MassLWI/Private/MassLWIStaticMeshManager.cpp
+++ b/Engine/Plugins/Runtime/MassGameplay/Source/MassLWI/Private/MassLWIStaticMeshManager.cpp
@@ -3,9 +3,11 @@
 #include "MassLWIStaticMeshManager.h"
 #include "MassLWITypes.h"
 #include "MassLWISubsystem.h"
+#include "MassLWITypes.h"
 #include "Components/HierarchicalInstancedStaticMeshComponent.h"
 #include "MassCommonFragments.h"
 #include "MassEntitySubsystem.h"
+#include "MassEntityView.h"
 #include "MassRepresentationSubsystem.h"
 #include "MassSpawnerSubsystem.h"
 #include "MassSpawnerTypes.h"
@@ -104,29 +106,20 @@ void AMassLWIStaticMeshManager::TransferDataToMass(FMassEntityManager& EntityMan
 		check(SpawnerSystem);
 		const int32 NumEntities = InstanceTransforms.Num();
 		
-		if (ManagerLocation.IsNearlyZero())
+		FMassTransformsSpawnData SpawnLocationData;
+		SpawnLocationData.bRandomize = false;
+		SpawnLocationData.Transforms = InstanceTransforms;
+		if (ManagerLocation.IsNearlyZero() == false)
 		{
-			// a tiny bit hacky conversion of InstanceTransforms to FMassTransformsSpawnData. Static assert here should catch if the type doesn't match
-			static_assert(sizeof(FMassTransformsSpawnData::FTransformsContainerType) == sizeof(InstanceTransforms)
-				// @todo use some "TTypeOf" template 
-				, "the InstanceTransforms-to-FMassTransformsSpawnData conversion relies on FMassTransformsSpawnData having only a single property that matches InstanceTransforms' type exactly");
-			FConstStructView SpawnLocationDataView(FMassTransformsSpawnData::StaticStruct(), reinterpret_cast<const uint8*>(&InstanceTransforms));
-
-			SpawnerSystem->SpawnEntities(FinalizedTemplate->GetTemplateID(), NumEntities, SpawnLocationDataView, UMassSpawnLocationProcessor::StaticClass(), Entities);
-		}
-		else
-		{
-			FMassTransformsSpawnData SpawnLocationData;
-			SpawnLocationData.Transforms = InstanceTransforms;
 			for (FTransform& EntityTransform : SpawnLocationData.Transforms)
 			{
 				EntityTransform.AddToTranslation(ManagerLocation);
 			}
-
-			FConstStructView SpawnLocationDataView(FMassTransformsSpawnData::StaticStruct(), reinterpret_cast<const uint8*>(&SpawnLocationData));
-			SpawnerSystem->SpawnEntities(FinalizedTemplate->GetTemplateID(), NumEntities, SpawnLocationDataView, UMassSpawnLocationProcessor::StaticClass(), Entities);
 		}
 
+		FConstStructView SpawnLocationDataView = FConstStructView::Make(SpawnLocationData);
+		SpawnerSystem->SpawnEntities(FinalizedTemplate->GetTemplateID(), NumEntities, SpawnLocationDataView, UMassSpawnLocationProcessor::StaticClass(), Entities);
+
 		// destroy actors that have already been created
 		for (const TPair<int32, TObjectPtr<AActor>>& ActorPair : Actors)
 		{
@@ -147,6 +140,27 @@ void AMassLWIStaticMeshManager::TransferDataToMass(FMassEntityManager& EntityMan
 	}
 }
 
+int32 AMassLWIStaticMeshManager::FindIndexForEntity(const FMassEntityHandle Entity) const
+{
+	return Entities.IndexOfByKey(Entity);
+}
+
+AMassLWIStaticMeshManager* AMassLWIStaticMeshManager::GetMassLWIManagerForEntity(const FMassEntityManager& EntityManager, const FMassEntityHandle Entity)
+{
+	FMassEntityView EntityView(EntityManager, Entity);
+	if (EntityView.IsSet())
+	{	
+		const FMassLWIManagerSharedFragment* LWIManagerSharedFragmentPtr = EntityView.GetSharedFragmentDataPtr<FMassLWIManagerSharedFragment>();
+		if (ensure(LWIManagerSharedFragmentPtr))
+		{
+			AMassLWIStaticMeshManager* LWIManager = LWIManagerSharedFragmentPtr->LWIManager.Get();
+			return LWIManager;
+		}
+	}
+
+	return nullptr;
+}
+
 void AMassLWIStaticMeshManager::CreateMassTemplate(FMassEntityManager& EntityManager)
 {
 	UWorld* World = GetWorld();
@@ -238,6 +252,12 @@ void AMassLWIStaticMeshManager::CreateMassTemplate(FMassEntityManager& EntityMan
 		//RepresentationFragment.LowResTemplateActorIndex = TemplateActorIndex;
 	}
 
+	FMassLWIManagerSharedFragment LWIManagerSharedFragment;
+	LWIManagerSharedFragment.LWIManager = this;
+	uint32 LWIManagerHash = UE::StructUtils::GetStructCrc32(FConstStructView::Make(LWIManagerSharedFragment));
+	FSharedStruct LWIManagerFragment = EntityManager.GetOrCreateSharedFragmentByHash<FMassLWIManagerSharedFragment>(LWIManagerHash, LWIManagerSharedFragment);
+	NewTemplate.AddSharedFragment(LWIManagerFragment);
+
 	FMassEntityTemplateRegistry& TemplateRegistry = SpawnerSystem->GetMutableTemplateRegistryInstance();
 	const uint32 FlavorHash = GetTypeHash(GetNameSafe(RepresentedClass));
 	MassTemplateID = FMassEntityTemplateIDFactory::MakeFlavor(SourceTemplate.GetTemplateID(), FlavorHash);
diff --git a/Engine/Plugins/Runtime/MassGameplay/Source/MassLWI/Public/MassLWIStaticMeshManager.h b/Engine/Plugins/Runtime/MassGameplay/Source/MassLWI/Public/MassLWIStaticMeshManager.h
--- a/Engine/Plugins/Runtime/MassGameplay/Source/MassLWI/Public/MassLWIStaticMeshManager.h
+++ b/Engine/Plugins/Runtime/MassGameplay/Source/MassLWI/Public/MassLWIStaticMeshManager.h
@@ -34,6 +34,11 @@ public:
 	void MarkRegisteredWithMass(const FMassLWIManagerRegistrationHandle RegistrationIndex);
 	void MarkUnregisteredWithMass();
 
+	// Returns the index of the light weight instance associated with InEntity if one exists; otherwise we return INDEX_NONE
+	int32 FindIndexForEntity(const FMassEntityHandle Entity) const;
+
+	static AMassLWIStaticMeshManager* GetMassLWIManagerForEntity(const FMassEntityManager& EntityManager, const FMassEntityHandle Entity);
+
 protected:
 	// UObject API
 	virtual void PostLoad() override;
diff --git a/Engine/Plugins/Runtime/MassGameplay/Source/MassLWI/Public/MassLWITypes.h b/Engine/Plugins/Runtime/MassGameplay/Source/MassLWI/Public/MassLWITypes.h
--- a/Engine/Plugins/Runtime/MassGameplay/Source/MassLWI/Public/MassLWITypes.h
+++ b/Engine/Plugins/Runtime/MassGameplay/Source/MassLWI/Public/MassLWITypes.h
@@ -3,6 +3,19 @@
 #pragma once
 
 #include "Logging/LogMacros.h"
+#include "MassEntityTypes.h"
 
+#include "MassLWITypes.generated.h"
 
 DECLARE_LOG_CATEGORY_EXTERN(LogMassLWI, Log, All)
+
+class AMassLWIStaticMeshManager;
+
+USTRUCT()
+struct FMassLWIManagerSharedFragment : public FMassSharedFragment
+{
+	GENERATED_BODY()
+
+	UPROPERTY(Transient)
+	TSoftObjectPtr<AMassLWIStaticMeshManager> LWIManager = nullptr;
+};
