Simplified UMassActorSpawnerSubsystem subclassing by moving post actor spawn steps that should always happen from SpawnActor into SpawnOrRetrieveFromPool, leaving SpawnActor now simply calling World->SpawnActor.

Also fixed an issue where pooled actors with UMassAgentComponent's wouldn't udate their Mass handle when reused from the pool.

diff --git a/Engine/Plugins/Runtime/MassGameplay/Source/MassActors/Private/MassActorSpawnerSubsystem.cpp b/Engine/Plugins/Runtime/MassGameplay/Source/MassActors/Private/MassActorSpawnerSubsystem.cpp
--- a/Engine/Plugins/Runtime/MassGameplay/Source/MassActors/Private/MassActorSpawnerSubsystem.cpp
+++ b/Engine/Plugins/Runtime/MassGameplay/Source/MassActors/Private/MassActorSpawnerSubsystem.cpp
@@ -180,9 +180,10 @@ FMassActorSpawnRequestHandle UMassActorSpawnerSubsystem::GetNextRequestToSpawn()
 
 ESpawnRequestStatus UMassActorSpawnerSubsystem::SpawnOrRetrieveFromPool(FConstStructView SpawnRequestView, TObjectPtr<AActor>& OutSpawnedActor)
 {
+	const FMassActorSpawnRequest& SpawnRequest = SpawnRequestView.Get<const FMassActorSpawnRequest>();
+
 	if (UE::MassActors::bUseActorPooling != 0 && bActorPoolingEnabled)
 	{
-		const FMassActorSpawnRequest& SpawnRequest = SpawnRequestView.Get<const FMassActorSpawnRequest>();
 		auto* Pool = PooledActors.Find(SpawnRequest.Template);
 
 		if (Pool && Pool->Num() > 0)
@@ -206,44 +207,43 @@ ESpawnRequestStatus UMassActorSpawnerSubsystem::SpawnOrRetrieveFromPool(FConstSt
 		}
 	}
 
-	return SpawnActor(SpawnRequestView, OutSpawnedActor);
-}
-
- ESpawnRequestStatus UMassActorSpawnerSubsystem::SpawnActor(FConstStructView SpawnRequestView, TObjectPtr<AActor>& OutSpawnedActor) const
-{
-	TRACE_CPUPROFILER_EVENT_SCOPE(UMassActorSpawnerSubsystem::SpawnActor);
-
-	UWorld* World = GetWorld();
-	check(World);
+	ESpawnRequestStatus SpawnStatus = SpawnActor(SpawnRequestView, OutSpawnedActor);
 
-	const FMassActorSpawnRequest& SpawnRequest = SpawnRequestView.Get<const FMassActorSpawnRequest>();
-	if (AActor* SpawnedActor = World->SpawnActorDeferred<AActor>(SpawnRequest.Template, SpawnRequest.Transform, nullptr, nullptr, ESpawnActorCollisionHandlingMethod::AlwaysSpawn))
+	if (SpawnStatus == ESpawnRequestStatus::Succeeded)
 	{
-		// Add code here before construction script
-
-		SpawnedActor->FinishSpawning(SpawnRequest.Transform);
-		++NumActorSpawned;
-		// The finish spawning might have failed and the spawned actor is destroyed.
-		if (IsValidChecked(SpawnedActor))
+		if (IsValidChecked(OutSpawnedActor))
 		{
-			if (UMassAgentComponent* AgentComp = SpawnedActor->FindComponentByClass<UMassAgentComponent>())
-			{
-				AgentComp->SetPuppetHandle(SpawnRequest.MassAgent);
-			}
-			OutSpawnedActor = SpawnedActor;
-			return ESpawnRequestStatus::Succeeded;
+			++NumActorSpawned;
 		}
 	}
-
-	UE_VLOG_CAPSULE(this, LogMassActor, Error,
+	else
+	{
+		UE_VLOG_CAPSULE(this, LogMassActor, Error,
 					SpawnRequest.Transform.GetLocation(),
 					SpawnRequest.Template.GetDefaultObject()->GetSimpleCollisionHalfHeight(),
 					SpawnRequest.Template.GetDefaultObject()->GetSimpleCollisionRadius(),
 					SpawnRequest.Transform.GetRotation(),
 					FColor::Red,
 					TEXT("Unable to spawn actor for Mass entity [%s]"), *SpawnRequest.MassAgent.DebugGetDescription());
+	}
+
+	return SpawnStatus;
+}
+
+ ESpawnRequestStatus UMassActorSpawnerSubsystem::SpawnActor(FConstStructView SpawnRequestView, TObjectPtr<AActor>& OutSpawnedActor) const
+{
+	TRACE_CPUPROFILER_EVENT_SCOPE(UMassActorSpawnerSubsystem::SpawnActor);
+
+	UWorld* World = GetWorld();
+	check(World);
+
+	const FMassActorSpawnRequest& SpawnRequest = SpawnRequestView.Get<const FMassActorSpawnRequest>();
 
-	return ESpawnRequestStatus::Failed;
+	FActorSpawnParameters SpawnParameters;
+	SpawnParameters.SpawnCollisionHandlingOverride = ESpawnActorCollisionHandlingMethod::AlwaysSpawn;
+	OutSpawnedActor = World->SpawnActor<AActor>(SpawnRequest.Template, SpawnRequest.Transform, SpawnParameters);
+
+	return IsValid(OutSpawnedActor) ? ESpawnRequestStatus::Succeeded : ESpawnRequestStatus::Failed;
 }
 
 void UMassActorSpawnerSubsystem::ProcessPendingSpawningRequest(const double MaxTimeSlicePerTick)
@@ -284,6 +284,14 @@ void UMassActorSpawnerSubsystem::ProcessPendingSpawningRequest(const double MaxT
 
 		if (SpawnRequest.IsFinished())
 		{
+			if (SpawnRequest.SpawnStatus == ESpawnRequestStatus::Succeeded && IsValid(SpawnRequest.SpawnedActor))
+			{
+				if (UMassAgentComponent* AgentComp = SpawnRequest.SpawnedActor->FindComponentByClass<UMassAgentComponent>())
+				{
+					AgentComp->SetPuppetHandle(SpawnRequest.MassAgent);
+				}
+			}
+
 			// Call the post spawn delegate on the spawn request
 			if (SpawnRequest.ActorPostSpawnDelegate.IsBound())
 			{
