Improved Mass Signalling's thread safety which in turn allowed us to mark the MassSignallingSubsystem as accessible off-gamethread

diff --git a/Engine/Plugins/Runtime/MassGameplay/Source/MassSignals/Private/MassSignalProcessorBase.cpp b/Engine/Plugins/Runtime/MassGameplay/Source/MassSignals/Private/MassSignalProcessorBase.cpp
--- a/Engine/Plugins/Runtime/MassGameplay/Source/MassSignals/Private/MassSignalProcessorBase.cpp
+++ b/Engine/Plugins/Runtime/MassGameplay/Source/MassSignals/Private/MassSignalProcessorBase.cpp
@@ -37,16 +37,18 @@ void UMassSignalProcessorBase::Execute(FMassEntityManager& EntityManager, FMassE
 {
 	QUICK_SCOPE_CYCLE_COUNTER(SignalEntities);
 
-	UE_MT_ACQUIRE_WRITE_ACCESS(ReceivedSignalAccessDetector);
-	// Frame buffer handling
 	const int32 ProcessingFrameBufferIndex = CurrentFrameBufferIndex;
-	CurrentFrameBufferIndex = (CurrentFrameBufferIndex + 1) % 2;
+	{
+		// we only need to lock the part where we change the current buffer index. Once that's done the incoming signals will end up 
+		// in the other buffer
+		UE::TScopeLock<UE::FSpinLock> ScopeLock(ReceivedSignalLock);
+		CurrentFrameBufferIndex = (CurrentFrameBufferIndex + 1) % BuffersCount;
+	}
+	
 	FFrameReceivedSignals& ProcessingFrameBuffer = FrameReceivedSignals[ProcessingFrameBufferIndex];
 	TArray<FEntitySignalRange>& ReceivedSignalRanges = ProcessingFrameBuffer.ReceivedSignalRanges;
 	TArray<FMassEntityHandle>& SignaledEntities = ProcessingFrameBuffer.SignaledEntities;
 
-	UE_MT_RELEASE_WRITE_ACCESS(ReceivedSignalAccessDetector);
-
 	if (ReceivedSignalRanges.IsEmpty())
 	{
 		return;
@@ -139,8 +141,6 @@ void UMassSignalProcessorBase::OnSignalReceived(FName SignalName, TConstArrayVie
 {
 	UE::TScopeLock<UE::FSpinLock> ScopeLock(ReceivedSignalLock);
 
-	UE_MT_SCOPED_WRITE_ACCESS(ReceivedSignalAccessDetector);
-
 	FFrameReceivedSignals& CurrentFrameBuffer = FrameReceivedSignals[CurrentFrameBufferIndex];
 
 	FEntitySignalRange& Range = CurrentFrameBuffer.ReceivedSignalRanges.AddDefaulted_GetRef();
diff --git a/Engine/Plugins/Runtime/MassGameplay/Source/MassSignals/Public/MassSignalProcessorBase.h b/Engine/Plugins/Runtime/MassGameplay/Source/MassSignals/Public/MassSignalProcessorBase.h
--- a/Engine/Plugins/Runtime/MassGameplay/Source/MassSignals/Public/MassSignalProcessorBase.h
+++ b/Engine/Plugins/Runtime/MassGameplay/Source/MassSignals/Public/MassSignalProcessorBase.h
@@ -79,8 +79,10 @@ private:
 		TArray<FMassEntityHandle> SignaledEntities;
 	};
 
+	static constexpr int BuffersCount = 2;
+
 	/** Double buffer frame received signal as we can receive new signals as we are processing them */
-	TStaticArray<FFrameReceivedSignals, 2> FrameReceivedSignals;
+	TStaticArray<FFrameReceivedSignals, BuffersCount> FrameReceivedSignals;
 	
 	/** Current frame buffer index of FrameReceivedSignals */
 	int32 CurrentFrameBufferIndex = 0;
@@ -91,7 +93,6 @@ private:
 	/** List of all the registered signal names*/
 	TArray<FName> RegisteredSignals;
 
-	UE_MT_DECLARE_RW_ACCESS_DETECTOR(ReceivedSignalAccessDetector);
 	UE::FSpinLock ReceivedSignalLock;
 };
 
diff --git a/Engine/Plugins/Runtime/MassGameplay/Source/MassSignals/Public/MassSignalSubsystem.h b/Engine/Plugins/Runtime/MassGameplay/Source/MassSignals/Public/MassSignalSubsystem.h
--- a/Engine/Plugins/Runtime/MassGameplay/Source/MassSignals/Public/MassSignalSubsystem.h
+++ b/Engine/Plugins/Runtime/MassGameplay/Source/MassSignals/Public/MassSignalSubsystem.h
@@ -123,6 +123,6 @@ struct TMassExternalSubsystemTraits<UMassSignalSubsystem> final
 {
 	enum
 	{
-		GameThreadOnly = true
+		GameThreadOnly = false
 	};
 };
