Fixed UMassActorSpawnerSubsystem erronously setting entity handle twice for Agent-actors pulled out of actor pool.

diff --git a/Engine/Plugins/Runtime/MassGameplay/Source/MassActors/Private/MassActorSpawnerSubsystem.cpp b/Engine/Plugins/Runtime/MassGameplay/Source/MassActors/Private/MassActorSpawnerSubsystem.cpp
--- a/Engine/Plugins/Runtime/MassGameplay/Source/MassActors/Private/MassActorSpawnerSubsystem.cpp
+++ b/Engine/Plugins/Runtime/MassGameplay/Source/MassActors/Private/MassActorSpawnerSubsystem.cpp
@@ -198,8 +198,9 @@ ESpawnRequestStatus UMassActorSpawnerSubsystem::SpawnOrRetrieveFromPool(FConstSt
 
 			if (UMassAgentComponent* AgentComp = PooledActor->FindComponentByClass<UMassAgentComponent>())
 			{
+				// normally this function gets called from UMassAgentComponent::OnRegister. We need to call it manually here
+				// since we're bringing this actor our of a pool.
 				AgentComp->RegisterWithAgentSubsystem();
-				AgentComp->SetPuppetHandle(SpawnRequest.MassAgent);
 			}
 
 			OutSpawnedActor = PooledActor;
@@ -292,14 +293,18 @@ void UMassActorSpawnerSubsystem::ProcessPendingSpawningRequest(const double MaxT
 				}
 			}
 
+			EMassActorSpawnRequestAction PostAction = EMassActorSpawnRequestAction::Remove;
+
 			// Call the post spawn delegate on the spawn request
 			if (SpawnRequest.ActorPostSpawnDelegate.IsBound())
 			{
-				if (SpawnRequest.ActorPostSpawnDelegate.Execute(SpawnRequestHandle, SpawnRequestView) == EMassActorSpawnRequestAction::Remove)
-				{
-					// If notified, remove the spawning request
-					ensureMsgf(SpawnRequestHandleManager.RemoveHandle(SpawnRequestHandle), TEXT("When providing a delegate, the spawn request gets automatically removed, no need to remove it on your side"));
-				}
+				PostAction = SpawnRequest.ActorPostSpawnDelegate.Execute(SpawnRequestHandle, SpawnRequestView);
+			}
+
+			if (PostAction == EMassActorSpawnRequestAction::Remove)
+			{
+				// If notified, remove the spawning request
+				ensureMsgf(SpawnRequestHandleManager.RemoveHandle(SpawnRequestHandle), TEXT("When providing a delegate, the spawn request gets automatically removed, no need to remove it on your side"));
 			}
 		}
 	}
