[MassLWI] Fixed an issue where AMassLWIStaticMeshManager::TransferToMass was being called too early in PostLoad, where the actor transform used for instance relative positions, won't be finalized yet for actors in transformed streaming levels.

Fix is simply to remove the registration of the AMassLWIStaticMeshManager from PostLoad and instead let registration occur in BeginPlay which was already there, when the actor transform will have been finalized.

#rnx

diff --git a/Engine/Plugins/Runtime/MassGameplay/Source/MassLWI/Private/MassLWIStaticMeshManager.cpp b/Engine/Plugins/Runtime/MassGameplay/Source/MassLWI/Private/MassLWIStaticMeshManager.cpp
--- a/Engine/Plugins/Runtime/MassGameplay/Source/MassLWI/Private/MassLWIStaticMeshManager.cpp
+++ b/Engine/Plugins/Runtime/MassGameplay/Source/MassLWI/Private/MassLWIStaticMeshManager.cpp
@@ -29,19 +29,6 @@ namespace UE::Mass::Tweakables
 //-----------------------------------------------------------------------------
 // AMassLWIStaticMeshManager
 //-----------------------------------------------------------------------------
-void AMassLWIStaticMeshManager::PostLoad()
-{
-	Super::PostLoad();
-
-	if (UWorld* World = GetWorld())
-	{
-		if (UMassLWISubsystem* MassLWI = World->GetSubsystem<UMassLWISubsystem>())
-		{
-			MassLWI->RegisterLWIManager(*this);
-		}
-	}
-}
-
 void AMassLWIStaticMeshManager::BeginPlay()
 {
 	Super::BeginPlay();
diff --git a/Engine/Plugins/Runtime/MassGameplay/Source/MassLWI/Public/MassLWIStaticMeshManager.h b/Engine/Plugins/Runtime/MassGameplay/Source/MassLWI/Public/MassLWIStaticMeshManager.h
--- a/Engine/Plugins/Runtime/MassGameplay/Source/MassLWI/Public/MassLWIStaticMeshManager.h
+++ b/Engine/Plugins/Runtime/MassGameplay/Source/MassLWI/Public/MassLWIStaticMeshManager.h
@@ -40,10 +40,6 @@ public:
 	static AMassLWIStaticMeshManager* GetMassLWIManagerForEntity(const FMassEntityManager& EntityManager, const FMassEntityHandle Entity);
 
 protected:
-	// UObject API
-	virtual void PostLoad() override;
-	// End UObject API
-
 	// AActor API
 	virtual void BeginPlay() override;
 	virtual void EndPlay(const EEndPlayReason::Type EndPlayReason) override;
