Fixed MssSignalSubsystem's flaw in deferred signalling - the issue was to rely on lambda-binding an array view.

diff --git a/Engine/Plugins/Runtime/MassGameplay/Source/MassSignals/Private/MassSignalSubsystem.cpp b/Engine/Plugins/Runtime/MassGameplay/Source/MassSignals/Private/MassSignalSubsystem.cpp
--- a/Engine/Plugins/Runtime/MassGameplay/Source/MassSignals/Private/MassSignalSubsystem.cpp
+++ b/Engine/Plugins/Runtime/MassGameplay/Source/MassSignals/Private/MassSignalSubsystem.cpp
@@ -96,10 +96,10 @@ void UMassSignalSubsystem::SignalEntityDeferred(FMassExecutionContext& Context,
 void UMassSignalSubsystem::SignalEntitiesDeferred(FMassExecutionContext& Context, FName SignalName, TConstArrayView<FMassEntityHandle> Entities)
 {
 	checkf(Entities.Num() > 0, TEXT("Expecting entities to signal"));
-	Context.Defer().PushCommand<FMassDeferredSetCommand>([SignalName, Entities](const FMassEntityManager& System)
+	Context.Defer().PushCommand<FMassDeferredSetCommand>([SignalName, InEntities = TArray<FMassEntityHandle>(Entities)](const FMassEntityManager& System)
 	{
 		UMassSignalSubsystem* SignalSubsystem = UWorld::GetSubsystem<UMassSignalSubsystem>(System.GetWorld());
-		SignalSubsystem->SignalEntities(SignalName, Entities);
+		SignalSubsystem->SignalEntities(SignalName, InEntities);
 	});
 
 	UE_CVLOG(Entities.Num() == 1, this, LogMassSignals, Log, TEXT("Raising deferred signal [%s] to entity [%s]"), *SignalName.ToString(), *Entities[0].DebugGetDescription());
@@ -115,10 +115,10 @@ void UMassSignalSubsystem::DelaySignalEntityDeferred(FMassExecutionContext& Cont
 void UMassSignalSubsystem::DelaySignalEntitiesDeferred(FMassExecutionContext& Context, FName SignalName, TConstArrayView<FMassEntityHandle> Entities, const float DelayInSeconds)
 {
 	checkf(Entities.Num() > 0, TEXT("Expecting entities to signal"));
-	Context.Defer().PushCommand<FMassDeferredSetCommand>([SignalName, Entities, DelayInSeconds](const FMassEntityManager& System)
+	Context.Defer().PushCommand<FMassDeferredSetCommand>([SignalName, InEntities = TArray<FMassEntityHandle>(Entities), DelayInSeconds](const FMassEntityManager& System)
 	{
 		UMassSignalSubsystem* SignalSubsystem = UWorld::GetSubsystem<UMassSignalSubsystem>(System.GetWorld());
-		SignalSubsystem->DelaySignalEntities(SignalName, Entities, DelayInSeconds);
+		SignalSubsystem->DelaySignalEntities(SignalName, InEntities, DelayInSeconds);
 	});
 
 	UE_CVLOG(Entities.Num() == 1, this, LogMassSignals, Log, TEXT("Delay deferred signal [%s] to entity [%s] in %.2f"), *SignalName.ToString(), *Entities[0].DebugGetDescription(), DelayInSeconds);
