Minor cosmetic tweaks to MassVisualizer.h, something that caught my eye.

diff --git a/Engine/Plugins/Runtime/MassGameplay/Source/MassRepresentation/Public/MassVisualizer.h b/Engine/Plugins/Runtime/MassGameplay/Source/MassRepresentation/Public/MassVisualizer.h
--- a/Engine/Plugins/Runtime/MassGameplay/Source/MassRepresentation/Public/MassVisualizer.h
+++ b/Engine/Plugins/Runtime/MassGameplay/Source/MassRepresentation/Public/MassVisualizer.h
@@ -3,7 +3,6 @@
 #pragma once
 
 #include "GameFramework/Actor.h"
-
 #include "MassVisualizer.generated.h"
 
 
@@ -20,10 +19,10 @@ class MASSREPRESENTATION_API AMassVisualizer : public AActor
 public:
 	AMassVisualizer();
 
-	/** Visualization component is garantee to exist if this class is created */
-	class UMassVisualizationComponent& GetVisualizationComponent() const { return *VisComponent; }
+	/** Visualization component is guaranteed to exist if this class is created */
+	UMassVisualizationComponent& GetVisualizationComponent() const { return *VisComponent; }
 
 protected:
 	UPROPERTY()
-	TObjectPtr<class UMassVisualizationComponent> VisComponent;
+	TObjectPtr<UMassVisualizationComponent> VisComponent;
 };
