PR #10527: [Mass] Improve assorted trait and fix a crash

diff --git a/Engine/Plugins/Runtime/MassGameplay/Source/MassRepresentation/Public/MassRepresentationTypes.h b/Engine/Plugins/Runtime/MassGameplay/Source/MassRepresentation/Public/MassRepresentationTypes.h
--- a/Engine/Plugins/Runtime/MassGameplay/Source/MassRepresentation/Public/MassRepresentationTypes.h
+++ b/Engine/Plugins/Runtime/MassGameplay/Source/MassRepresentation/Public/MassRepresentationTypes.h
@@ -291,7 +291,8 @@ public:
 			if(PrevLODSignificance >= 0.0f)
 			{
 				FMassLODSignificanceRange* PrevRange = GetLODSignificanceRange(PrevLODSignificance);
-				if (PrevRange != Range)
+				if (ensureMsgf(PrevRange, TEXT("Couldn't find a valid LODSignificanceRange for PrevLODSignificance %f"), PrevLODSignificance)
+					&& PrevRange != Range)
 				{
 					PrevRange->AddBatchedTransform(InstanceId, Transform, PrevTransform, Range->StaticMeshRefs);
 				}
@@ -318,7 +319,8 @@ public:
 			if(PrevLODSignificance >= 0.0f)
 			{
 				FMassLODSignificanceRange* PrevRange = GetLODSignificanceRange(PrevLODSignificance);
-				if (PrevRange != Range)
+				if (ensureMsgf(PrevRange, TEXT("Couldn't find a valid LODSignificanceRange for PrevLODSignificance %f"), PrevLODSignificance)
+					&& PrevRange != Range)
 				{
 					PrevRange->AddBatchedCustomData(InCustomData, Range->StaticMeshRefs, NumFloatsToPad);
 				}
@@ -334,7 +336,8 @@ public:
 			if(PrevLODSignificance >= 0.0f)
 			{
 				FMassLODSignificanceRange* PrevRange = GetLODSignificanceRange(PrevLODSignificance);
-				if (PrevRange != Range)
+				if (ensureMsgf(PrevRange, TEXT("Couldn't find a valid LODSignificanceRange for PrevLODSignificance %f"), PrevLODSignificance)
+					&& PrevRange != Range)
 				{
 					PrevRange->AddBatchedCustomDataFloats(CustomFloats, Range->StaticMeshRefs);
 				}
@@ -350,7 +353,8 @@ public:
 			if(PrevLODSignificance >= 0.0f)
 			{
 				FMassLODSignificanceRange* PrevRange = GetLODSignificanceRange(PrevLODSignificance);
-				if (PrevRange != Range)
+				if (ensureMsgf(PrevRange, TEXT("Couldn't find a valid LODSignificanceRange for PrevLODSignificance %f"), PrevLODSignificance)
+					&& PrevRange != Range)
 				{
 					PrevRange->WriteCustomDataFloatsAtStartIndex(StaticMeshIndex, CustomFloats, FloatsPerInstance, FloatStartIndex, Range->StaticMeshRefs);
 				}
diff --git a/Engine/Plugins/Runtime/MassGameplay/Source/MassSpawner/Private/MassAssortedFragmentsTrait.cpp b/Engine/Plugins/Runtime/MassGameplay/Source/MassSpawner/Private/MassAssortedFragmentsTrait.cpp
--- a/Engine/Plugins/Runtime/MassGameplay/Source/MassSpawner/Private/MassAssortedFragmentsTrait.cpp
+++ b/Engine/Plugins/Runtime/MassGameplay/Source/MassSpawner/Private/MassAssortedFragmentsTrait.cpp
@@ -16,14 +16,27 @@ void UMassAssortedFragmentsTrait::BuildTemplate(FMassEntityTemplateBuildContext&
 			{
 				BuildContext.AddFragment(Fragment);
 			}
-			else if (Type->IsChildOf(FMassTag::StaticStruct()))
+			else
+			{
+				UE_LOG(LogMass, Error, TEXT("Struct type %s is not a child of FMassFragment"), *GetPathNameSafe(Type));
+			}
+		}
+	}
+	
+	for (const FInstancedStruct& Tag : Tags)
+	{
+		if (Tag.IsValid())
+		{
+			const UScriptStruct* Type = Tag.GetScriptStruct();
+			CA_ASSUME(Type);
+			if (Type->IsChildOf(FMassTag::StaticStruct()))
 			{
 				BuildContext.AddTag(*Type);
 			}
 			else
 			{
-				UE_LOG(LogMass, Error, TEXT("Invalid fragment type found: '%s', nor a data fragment nor a tag"), *GetPathNameSafe(Type));
+				UE_LOG(LogMass, Error, TEXT("Struct type %s is not a child of FMassTag"), *GetPathNameSafe(Type));
 			}
 		}
 	}
-}
\ No newline at end of file
+}
diff --git a/Engine/Plugins/Runtime/MassGameplay/Source/MassSpawner/Public/MassAssortedFragmentsTrait.h b/Engine/Plugins/Runtime/MassGameplay/Source/MassSpawner/Public/MassAssortedFragmentsTrait.h
--- a/Engine/Plugins/Runtime/MassGameplay/Source/MassSpawner/Public/MassAssortedFragmentsTrait.h
+++ b/Engine/Plugins/Runtime/MassGameplay/Source/MassSpawner/Public/MassAssortedFragmentsTrait.h
@@ -19,6 +19,9 @@ protected:
 
 	virtual void BuildTemplate(FMassEntityTemplateBuildContext& BuildContext, const UWorld& World) const override;
 
-	UPROPERTY(Category="Fragments", EditAnywhere, meta = (BaseStruct = "/Script/MassEntity.MassFragment", ExcludeBaseStruct))
+	UPROPERTY(Category="Mass", EditAnywhere, meta = (BaseStruct = "/Script/MassEntity.MassFragment", ExcludeBaseStruct))
 	TArray<FInstancedStruct> Fragments;
+
+	UPROPERTY(Category="Mass", EditAnywhere, meta = (BaseStruct = "/Script/MassEntity.MassTag", ExcludeBaseStruct))
+	TArray<FInstancedStruct> Tags;
 };
