Moved the code in AMassLWIStaticMeshManager::EndPlay that's copying entity transforms back to InstanceTransforms array to a dedicated function. No change in functionality.

diff --git a/Engine/Plugins/Runtime/MassGameplay/Source/MassLWI/Private/MassLWIStaticMeshManager.cpp b/Engine/Plugins/Runtime/MassGameplay/Source/MassLWI/Private/MassLWIStaticMeshManager.cpp
--- a/Engine/Plugins/Runtime/MassGameplay/Source/MassLWI/Private/MassLWIStaticMeshManager.cpp
+++ b/Engine/Plugins/Runtime/MassGameplay/Source/MassLWI/Private/MassLWIStaticMeshManager.cpp
@@ -56,31 +56,7 @@ void AMassLWIStaticMeshManager::EndPlay(const EEndPlayReason::Type EndPlayReason
 			if (UMassEntitySubsystem* EntitySubsystem = UWorld::GetSubsystem<UMassEntitySubsystem>(GetWorld()))
 			{
 				FMassEntityManager& EntityManager = EntitySubsystem->GetMutableEntityManager();
-
-				TArray<FMassArchetypeEntityCollection> EntityCollectionsToDestroy;
-				UE::Mass::Utils::CreateEntityCollections(EntityManager, Entities, FMassArchetypeEntityCollection::FoldDuplicates, EntityCollectionsToDestroy);
-
-				// we create a query to "recreate" contents of InstanceTransforms just in case BeginPlay will be get
-				// called again for this LWIManager, which can happen with actor streaming
-				FMassEntityQuery LocationQuery;
-				LocationQuery.AddRequirement<FTransformFragment>(EMassFragmentAccess::ReadOnly);
-
-				FMassExecutionContext ExecutionContext(EntityManager);
-				for (FMassArchetypeEntityCollection& Collection : EntityCollectionsToDestroy)
-				{
-					LocationQuery.ForEachEntityChunk(Collection, EntityManager, ExecutionContext, [this](FMassExecutionContext& Context)
-					{
-						TConstArrayView<FTransformFragment> TransformsList = Context.GetFragmentView<FTransformFragment>();
-						for (const FTransformFragment& Fragment : TransformsList)
-						{
-							InstanceTransforms.Add(Fragment.GetTransform());
-						}
-					});
-
-					EntityManager.BatchDestroyEntityChunks(Collection);
-				}
-
-				Entities.Reset();
+				StoreMassDataInActor(EntityManager);
 			}
 		}
 
@@ -147,6 +123,34 @@ void AMassLWIStaticMeshManager::TransferDataToMass(FMassEntityManager& EntityMan
 	}
 }
 
+void AMassLWIStaticMeshManager::StoreMassDataInActor(FMassEntityManager& EntityManager)
+{
+	TArray<FMassArchetypeEntityCollection> EntityCollectionsToDestroy;
+	UE::Mass::Utils::CreateEntityCollections(EntityManager, Entities, FMassArchetypeEntityCollection::FoldDuplicates, EntityCollectionsToDestroy);
+
+	// we create a query to "recreate" contents of InstanceTransforms just in case BeginPlay will be get
+	// called again for this LWIManager, which can happen with actor streaming
+	FMassEntityQuery LocationQuery;
+	LocationQuery.AddRequirement<FTransformFragment>(EMassFragmentAccess::ReadOnly);
+
+	FMassExecutionContext ExecutionContext(EntityManager);
+	for (FMassArchetypeEntityCollection& Collection : EntityCollectionsToDestroy)
+	{
+		LocationQuery.ForEachEntityChunk(Collection, EntityManager, ExecutionContext, [this](FMassExecutionContext& Context)
+		{
+			TConstArrayView<FTransformFragment> TransformsList = Context.GetFragmentView<FTransformFragment>();
+			for (const FTransformFragment& Fragment : TransformsList)
+			{
+				InstanceTransforms.Add(Fragment.GetTransform());
+			}
+		});
+
+		EntityManager.BatchDestroyEntityChunks(Collection);
+	}
+
+	Entities.Reset();
+}
+
 int32 AMassLWIStaticMeshManager::FindIndexForEntity(const FMassEntityHandle Entity) const
 {
 	return Entities.IndexOfByKey(Entity);
diff --git a/Engine/Plugins/Runtime/MassGameplay/Source/MassLWI/Public/MassLWIStaticMeshManager.h b/Engine/Plugins/Runtime/MassGameplay/Source/MassLWI/Public/MassLWIStaticMeshManager.h
--- a/Engine/Plugins/Runtime/MassGameplay/Source/MassLWI/Public/MassLWIStaticMeshManager.h
+++ b/Engine/Plugins/Runtime/MassGameplay/Source/MassLWI/Public/MassLWIStaticMeshManager.h
@@ -29,7 +29,8 @@ class MASSLWI_API AMassLWIStaticMeshManager : public ALightWeightInstanceStaticM
 public:
 	bool IsRegisteredWithMass() const { return MassRegistrationHandle.IsValid(); }
 
-	void TransferDataToMass(FMassEntityManager& EntityManager);
+	virtual void TransferDataToMass(FMassEntityManager& EntityManager);
+	virtual void StoreMassDataInActor(FMassEntityManager& EntityManager);
 	FMassLWIManagerRegistrationHandle GetMassRegistrationHandle() const { return MassRegistrationHandle; }
 	void MarkRegisteredWithMass(const FMassLWIManagerRegistrationHandle RegistrationIndex);
 	void MarkUnregisteredWithMass();
