Removed a couple of operators from structs in MassRepresenationTypes.h to improve code readability (less hidden operations taking place)

diff --git a/Engine/Plugins/Runtime/MassGameplay/Source/MassRepresentation/Private/MassVisualizationComponent.cpp b/Engine/Plugins/Runtime/MassGameplay/Source/MassRepresentation/Private/MassVisualizationComponent.cpp
--- a/Engine/Plugins/Runtime/MassGameplay/Source/MassRepresentation/Private/MassVisualizationComponent.cpp
+++ b/Engine/Plugins/Runtime/MassGameplay/Source/MassRepresentation/Private/MassVisualizationComponent.cpp
@@ -36,14 +36,14 @@ void UMassVisualizationComponent::PostInitProperties()
 int16 UMassVisualizationComponent::FindOrAddVisualDesc(const FStaticMeshInstanceVisualizationDesc& Desc)
 {
 	UE_MT_SCOPED_WRITE_ACCESS(InstancedStaticMeshInfosDetector);
-	int32 VisualIndex = InstancedStaticMeshInfos.IndexOfByPredicate([&Desc](const FMassInstancedStaticMeshInfo& Info) { return Info == Desc; });
+	int32 VisualIndex = InstancedStaticMeshInfos.IndexOfByPredicate([&Desc](const FMassInstancedStaticMeshInfo& Info) { return Info.GetDesc() == Desc; });
 	if (VisualIndex == INDEX_NONE)
 	{
 		VisualIndex = InstancedStaticMeshInfos.Emplace(Desc);
 
 		for (const FMassStaticMeshInstanceVisualizationMeshDesc& MeshDesc : Desc.Meshes)
 		{
-			ISMCSharedData.FindOrAdd(MeshDesc, FMassISMCSharedData());
+			ISMCSharedData.FindOrAdd(GetTypeHash(MeshDesc), FMassISMCSharedData());
 		}
 
 		BuildLODSignificanceForInfo(InstancedStaticMeshInfos[VisualIndex]);
@@ -76,7 +76,7 @@ void UMassVisualizationComponent::ConstructStaticMeshComponents()
 		}
 		for (const FMassStaticMeshInstanceVisualizationMeshDesc& MeshDesc : Info.Desc.Meshes)
 		{
-			FMassISMCSharedData* SharedData = ISMCSharedData.Find(MeshDesc);
+			FMassISMCSharedData* SharedData = ISMCSharedData.Find(GetTypeHash(MeshDesc));
 			UInstancedStaticMeshComponent* ISMC = SharedData ? SharedData->GetISMComponent() : nullptr;
 
 			if (SharedData == nullptr || SharedData->GetISMComponent() == nullptr)
@@ -104,7 +104,7 @@ void UMassVisualizationComponent::ConstructStaticMeshComponents()
 
 				if (SharedData == nullptr)
 				{
-					SharedData = &ISMCSharedData.Emplace(MeshDesc, FMassISMCSharedData(ISMC));
+					SharedData = &ISMCSharedData.Emplace(GetTypeHash(MeshDesc), FMassISMCSharedData(ISMC));
 				}
 				else
 				{
@@ -166,7 +166,7 @@ void UMassVisualizationComponent::BuildLODSignificanceForInfo(FMassInstancedStat
 				const bool bAddMeshInRange = (Range.MinSignificance >= MeshDesc.MinLODSignificance && Range.MinSignificance < MeshDesc.MaxLODSignificance);
 				if (bAddMeshInRange)
 				{
-					Range.StaticMeshRefs.Add(MeshDesc);
+					Range.StaticMeshRefs.Add(GetTypeHash(MeshDesc));
 				}
 			}
 		}
@@ -310,7 +310,8 @@ void FMassInstancedStaticMeshInfo::ClearVisualInstance(FMassISMCSharedDataMap& I
 {
 	for (int i = 0; i < Desc.Meshes.Num(); i++)
 	{
-		FMassISMCSharedData* SharedData = ISMCSharedData.Find(Desc.Meshes[i]);
+		const uint32 MeshDescHash = GetTypeHash(Desc.Meshes[i]);
+		FMassISMCSharedData* SharedData = ISMCSharedData.Find(MeshDescHash);
 		if (SharedData && SharedData->ReleaseReference() == 0)
 		{
 			if (UInstancedStaticMeshComponent* ISMC = SharedData->GetISMComponent())
@@ -318,7 +319,7 @@ void FMassInstancedStaticMeshInfo::ClearVisualInstance(FMassISMCSharedDataMap& I
 				ISMC->ClearInstances();
 				ISMC->DestroyComponent();
 			}
-			ISMCSharedData.Remove(Desc.Meshes[i]);
+			ISMCSharedData.Remove(MeshDescHash);
 		}
 	}
 
diff --git a/Engine/Plugins/Runtime/MassGameplay/Source/MassRepresentation/Public/MassRepresentationTypes.h b/Engine/Plugins/Runtime/MassGameplay/Source/MassRepresentation/Public/MassRepresentationTypes.h
--- a/Engine/Plugins/Runtime/MassGameplay/Source/MassRepresentation/Public/MassRepresentationTypes.h
+++ b/Engine/Plugins/Runtime/MassGameplay/Source/MassRepresentation/Public/MassRepresentationTypes.h
@@ -100,11 +100,6 @@ struct MASSREPRESENTATION_API FMassStaticMeshInstanceVisualizationMeshDesc
 		}
 		return Hash;
 	}
-
-	operator uint32() const
-	{
-		return GetTypeHash(*this);
-	}
 };
 
 USTRUCT()
@@ -257,11 +252,6 @@ public:
 	{
 	}
 
-	bool operator==(const FStaticMeshInstanceVisualizationDesc& OtherDesc) const
-	{
-		return Desc == OtherDesc;
-	}
-	
 	const FStaticMeshInstanceVisualizationDesc& GetDesc() const
 	{
 		return Desc;
