Minor change to AMassLWIStaticMeshManager::CreateMassTemplate to support users not putting any mesh info in the entity config being used. Also verified the data to ensure there's exactly one mesh description and that it covers the whole significance range.

diff --git a/Engine/Plugins/Runtime/MassGameplay/Source/MassLWI/Private/MassLWIStaticMeshManager.cpp b/Engine/Plugins/Runtime/MassGameplay/Source/MassLWI/Private/MassLWIStaticMeshManager.cpp
--- a/Engine/Plugins/Runtime/MassGameplay/Source/MassLWI/Private/MassLWIStaticMeshManager.cpp
+++ b/Engine/Plugins/Runtime/MassGameplay/Source/MassLWI/Private/MassLWIStaticMeshManager.cpp
@@ -5,6 +5,7 @@
 #include "MassLWISubsystem.h"
 #include "MassLWITypes.h"
 #include "Components/HierarchicalInstancedStaticMeshComponent.h"
+#include "MassInstancedStaticMeshComponent.h"
 #include "MassCommonFragments.h"
 #include "MassEntitySubsystem.h"
 #include "MassEntityView.h"
@@ -206,12 +207,6 @@ void AMassLWIStaticMeshManager::CreateMassTemplate(FMassEntityManager& EntityMan
 	{
 		return;
 	}
-	if (!ensureMsgf(VisTrait->StaticMeshInstanceDesc.Meshes.Num(), TEXT("The VisualizationTrait.StaticMeshInstanceDesc.Meshes being used needs to have at least one entry")))
-	{
-		return;
-	}
-	ensureMsgf(VisTrait->StaticMeshInstanceDesc.Meshes.Num() == 1
-		, TEXT("It's recommended for the VisualizationTrait.StaticMeshInstanceDesc.Meshes being used to have exactly one entry, since only the first entry will be used"));
 
 	UMassRepresentationSubsystem* RepresentationSubsystem = nullptr;
 	
@@ -232,8 +227,19 @@ void AMassLWIStaticMeshManager::CreateMassTemplate(FMassEntityManager& EntityMan
 	}
 	
 	FStaticMeshInstanceVisualizationDesc StaticMeshInstanceDesc = VisTrait->StaticMeshInstanceDesc;
-	// @todo make sure a nullptr here won't break stuff
+	
+	// we don't care about the mesh that has been set in the source template, we're overriding it anyway.
+	StaticMeshInstanceDesc.Meshes.SetNumZeroed(1);
 	FMassStaticMeshInstanceVisualizationMeshDesc& MeshDesc = StaticMeshInstanceDesc.Meshes[0];
+
+	if (!MeshDesc.ISMComponentClass)
+	{
+		MeshDesc.ISMComponentClass = UMassInstancedStaticMeshComponent::StaticClass();
+	}
+
+	// forcing the "full range" since we only ever expect there to be one mesh for the relevant ISM component
+	MeshDesc.SetSignificanceRange(EMassLOD::High, EMassLOD::Max);
+	
 	MeshDesc.Mesh = StaticMesh.Get();
 
 	if (!ensure(MeshDesc.Mesh))
@@ -249,6 +255,7 @@ void AMassLWIStaticMeshManager::CreateMassTemplate(FMassEntityManager& EntityMan
 	if (InstancedStaticMeshComponent)
 	{
 		MeshDesc.MaterialOverrides = InstancedStaticMeshComponent->OverrideMaterials;
+		MeshDesc.bCastShadows = (InstancedStaticMeshComponent->CastShadow != 0);
 	}
 
 	check(MeshDesc.Mesh);
@@ -260,6 +267,7 @@ void AMassLWIStaticMeshManager::CreateMassTemplate(FMassEntityManager& EntityMan
 		RepresentationFragment.StaticMeshDescIndex = RepresentationSubsystem->FindOrAddStaticMeshDesc(StaticMeshInstanceDesc);
 		const int32 TemplateActorIndex = RepresentationSubsystem->FindOrAddTemplateActor(RepresentedClass);
 		RepresentationFragment.HighResTemplateActorIndex = TemplateActorIndex;
+		// leaving for reference here, since at some point we probably will need to set it.
 		//RepresentationFragment.LowResTemplateActorIndex = TemplateActorIndex;
 	}
 
diff --git a/Engine/Plugins/Runtime/MassGameplay/Source/MassRepresentation/Public/MassRepresentationTypes.h b/Engine/Plugins/Runtime/MassGameplay/Source/MassRepresentation/Public/MassRepresentationTypes.h
--- a/Engine/Plugins/Runtime/MassGameplay/Source/MassRepresentation/Public/MassRepresentationTypes.h
+++ b/Engine/Plugins/Runtime/MassGameplay/Source/MassRepresentation/Public/MassRepresentationTypes.h
@@ -100,6 +100,14 @@ struct MASSREPRESENTATION_API FMassStaticMeshInstanceVisualizationMeshDesc
 		}
 		return Hash;
 	}
+
+	// convenience function for setting MinLODSinificance and MaxLODSinificance based on EMassLOD values
+	void SetSignificanceRange(const EMassLOD::Type MinLOD, const EMassLOD::Type MaxLOD)
+	{
+		checkSlow(MinLOD <= MaxLOD);
+		MinLODSignificance = float(MinLOD);
+		MaxLODSignificance = float(MaxLOD);
+	}
 };
 
 USTRUCT()
