Fixed data validity testing in UMassVisualizationComponent::FindOrAddVisualDesc

diff --git a/Engine/Plugins/Runtime/MassGameplay/Source/MassRepresentation/Private/MassVisualizationComponent.cpp b/Engine/Plugins/Runtime/MassGameplay/Source/MassRepresentation/Private/MassVisualizationComponent.cpp
--- a/Engine/Plugins/Runtime/MassGameplay/Source/MassRepresentation/Private/MassVisualizationComponent.cpp
+++ b/Engine/Plugins/Runtime/MassGameplay/Source/MassRepresentation/Private/MassVisualizationComponent.cpp
@@ -43,7 +43,7 @@ int16 UMassVisualizationComponent::FindOrAddVisualDesc(const FStaticMeshInstance
 
 		for (const FMassStaticMeshInstanceVisualizationMeshDesc& MeshDesc : Desc.Meshes)
 		{
-			if (MeshDesc.Mesh)
+			if (MeshDesc.Mesh && MeshDesc.ISMComponentClass)
 			{
 				ISMCSharedData.FindOrAdd(GetTypeHash(MeshDesc), FMassISMCSharedData());
 				bValidDescription = true;
