Added a ShouldAllowQueryBasedPruning override to Mass signaling processor base to make processors signals not get pruned by default, since that can mess up the initial signaling processing

diff --git a/Engine/Plugins/Runtime/MassGameplay/Source/MassSignals/Public/MassSignalProcessorBase.h b/Engine/Plugins/Runtime/MassGameplay/Source/MassSignals/Public/MassSignalProcessorBase.h
--- a/Engine/Plugins/Runtime/MassGameplay/Source/MassSignals/Public/MassSignalProcessorBase.h
+++ b/Engine/Plugins/Runtime/MassGameplay/Source/MassSignals/Public/MassSignalProcessorBase.h
@@ -51,6 +51,13 @@ protected:
 	 */
 	virtual void Execute(FMassEntityManager& EntityManager, FMassExecutionContext& Context) override;
 
+	/** 
+	 * We don't want signaling processors to be pruned by default, since that can mess up signal processing 
+	 * just after processor's creation (might miss some signals initially).
+	 */
+	virtual bool ShouldAllowQueryBasedPruning(const bool bRuntimeMode = true) const override { return false; }
+
+
 	/**
 	 * To receive notification about a particular signal, you need to subscribe to it.
 	 * @param SignalName is the name of the signal to receive notification about
