Fix non unity build

diff --git a/Engine/Plugins/Runtime/MassGameplay/Source/MassSmartObjects/Public/MassSmartObjectTypes.h b/Engine/Plugins/Runtime/MassGameplay/Source/MassSmartObjects/Public/MassSmartObjectTypes.h
--- a/Engine/Plugins/Runtime/MassGameplay/Source/MassSmartObjects/Public/MassSmartObjectTypes.h
+++ b/Engine/Plugins/Runtime/MassGameplay/Source/MassSmartObjects/Public/MassSmartObjectTypes.h
@@ -2,7 +2,7 @@
 
 #pragma once
 
-#include "SmartObjectTypes.h"
+#include "MassEntityTypes.h"
 #include "MassSmartObjectTypes.generated.h"
 
 namespace UE::Mass::Signals
