Made count-based LOD adjusting by [REDACTED_PERSON] processors optional.
This is a temporary solution to a discovered issue, the proper fix is WIP.

diff --git a/Engine/Plugins/Runtime/MassGameplay/Source/MassLOD/Private/MassSimulationLOD.cpp b/Engine/Plugins/Runtime/MassGameplay/Source/MassLOD/Private/MassSimulationLOD.cpp
--- a/Engine/Plugins/Runtime/MassGameplay/Source/MassLOD/Private/MassSimulationLOD.cpp
+++ b/Engine/Plugins/Runtime/MassGameplay/Source/MassLOD/Private/MassSimulationLOD.cpp
@@ -140,6 +140,7 @@ void UMassSimulationLODProcessor::Execute(FMassEntityManager& EntityManager, FMa
 		});
 	}
 
+	if (bDoAdjustmentFromCount)
 	{
 		TRACE_CPUPROFILER_EVENT_SCOPE(TEXT("AdjustDistancesAndLODFromCount"));
 		EntityManager.ForEachSharedFragment<FMassSimulationLODSharedFragment>([](FMassSimulationLODSharedFragment& LODSharedFragment)
diff --git a/Engine/Plugins/Runtime/MassGameplay/Source/MassLOD/Public/MassSimulationLOD.h b/Engine/Plugins/Runtime/MassGameplay/Source/MassLOD/Public/MassSimulationLOD.h
--- a/Engine/Plugins/Runtime/MassGameplay/Source/MassLOD/Public/MassSimulationLOD.h
+++ b/Engine/Plugins/Runtime/MassGameplay/Source/MassLOD/Public/MassSimulationLOD.h
@@ -181,4 +181,7 @@ protected:
 	FMassEntityQuery EntityQueryAdjustDistances;
 	FMassEntityQuery EntityQueryVariableTick;
 	FMassEntityQuery EntityQuerySetLODTag;
+
+	UPROPERTY(config, EditDefaultsOnly, Category="Mass")
+	bool bDoAdjustmentFromCount = true;
 };
diff --git a/Engine/Plugins/Runtime/MassGameplay/Source/MassRepresentation/Private/MassVisualizationLODProcessor.cpp b/Engine/Plugins/Runtime/MassGameplay/Source/MassRepresentation/Private/MassVisualizationLODProcessor.cpp
--- a/Engine/Plugins/Runtime/MassGameplay/Source/MassRepresentation/Private/MassVisualizationLODProcessor.cpp
+++ b/Engine/Plugins/Runtime/MassGameplay/Source/MassRepresentation/Private/MassVisualizationLODProcessor.cpp
@@ -91,6 +91,7 @@ void UMassVisualizationLODProcessor::Execute(FMassEntityManager& EntityManager,
 		FarEntityQuery.ForEachEntityChunk(EntityManager, Context, CalculateLOD);
 	}
 
+	if (bDoAdjustmentFromCount)
 	{
 		TRACE_CPUPROFILER_EVENT_SCOPE(AdjustDistanceAndLODFromCount)
 		EntityManager.ForEachSharedFragment<FMassVisualizationLODSharedFragment>([this](FMassVisualizationLODSharedFragment& LODSharedFragment)
diff --git a/Engine/Plugins/Runtime/MassGameplay/Source/MassRepresentation/Public/MassVisualizationLODProcessor.h b/Engine/Plugins/Runtime/MassGameplay/Source/MassRepresentation/Public/MassVisualizationLODProcessor.h
--- a/Engine/Plugins/Runtime/MassGameplay/Source/MassRepresentation/Public/MassVisualizationLODProcessor.h
+++ b/Engine/Plugins/Runtime/MassGameplay/Source/MassRepresentation/Public/MassVisualizationLODProcessor.h
@@ -45,4 +45,7 @@ protected:
 
 	UPROPERTY(Transient)
 	TObjectPtr<const UScriptStruct> FilterTag = nullptr;
+
+	UPROPERTY(config, EditDefaultsOnly, Category = "Mass")
+	bool bDoAdjustmentFromCount = true;
 };
\ No newline at end of file
