Fixed typos in log messages
#rnx

diff --git a/Engine/Plugins/Runtime/MassGameplay/Source/MassSpawner/Private/MassEntitySpawnDataGeneratorBase.cpp b/Engine/Plugins/Runtime/MassGameplay/Source/MassSpawner/Private/MassEntitySpawnDataGeneratorBase.cpp
--- a/Engine/Plugins/Runtime/MassGameplay/Source/MassSpawner/Private/MassEntitySpawnDataGeneratorBase.cpp
+++ b/Engine/Plugins/Runtime/MassGameplay/Source/MassSpawner/Private/MassEntitySpawnDataGeneratorBase.cpp
@@ -14,7 +14,7 @@ void UMassEntitySpawnDataGeneratorBase::BuildResultsFromEntityTypes(const int32
 
 	if (TotalProportion <= 0)
 	{
-		UE_VLOG_UELOG(this, LogMassSpawner, Error, TEXT("The total combined porportion of all the entity types needs to be greater than 0.0f."));
+		UE_VLOG_UELOG(this, LogMassSpawner, Error, TEXT("The total combined proportion of all the entity types needs to be greater than 0."));
 		return;
 	}
 
diff --git a/Engine/Plugins/Runtime/MassGameplay/Source/MassSpawner/Private/MassSpawner.cpp b/Engine/Plugins/Runtime/MassGameplay/Source/MassSpawner/Private/MassSpawner.cpp
--- a/Engine/Plugins/Runtime/MassGameplay/Source/MassSpawner/Private/MassSpawner.cpp
+++ b/Engine/Plugins/Runtime/MassGameplay/Source/MassSpawner/Private/MassSpawner.cpp
@@ -300,7 +300,7 @@ void AMassSpawner::DoSpawning()
 
 	if (TotalProportion <= 0.0f)
 	{
-		UE_VLOG_UELOG(this, LogMassSpawner, Error, TEXT("The total combined porportion of all the generator needs to be greater than 0.0f."));
+		UE_VLOG_UELOG(this, LogMassSpawner, Error, TEXT("The total combined proportion of all the generator needs to be greater than 0."));
 		return;
 	}
 	
