Fixed a condition in TMassLODCalculator<FLODLogic>::AdjustLODFromCount that impaired finding the highest LOD level for a given entity

diff --git a/Engine/Plugins/Runtime/MassGameplay/Source/MassLOD/Public/MassLODCalculator.h b/Engine/Plugins/Runtime/MassGameplay/Source/MassLOD/Public/MassLODCalculator.h
--- a/Engine/Plugins/Runtime/MassGameplay/Source/MassLOD/Public/MassLODCalculator.h
+++ b/Engine/Plugins/Runtime/MassGameplay/Source/MassLOD/Public/MassLODCalculator.h
@@ -669,7 +669,7 @@ void TMassLODCalculator<FLODLogic>::AdjustLODFromCount(FMassExecutionContext& Co
 
 				LODPerViewer = ComputeLODFromSettings<bCalculateVisibilityPerViewer>(LODPerViewer, DistanceToViewerSq, bIsVisibleByViewer, nullptr, RuntimeDataPerViewer[ViewerIdx]);
 
-				if (HighestViewerLOD < LODPerViewer)
+				if (LODPerViewer < HighestViewerLOD)
 				{
 					HighestViewerLOD = LODPerViewer;
 				}
@@ -682,8 +682,8 @@ void TMassLODCalculator<FLODLogic>::AdjustLODFromCount(FMassExecutionContext& Co
 		const bool bIsVisibleByAViewer = GetPrevVisibility<bCalculateVisibility>(EntityLOD, EMassVisibility::Max) == EMassVisibility::CanBeSeen;
 		EMassLOD::Type NewLOD = ComputeLODFromSettings<bCalculateVisibility>(EntityLOD.PrevLOD, EntityViewersInfo.ClosestViewerDistanceSq, bIsVisibleByAViewer, nullptr, RuntimeData);
 
-		// Maybe the highest of all the viewers is now lower than the global entity LOD, make sure to update the it accordingly
-		if (bMaximizeCountPerViewer && NewLOD < HighestViewerLOD)
+		// Maybe the highest of all the viewers is now lower than the global entity LOD, make sure to update it accordingly
+		if (bMaximizeCountPerViewer && NewLOD > HighestViewerLOD)
 		{
 			NewLOD = HighestViewerLOD;
 		}
