Fixed Mass vis actors not going away when player moves away

diff --git a/Engine/Plugins/Runtime/MassGameplay/Source/MassRepresentation/Private/MassStationaryISMSwitcherProcessor.cpp b/Engine/Plugins/Runtime/MassGameplay/Source/MassRepresentation/Private/MassStationaryISMSwitcherProcessor.cpp
--- a/Engine/Plugins/Runtime/MassGameplay/Source/MassRepresentation/Private/MassStationaryISMSwitcherProcessor.cpp
+++ b/Engine/Plugins/Runtime/MassGameplay/Source/MassRepresentation/Private/MassStationaryISMSwitcherProcessor.cpp
@@ -20,6 +20,7 @@ void UMassStationaryISMSwitcherProcessor::ConfigureQueries()
 	EntityQuery.AddRequirement<FMassRepresentationLODFragment>(EMassFragmentAccess::ReadOnly);
 	EntityQuery.AddRequirement<FTransformFragment>(EMassFragmentAccess::ReadOnly);
 	EntityQuery.AddRequirement<FMassRepresentationFragment>(EMassFragmentAccess::ReadWrite);
+	EntityQuery.AddConstSharedRequirement<FMassRepresentationParameters>();
 	EntityQuery.AddSharedRequirement<FMassRepresentationSubsystemSharedFragment>(EMassFragmentAccess::ReadWrite);
 	EntityQuery.AddTagRequirement<FMassStaticRepresentationTag>(EMassFragmentPresence::All);
 	EntityQuery.AddSubsystemRequirement<UMassSignalSubsystem>(EMassFragmentAccess::ReadWrite);
@@ -39,6 +40,9 @@ void UMassStationaryISMSwitcherProcessor::Execute(FMassEntityManager& EntityMana
 		const TConstArrayView<FTransformFragment> TransformList = Context.GetFragmentView<FTransformFragment>();
 		const TArrayView<FMassRepresentationFragment> RepresentationList = Context.GetMutableFragmentView<FMassRepresentationFragment>();
 
+		const FMassRepresentationParameters& RepresentationParams = Context.GetConstSharedFragment<FMassRepresentationParameters>();
+		const bool bDoKeepActorExtraFrame = UE::MassRepresentation::bAllowKeepActorExtraFrame ? RepresentationParams.bKeepLowResActors : false;
+
 		const int32 NumEntities = Context.GetNumEntities();
 		for (int32 EntityIdx = 0; EntityIdx < NumEntities; EntityIdx++)
 		{
@@ -97,7 +101,13 @@ void UMassStationaryISMSwitcherProcessor::Execute(FMassEntityManager& EntityMana
 				}
 
 				// consume "prev" data
-				Representation.PrevRepresentation = Representation.CurrentRepresentation;
+				// @note crazy hacky, but we don't want to consume if bDoKeepActorExtraFrame is true. In that case 
+				// UMassRepresentationProcessor::UpdateRepresentation expects the "prev" state not to be consumed 
+				// a frame longer so that it can do the consuming (and call "disable actor").
+				if (bDoKeepActorExtraFrame == false)
+				{
+					Representation.PrevRepresentation = Representation.CurrentRepresentation;
+				}
 
 				SignalSubsystem.SignalEntity(UE::Mass::Signals::SwitchedToISM, EntityHandle);
 			}
diff --git a/Engine/Plugins/Runtime/MassGameplay/Source/MassRepresentation/Public/MassRepresentationProcessor.h b/Engine/Plugins/Runtime/MassGameplay/Source/MassRepresentation/Public/MassRepresentationProcessor.h
--- a/Engine/Plugins/Runtime/MassGameplay/Source/MassRepresentation/Public/MassRepresentationProcessor.h
+++ b/Engine/Plugins/Runtime/MassGameplay/Source/MassRepresentation/Public/MassRepresentationProcessor.h
@@ -13,6 +13,12 @@ class UMassRepresentationSubsystem;
 class UMassActorSubsystem;
 struct FMassActorFragment;
 
+
+namespace UE::MassRepresentation
+{
+	extern MASSREPRESENTATION_API int32 bAllowKeepActorExtraFrame;
+}
+
 UCLASS()
 class MASSREPRESENTATION_API UMassRepresentationProcessor : public UMassProcessor
 {
