Renamed namesapce `UE::MassRepresentation` to `UE::Mass::Representation`

diff --git a/Engine/Plugins/Runtime/MassGameplay/Source/MassRepresentation/Private/MassRepresentationProcessor.cpp b/Engine/Plugins/Runtime/MassGameplay/Source/MassRepresentation/Private/MassRepresentationProcessor.cpp
--- a/Engine/Plugins/Runtime/MassGameplay/Source/MassRepresentation/Private/MassRepresentationProcessor.cpp
+++ b/Engine/Plugins/Runtime/MassGameplay/Source/MassRepresentation/Private/MassRepresentationProcessor.cpp
@@ -14,7 +14,7 @@
 #include "MassExecutionContext.h"
 
 
-namespace UE::MassRepresentation
+namespace UE::Mass::Representation
 {
 	int32 bAllowKeepActorExtraFrame = 1;
 	FAutoConsoleVariableRef CVarAllowKeepActorExtraFrame(TEXT("ai.massrepresentation.AllowKeepActorExtraFrame"), bAllowKeepActorExtraFrame, TEXT("Allow the mass representation to keep actor an extra frame when switching to ISM"), ECVF_Default);
@@ -59,7 +59,7 @@ void UMassRepresentationProcessor::UpdateRepresentation(FMassExecutionContext& C
 	const TConstArrayView<FMassRepresentationLODFragment> RepresentationLODList = Context.GetFragmentView<FMassRepresentationLODFragment>();
 	const TArrayView<FMassActorFragment> ActorList = Context.GetMutableFragmentView<FMassActorFragment>();
 
-	const bool bDoKeepActorExtraFrame = UE::MassRepresentation::bAllowKeepActorExtraFrame ? RepresentationParams.bKeepLowResActors : false;
+	const bool bDoKeepActorExtraFrame = UE::Mass::Representation::bAllowKeepActorExtraFrame ? RepresentationParams.bKeepLowResActors : false;
 
 	const int32 NumEntities = Context.GetNumEntities();
 	for (int32 EntityIdx = 0; EntityIdx < NumEntities; EntityIdx++)
@@ -309,13 +309,13 @@ FMassVisualizationChunkFragment& UMassVisualizationProcessor::UpdateChunkVisibil
 	if (ChunkVisibility == EMassVisibility::Max)
 	{
 		// The visibility on the chunk fragment data isn't set yet, let see if the Archetype has an visibility tag and set it on the ChunkData
-		ChunkVisibility = UE::MassRepresentation::GetVisibilityFromArchetype(Context);
+		ChunkVisibility = UE::Mass::Representation::GetVisibilityFromArchetype(Context);
 		ChunkData.SetVisibility(ChunkVisibility);
 		bFirstUpdate = RepresentationParams.bSpreadFirstVisualizationUpdate;
 	}
 	else
 	{
-		checkfSlow(UE::MassRepresentation::IsVisibilityTagSet(Context, ChunkVisibility), TEXT("Expecting the same Visibility as what we saved in the chunk data, maybe external code is modifying the tags"))
+		checkfSlow(UE::Mass::Representation::IsVisibilityTagSet(Context, ChunkVisibility), TEXT("Expecting the same Visibility as what we saved in the chunk data, maybe external code is modifying the tags"))
 	}
 
 	if (ChunkVisibility == EMassVisibility::CulledByDistance)
@@ -349,7 +349,7 @@ void UMassVisualizationProcessor::UpdateEntityVisibility(const FMassEntityHandle
 	const EMassVisibility ChunkVisibility = ChunkData.GetVisibility();
 	if (ChunkVisibility != Visibility)
 	{
-		UE::MassRepresentation::PushSwapTagsCommand(CommandBuffer, Entity, ChunkVisibility, Visibility);
+		UE::Mass::Representation::PushSwapTagsCommand(CommandBuffer, Entity, ChunkVisibility, Visibility);
 		ChunkData.SetContainsNewlyVisibleEntity(Visibility == EMassVisibility::CanBeSeen);
 	}
 }
diff --git a/Engine/Plugins/Runtime/MassGameplay/Source/MassRepresentation/Private/MassRepresentationTypes.cpp b/Engine/Plugins/Runtime/MassGameplay/Source/MassRepresentation/Private/MassRepresentationTypes.cpp
--- a/Engine/Plugins/Runtime/MassGameplay/Source/MassRepresentation/Private/MassRepresentationTypes.cpp
+++ b/Engine/Plugins/Runtime/MassGameplay/Source/MassRepresentation/Private/MassRepresentationTypes.cpp
@@ -7,7 +7,7 @@
 
 DEFINE_LOG_CATEGORY(LogMassRepresentation);
 
-namespace UE::MassRepresentation
+namespace UE::Mass::Representation
 {
 void PushSwapTagsCommand(FMassCommandBuffer& CommandBuffer, const FMassEntityHandle Entity, const EMassVisibility PrevVisibility, const EMassVisibility NewVisibility)
 {
@@ -73,7 +73,7 @@ default: \
 
 #undef CASE_SWAP_TAGS
 }
-} // UE::MassRepresentation
+} // UE::Mass::Representation
 
 
 //-----------------------------------------------------------------------------
diff --git a/Engine/Plugins/Runtime/MassGameplay/Source/MassRepresentation/Private/MassStationaryISMSwitcherProcessor.cpp b/Engine/Plugins/Runtime/MassGameplay/Source/MassRepresentation/Private/MassStationaryISMSwitcherProcessor.cpp
--- a/Engine/Plugins/Runtime/MassGameplay/Source/MassRepresentation/Private/MassStationaryISMSwitcherProcessor.cpp
+++ b/Engine/Plugins/Runtime/MassGameplay/Source/MassRepresentation/Private/MassStationaryISMSwitcherProcessor.cpp
@@ -41,7 +41,7 @@ void UMassStationaryISMSwitcherProcessor::Execute(FMassEntityManager& EntityMana
 		const TArrayView<FMassRepresentationFragment> RepresentationList = Context.GetMutableFragmentView<FMassRepresentationFragment>();
 
 		const FMassRepresentationParameters& RepresentationParams = Context.GetConstSharedFragment<FMassRepresentationParameters>();
-		const bool bDoKeepActorExtraFrame = UE::MassRepresentation::bAllowKeepActorExtraFrame ? RepresentationParams.bKeepLowResActors : false;
+		const bool bDoKeepActorExtraFrame = UE::Mass::Representation::bAllowKeepActorExtraFrame ? RepresentationParams.bKeepLowResActors : false;
 
 		const int32 NumEntities = Context.GetNumEntities();
 		for (int32 EntityIdx = 0; EntityIdx < NumEntities; EntityIdx++)
diff --git a/Engine/Plugins/Runtime/MassGameplay/Source/MassRepresentation/Private/MassVisualizationComponent.cpp b/Engine/Plugins/Runtime/MassGameplay/Source/MassRepresentation/Private/MassVisualizationComponent.cpp
--- a/Engine/Plugins/Runtime/MassGameplay/Source/MassRepresentation/Private/MassVisualizationComponent.cpp
+++ b/Engine/Plugins/Runtime/MassGameplay/Source/MassRepresentation/Private/MassVisualizationComponent.cpp
@@ -18,11 +18,11 @@
 // UMassVisualizationComponent
 //---------------------------------------------------------------
 
-namespace UE::MassRepresentation
+namespace UE::Mass::Representation
 {
 	int32 GCallUpdateInstances = 1;
 	FAutoConsoleVariableRef  CVarCallUpdateInstances(TEXT("Mass.CallUpdateInstances"), GCallUpdateInstances, TEXT("Toggle between UpdateInstances and BatchUpdateTransform."));
-}  // UE::MassRepresentation
+}  // UE::Mass::Representation
 
 void UMassVisualizationComponent::PostInitProperties()
 {
@@ -521,7 +521,7 @@ void UMassVisualizationComponent::EndVisualChanges()
 				const int32 NewNumInstances = SharedData.StaticMeshInstanceTransforms.Num();
 
 				// Update PerInstanceSMData transforms
-				if ((bool)UE::MassRepresentation::GCallUpdateInstances)
+				if ((bool)UE::Mass::Representation::GCallUpdateInstances)
 				{
 					InstancedStaticMeshComponent->UpdateInstances(SharedData.UpdateInstanceIds, SharedData.StaticMeshInstanceTransforms, SharedData.StaticMeshInstancePrevTransforms, NumCustomDataFloats, SharedData.StaticMeshInstanceCustomFloats);
 					if (UHierarchicalInstancedStaticMeshComponent* HISMComp = Cast<UHierarchicalInstancedStaticMeshComponent>(InstancedStaticMeshComponent))
diff --git a/Engine/Plugins/Runtime/MassGameplay/Source/MassRepresentation/Private/MassVisualizationLODProcessor.cpp b/Engine/Plugins/Runtime/MassGameplay/Source/MassRepresentation/Private/MassVisualizationLODProcessor.cpp
--- a/Engine/Plugins/Runtime/MassGameplay/Source/MassRepresentation/Private/MassVisualizationLODProcessor.cpp
+++ b/Engine/Plugins/Runtime/MassGameplay/Source/MassRepresentation/Private/MassVisualizationLODProcessor.cpp
@@ -3,11 +3,11 @@
 #include "MassVisualizationLODProcessor.h"
 #include "MassExecutionContext.h"
 
-namespace UE::MassRepresentation
+namespace UE::Mass::Representation
 {
 	int32 bDebugRepresentationLOD = 0;
 	FAutoConsoleVariableRef CVarDebugRepresentationLOD(TEXT("mass.debug.RepresentationLOD"), bDebugRepresentationLOD, TEXT("Debug representation LOD"), ECVF_Cheat);
-} // UE::MassRepresentation
+} // UE::Mass::Representation
 
 
 UMassVisualizationLODProcessor::UMassVisualizationLODProcessor()
@@ -113,7 +113,7 @@ void UMassVisualizationLODProcessor::Execute(FMassEntityManager& EntityManager,
 	}
 
 	// Optional debug display
-	if (UE::MassRepresentation::bDebugRepresentationLOD)
+	if (UE::Mass::Representation::bDebugRepresentationLOD)
 	{
 		TRACE_CPUPROFILER_EVENT_SCOPE(DebugDisplayLOD)
 		UWorld* World = EntityManager.GetWorld();
diff --git a/Engine/Plugins/Runtime/MassGameplay/Source/MassRepresentation/Public/MassRepresentationProcessor.h b/Engine/Plugins/Runtime/MassGameplay/Source/MassRepresentation/Public/MassRepresentationProcessor.h
--- a/Engine/Plugins/Runtime/MassGameplay/Source/MassRepresentation/Public/MassRepresentationProcessor.h
+++ b/Engine/Plugins/Runtime/MassGameplay/Source/MassRepresentation/Public/MassRepresentationProcessor.h
@@ -14,7 +14,7 @@ class UMassActorSubsystem;
 struct FMassActorFragment;
 
 
-namespace UE::MassRepresentation
+namespace UE::Mass::Representation
 {
 	extern MASSREPRESENTATION_API int32 bAllowKeepActorExtraFrame;
 }
diff --git a/Engine/Plugins/Runtime/MassGameplay/Source/MassRepresentation/Public/MassRepresentationUtils.h b/Engine/Plugins/Runtime/MassGameplay/Source/MassRepresentation/Public/MassRepresentationUtils.h
--- a/Engine/Plugins/Runtime/MassGameplay/Source/MassRepresentation/Public/MassRepresentationUtils.h
+++ b/Engine/Plugins/Runtime/MassGameplay/Source/MassRepresentation/Public/MassRepresentationUtils.h
@@ -10,7 +10,7 @@
 
 struct FMassCommandBuffer;
 
-namespace UE::MassRepresentation
+namespace UE::Mass::Representation
 {
 
 inline EMassVisibility GetVisibilityFromArchetype(const FMassExecutionContext& Context)
@@ -90,4 +90,4 @@ inline bool IsVisibilityTagSet(const FMassExecutionContext& Context, EMassVisibi
 
 void PushSwapTagsCommand(FMassCommandBuffer& CommandBuffer, const FMassEntityHandle Entity, const EMassVisibility PrevVisibility, const EMassVisibility NewVisibility);
 
-} // UE::MassRepresentation
\ No newline at end of file
+} // UE::Mass::Representation
\ No newline at end of file
