[Mass] Added UPROPERTY markup to various fragment struct members for text export / to show up in the AI gameplay debugger

diff --git a/Engine/Plugins/Runtime/MassGameplay/Source/MassLOD/Public/MassLODFragments.h b/Engine/Plugins/Runtime/MassGameplay/Source/MassLOD/Public/MassLODFragments.h
--- a/Engine/Plugins/Runtime/MassGameplay/Source/MassLOD/Public/MassLODFragments.h
+++ b/Engine/Plugins/Runtime/MassGameplay/Source/MassLOD/Public/MassLODFragments.h
@@ -39,9 +39,11 @@ struct MASSLOD_API FMassViewerInfoFragment : public FMassFragment
 	GENERATED_BODY()
 
 	// Closest viewer distance
+	UPROPERTY()
 	float ClosestViewerDistanceSq;
 
 	// Closest distance to frustum
+	UPROPERTY()
 	float ClosestDistanceToFrustum;
 };
 
@@ -84,9 +86,16 @@ struct MASSLOD_API FMassVariableTickChunkFragment : public FMassChunkFragment
 	}
 
 private:
+	UPROPERTY()
 	bool bShouldTickThisFrame = true;
-	EMassLOD::Type LOD = EMassLOD::Max;
+
+	UPROPERTY()
+	TEnumAsByte<EMassLOD::Type> LOD = EMassLOD::Max;
+
+	UPROPERTY()
 	float TimeUntilNextTick = 0.0f;
+
+	UPROPERTY()
 	int32 LastChunkSerialModificationNumber = INDEX_NONE;
 };
 
@@ -216,11 +225,14 @@ struct MASSLOD_API FMassVisualizationChunkFragment : public FMassChunkFragment
 protected:
 
 	/** Visibility of the current chunk, should never change */
+	UPROPERTY()
 	EMassVisibility Visibility = EMassVisibility::Max;
 
 	/** Not visible chunks, might contains entity that are newly visible and not yet moved. */
+	UPROPERTY()
 	bool bContainsNewlyVisibleEntity = true;
 
 	/** Not visible chunks delta time until next update */
+	UPROPERTY()
 	float DeltaTime = 0;
 };
diff --git a/Engine/Plugins/Runtime/MassGameplay/Source/MassLOD/Public/MassLODTypes.h b/Engine/Plugins/Runtime/MassGameplay/Source/MassLOD/Public/MassLODTypes.h
--- a/Engine/Plugins/Runtime/MassGameplay/Source/MassLOD/Public/MassLODTypes.h
+++ b/Engine/Plugins/Runtime/MassGameplay/Source/MassLOD/Public/MassLODTypes.h
@@ -46,6 +46,7 @@ namespace EMassLOD
 }
 
 
+UENUM()
 enum class EMassVisibility : uint8
 {
 	CanBeSeen, // Not too far and within camera frustum
diff --git a/Engine/Plugins/Runtime/MassGameplay/Source/MassRepresentation/Public/MassRepresentationFragments.h b/Engine/Plugins/Runtime/MassGameplay/Source/MassRepresentation/Public/MassRepresentationFragments.h
--- a/Engine/Plugins/Runtime/MassGameplay/Source/MassRepresentation/Public/MassRepresentationFragments.h
+++ b/Engine/Plugins/Runtime/MassGameplay/Source/MassRepresentation/Public/MassRepresentationFragments.h
@@ -27,14 +27,21 @@ struct MASSREPRESENTATION_API FMassRepresentationLODFragment : public FMassFragm
 	GENERATED_BODY()
 
 	/** LOD information */
+	UPROPERTY()
 	TEnumAsByte<EMassLOD::Type> LOD = EMassLOD::Max;
+
+	UPROPERTY()
 	TEnumAsByte<EMassLOD::Type> PrevLOD = EMassLOD::Max;
 
 	/** Visibility Info */
+	UPROPERTY()
 	EMassVisibility Visibility = EMassVisibility::Max;
+
+	UPROPERTY()
 	EMassVisibility PrevVisibility = EMassVisibility::Max;
 
 	/** Value scaling from 0 to 3, 0 highest LOD we support and 3 being completely off LOD */
+	UPROPERTY()
 	float LODSignificance = 0.0f;
 };
 
@@ -43,21 +50,29 @@ struct MASSREPRESENTATION_API FMassRepresentationFragment : public FMassFragment
 {
 	GENERATED_BODY()
 
+	UPROPERTY()
 	EMassRepresentationType CurrentRepresentation = EMassRepresentationType::None;
 
+	UPROPERTY()
 	EMassRepresentationType PrevRepresentation = EMassRepresentationType::None;
 
+	UPROPERTY()
 	int16 HighResTemplateActorIndex = INDEX_NONE;
 
+	UPROPERTY()
 	int16 LowResTemplateActorIndex = INDEX_NONE;
 
+	UPROPERTY()
 	int16 StaticMeshDescIndex = INDEX_NONE;
 
+	UPROPERTY()
 	FMassActorSpawnRequestHandle ActorSpawnRequestHandle;
 
+	UPROPERTY()
 	FTransform PrevTransform;
 
 	/** Value scaling from 0 to 3, 0 highest LOD we support and 3 being completely off LOD */
+	UPROPERTY()
 	float PrevLODSignificance = -1.0f;
 };
 
