Added MassLWI settings to support disabling the system even though the plugin is enabled

diff --git a/Engine/Plugins/Runtime/MassGameplay/Source/MassLWI/Private/MassLWISettings.cpp b/Engine/Plugins/Runtime/MassGameplay/Source/MassLWI/Private/MassLWISettings.cpp
new file mode 100644
--- /dev/null
+++ b/Engine/Plugins/Runtime/MassGameplay/Source/MassLWI/Private/MassLWISettings.cpp
@@ -0,0 +1,17 @@
+// Copyright Epic Games, Inc. All Rights Reserved.
+
+#include "MassLWISettings.h"
+#include "HAL/IConsoleManager.h"
+
+
+namespace UE::Mass::LWI::Private
+{
+	bool bMassLWIForceEnable = false;
+
+	FAutoConsoleVariableRef CVarMassLWIEnable(TEXT("mass.lwi.ForceEnable"), bMassLWIForceEnable, TEXT("Use this to temporarily enable MassLWI without affecting the project wide UMassLWISettings::bMassLWIEnabled setting"));
+} // UE::MassLOD::Debug
+
+bool UMassLWISettings::IsMassLWIEnabled() const
+{
+	return bMassLWIEnabled || UE::Mass::LWI::Private::bMassLWIForceEnable;
+}
diff --git a/Engine/Plugins/Runtime/MassGameplay/Source/MassLWI/Private/MassLWISubsystem.cpp b/Engine/Plugins/Runtime/MassGameplay/Source/MassLWI/Private/MassLWISubsystem.cpp
--- a/Engine/Plugins/Runtime/MassGameplay/Source/MassLWI/Private/MassLWISubsystem.cpp
+++ b/Engine/Plugins/Runtime/MassGameplay/Source/MassLWI/Private/MassLWISubsystem.cpp
@@ -8,6 +8,7 @@
 #include "MassEntitySubsystem.h"
 #include "MassLWIClientActorSpawnerSubsystem.h"
 #include "VisualLogger/VisualLogger.h"
+#include "MassLWISettings.h"
 
 
 #include UE_INLINE_GENERATED_CPP_BY_NAME(MassLWISubsystem)
@@ -49,7 +50,7 @@ void UMassLWISubsystem::Deinitialize()
 bool UMassLWISubsystem::ShouldCreateSubsystem(UObject* Outer) const
 {
 	UWorld* World = Outer ? Outer->GetWorld() : nullptr;
-	return World && World->IsGameWorld() && Super::ShouldCreateSubsystem(Outer);
+	return World && World->IsGameWorld() && Super::ShouldCreateSubsystem(Outer) && GET_MASSLWI_CONFIG_VALUE(IsMassLWIEnabled());
 }
 
 void UMassLWISubsystem::RegisterLWIManager(AMassLWIStaticMeshManager& Manager)
diff --git a/Engine/Plugins/Runtime/MassGameplay/Source/MassLWI/Public/MassLWISettings.h b/Engine/Plugins/Runtime/MassGameplay/Source/MassLWI/Public/MassLWISettings.h
new file mode 100644
--- /dev/null
+++ b/Engine/Plugins/Runtime/MassGameplay/Source/MassLWI/Public/MassLWISettings.h
@@ -0,0 +1,25 @@
+// Copyright Epic Games, Inc. All Rights Reserved.
+
+#pragma once
+
+#include "MassSettings.h"
+#include "MassLWISettings.generated.h"
+
+
+#define GET_MASSLWI_CONFIG_VALUE(a) (GetMutableDefault<UMassLWISettings>()->a)
+
+/**
+ * Implements the settings for MassLWI plugin
+ */
+UCLASS(config = Mass, defaultconfig, DisplayName = "Mass LWI")
+class MASSLWI_API UMassLWISettings : public UMassModuleSettings
+{
+	GENERATED_BODY()
+
+public:
+	bool IsMassLWIEnabled() const;
+
+protected:
+	UPROPERTY(EditDefaultsOnly, Category = Mass, config)
+	bool bMassLWIEnabled = true;
+};
