[Mass] Added mass.debug.Representation CVar similar to mass.debug.RepresentationLOD to debug draw representation states for entities.

Also added the ability to set both mass.debug.Representation and mass.debug.RepresentationLOD to 2 to VisLog states instead, or 3 for both VisLog and debug draw.

diff --git a/Engine/Plugins/Runtime/MassGameplay/Source/MassLOD/Private/MassSimulationLOD.cpp b/Engine/Plugins/Runtime/MassGameplay/Source/MassLOD/Private/MassSimulationLOD.cpp
--- a/Engine/Plugins/Runtime/MassGameplay/Source/MassLOD/Private/MassSimulationLOD.cpp
+++ b/Engine/Plugins/Runtime/MassGameplay/Source/MassLOD/Private/MassSimulationLOD.cpp
@@ -191,6 +191,7 @@ void UMassSimulationLODProcessor::Execute(FMassEntityManager& EntityManager, FMa
 		});
 	}
 
+#if WITH_MASSGAMEPLAY_DEBUG
 	// Optional debug display
 	if (UE::MassLOD::bDebugSimulationLOD)
 	{
@@ -203,4 +204,5 @@ void UMassSimulationLODProcessor::Execute(FMassEntityManager& EntityManager, FMa
 			LODSharedFragment.LODCalculator.DebugDisplayLOD(Context, SimulationLODList, LocationList, World);
 		});
 	}
+#endif // WITH_MASSGAMEPLAY_DEBUG
 }
diff --git a/Engine/Plugins/Runtime/MassGameplay/Source/MassLOD/Public/MassLODCalculator.h b/Engine/Plugins/Runtime/MassGameplay/Source/MassLOD/Public/MassLODCalculator.h
--- a/Engine/Plugins/Runtime/MassGameplay/Source/MassLOD/Public/MassLODCalculator.h
+++ b/Engine/Plugins/Runtime/MassGameplay/Source/MassLOD/Public/MassLODCalculator.h
@@ -5,7 +5,7 @@
 #include "MassLODLogic.h"
 #include "MassLODUtils.h"
 #include "DrawDebugHelpers.h"
-
+#include "VisualLogger/VisualLogger.h"
 /**
  * Helper struct to calculate LOD for each agent and maximize count per LOD
  *   Requires TViewerInfoFragment fragment collected by the TMassLODCollector.
@@ -152,6 +152,7 @@ public:
 	template <typename TLODFragment>
 	void ForceOffLOD(FMassExecutionContext& Context, TArrayView<TLODFragment> LODList);
 
+#if WITH_MASSGAMEPLAY_DEBUG
 	/**
 	 * Debug draw the current state of each agent as a color coded square
 	 * @Param Context of the chunk execution
@@ -162,6 +163,17 @@ public:
 	template <typename TLODFragment, typename TTransformFragment>
 	void DebugDisplayLOD(FMassExecutionContext& Context, TConstArrayView<TLODFragment> LODList, TConstArrayView<TTransformFragment> LocationList, UWorld* World);
 
+	/**
+	 * Add Visual Log entries for the current state of each agent as a color coded location
+	 * @Param Context of the chunk execution
+	 * @Param LODList is the fragment where calculation are stored
+	 * @Param LocationList is the fragment transforms of the entities
+	 * @Param World where the debug display should be drawn
+	 */
+	template <typename TLODFragment, typename TTransformFragment>
+	void VisLogLOD(FMassExecutionContext& Context, TConstArrayView<TLODFragment> LODList, TConstArrayView<TTransformFragment> LocationList, UObject* LogOwner);
+#endif // WITH_MASSGAMEPLAY_DEBUG
+
 	/**
 	 * Return the maximum distance at which the LOD will be turn off
 	 */
@@ -753,6 +765,7 @@ void TMassLODCalculator<FLODLogic>::ForceOffLOD(FMassExecutionContext& Context,
 	}
 }
 
+#if WITH_MASSGAMEPLAY_DEBUG
 template <typename FLODLogic>
 template <typename TLODFragment, typename TTransformFragment>
 void TMassLODCalculator<FLODLogic>::DebugDisplayLOD(FMassExecutionContext& Context, TConstArrayView<TLODFragment> LODList, TConstArrayView<TTransformFragment> LocationList, UWorld* World)
@@ -767,6 +780,23 @@ void TMassLODCalculator<FLODLogic>::DebugDisplayLOD(FMassExecutionContext& Conte
 	}
 }
 
+template <typename FLODLogic>
+template <typename TLODFragment, typename TTransformFragment>
+void TMassLODCalculator<FLODLogic>::VisLogLOD(FMassExecutionContext& Context, TConstArrayView<TLODFragment> LODList, TConstArrayView<TTransformFragment> LocationList, UObject* LogOwner)
+{
+#if ENABLE_VISUAL_LOG
+	const int32 NumEntities = Context.GetNumEntities();
+	for (int EntityIdx = 0; EntityIdx < NumEntities; EntityIdx++)
+	{
+		const TTransformFragment& EntityLocation = LocationList[EntityIdx];
+		const TLODFragment& EntityLOD = LODList[EntityIdx];
+		int32 LODIdx = (int32)EntityLOD.LOD;
+		UE_VLOG_LOCATION(LogOwner, LogMassLOD, Verbose, EntityLocation.GetTransform().GetLocation(), 20.0f, UE::MassLOD::LODColors[LODIdx], TEXT("%s %d"), *Context.GetEntity(EntityIdx).DebugGetDescription(), LODIdx);
+	}
+#endif
+}
+#endif // WITH_MASSGAMEPLAY_DEBUG
+
 template <typename FLODLogic>
 template<bool bCalculateVisibility>
 EMassLOD::Type TMassLODCalculator<FLODLogic>::ComputeLODFromSettings(const EMassLOD::Type PrevLOD, const float DistanceToViewerSq, const bool bIsVisible, bool* bIsInAVisibleRange, const FMassLODRuntimeData& Data) const
diff --git a/Engine/Plugins/Runtime/MassGameplay/Source/MassReplication/Private/MassReplicationProcessor.cpp b/Engine/Plugins/Runtime/MassGameplay/Source/MassReplication/Private/MassReplicationProcessor.cpp
--- a/Engine/Plugins/Runtime/MassGameplay/Source/MassReplication/Private/MassReplicationProcessor.cpp
+++ b/Engine/Plugins/Runtime/MassGameplay/Source/MassReplication/Private/MassReplicationProcessor.cpp
@@ -377,6 +377,7 @@ void UMassReplicationProcessor::Execute(FMassEntityManager& EntityManager, FMass
 					});
 				}
 
+#if WITH_MASSGAMEPLAY_DEBUG
 				// Optional debug display
 				if (UE::Mass::Replication::DebugClientReplicationLOD == ClientHandle.GetIndex())
 				{
@@ -390,6 +391,7 @@ void UMassReplicationProcessor::Execute(FMassEntityManager& EntityManager, FMass
 						});
 					});
 				}
+#endif // WITH_MASSGAMEPLAY_DEBUG
 
 				Context.ClearEntityCollection();
 			}
diff --git a/Engine/Plugins/Runtime/MassGameplay/Source/MassRepresentation/Private/MassRepresentationDebug.cpp b/Engine/Plugins/Runtime/MassGameplay/Source/MassRepresentation/Private/MassRepresentationDebug.cpp
new file mode 100644
--- /dev/null
+++ b/Engine/Plugins/Runtime/MassGameplay/Source/MassRepresentation/Private/MassRepresentationDebug.cpp
@@ -0,0 +1,64 @@
+// Copyright Epic Games, Inc. All Rights Reserved.
+
+#pragma once
+
+#include "MassRepresentationDebug.h"
+
+#if WITH_MASSGAMEPLAY_DEBUG
+
+#include "DrawDebugHelpers.h"
+#include "MassCommonFragments.h"
+#include "MassExecutionContext.h"
+#include "MassRepresentationFragments.h"
+#include "Math/Color.h"
+#include "VisualLogger/VisualLogger.h"
+
+namespace UE::Mass::Representation::Debug
+{
+	FColor RepresentationColors[] =
+	{
+		FColor::Purple, // HighResSpawnedActor,
+		FColor::Turquoise, // LowResSpawnedActor
+		FColor::Orange, // StaticMeshInstance
+		FColor::White, // None
+	};
+
+	int32 DebugRepresentation = 0;
+	FAutoConsoleVariableRef CVarDebugRepresentation(TEXT("mass.debug.Representation"), DebugRepresentation, TEXT("Debug current representation 0 = Off, 1 = Debug Draw, 2 = VisLog, 3 = Both"), ECVF_Cheat);
+
+	int32 DebugRepresentationLOD = 0;
+	FAutoConsoleVariableRef CVarDebugRepresentationLOD(TEXT("mass.debug.RepresentationLOD"), DebugRepresentationLOD, TEXT("Debug representation LOD 0 = Off, 1 = Debug Draw, 2 = VisLog, 3 = Both"), ECVF_Cheat);
+
+	void DebugDisplayRepresentation(FMassExecutionContext& Context, TConstArrayView<FMassRepresentationFragment> RepresentationList, TConstArrayView<FTransformFragment> LocationList, UWorld* World)
+	{
+#if UE_ENABLE_DEBUG_DRAWING
+		const int32 NumEntities = Context.GetNumEntities();
+		for (int EntityIdx = 0; EntityIdx < NumEntities; EntityIdx++)
+		{
+			const FTransformFragment& EntityLocation = LocationList[EntityIdx];
+			const FMassRepresentationFragment& Representation = RepresentationList[EntityIdx];
+			int32 CurrentRepresentationIdx = (int32)Representation.CurrentRepresentation;
+			DrawDebugSolidBox(World, EntityLocation.GetTransform().GetLocation() + FVector(0.0f, 0.0f, 150.0f), FVector(25.0f), UE::Mass::Representation::Debug::RepresentationColors[CurrentRepresentationIdx]);
+		}
+#endif // UE_ENABLE_DEBUG_DRAWING
+	}
+
+	void VisLogRepresentation(FMassExecutionContext& Context, TConstArrayView<FMassRepresentationFragment> RepresentationList, TConstArrayView<FTransformFragment> LocationList, UObject* LogOwner)
+	{
+#if ENABLE_VISUAL_LOG
+		const int32 NumEntities = Context.GetNumEntities();
+		for (int EntityIdx = 0; EntityIdx < NumEntities; EntityIdx++)
+		{
+			const FTransformFragment& EntityLocation = LocationList[EntityIdx];
+			const FMassRepresentationFragment& Representation = RepresentationList[EntityIdx];
+			int32 CurrentRepresentationIdx = (int32)Representation.CurrentRepresentation;
+			int32 PrevRepresentationIdx = (int32)Representation.PrevRepresentation;
+			// Add 20cm +Z offset to draw above mass.debug.RepresentationLOD so they can be viewed together
+			UE_CVLOG_LOCATION(Representation.CurrentRepresentation == Representation.PrevRepresentation, LogOwner, LogMassRepresentation, Verbose, EntityLocation.GetTransform().GetLocation() + FVector(0.0f, 0.0f, 20.0f), 20.0f, UE::Mass::Representation::Debug::RepresentationColors[CurrentRepresentationIdx], TEXT("%s %d"), *Context.GetEntity(EntityIdx).DebugGetDescription(), CurrentRepresentationIdx);
+			UE_CVLOG_LOCATION(Representation.CurrentRepresentation != Representation.PrevRepresentation, LogOwner, LogMassRepresentation, Verbose, EntityLocation.GetTransform().GetLocation() + FVector(0.0f, 0.0f, 20.0f), 20.0f, UE::Mass::Representation::Debug::RepresentationColors[CurrentRepresentationIdx], TEXT("%s %d -> %d"), *Context.GetEntity(EntityIdx).DebugGetDescription(), PrevRepresentationIdx, CurrentRepresentationIdx);
+		}
+#endif // ENABLE_VISUAL_LOG
+	}
+}
+
+#endif // WITH_MASSGAMEPLAY_DEBUG
diff --git a/Engine/Plugins/Runtime/MassGameplay/Source/MassRepresentation/Private/MassRepresentationProcessor.cpp b/Engine/Plugins/Runtime/MassGameplay/Source/MassRepresentation/Private/MassRepresentationProcessor.cpp
--- a/Engine/Plugins/Runtime/MassGameplay/Source/MassRepresentation/Private/MassRepresentationProcessor.cpp
+++ b/Engine/Plugins/Runtime/MassGameplay/Source/MassRepresentation/Private/MassRepresentationProcessor.cpp
@@ -1,6 +1,7 @@
 // Copyright Epic Games, Inc. All Rights Reserved.
 
 #include "MassRepresentationProcessor.h"
+#include "MassRepresentationDebug.h"
 #include "MassRepresentationSubsystem.h"
 #include "MassRepresentationUtils.h"
 #include "MassActorSubsystem.h"
@@ -217,6 +218,22 @@ void UMassRepresentationProcessor::UpdateRepresentation(FMassExecutionContext& C
 			DisableActorForISM(Actor);
 		}
 	}
+
+#if WITH_MASSGAMEPLAY_DEBUG
+	// Optional debug display
+	if (UE::Mass::Representation::Debug::DebugRepresentation == 1 || UE::Mass::Representation::Debug::DebugRepresentation >= 3)
+	{
+		TRACE_CPUPROFILER_EVENT_SCOPE(DebugDisplayRepresentation)
+		UWorld* World = CachedEntityManager.GetWorld();
+		UE::Mass::Representation::Debug::DebugDisplayRepresentation(Context, RepresentationList, TransformList, World);
+	}
+	// Optional vislog
+	if (UE::Mass::Representation::Debug::DebugRepresentation >= 2)
+	{
+		TRACE_CPUPROFILER_EVENT_SCOPE(VisLogRepresentation)
+		UE::Mass::Representation::Debug::VisLogRepresentation(Context, RepresentationList, TransformList, this);
+	}
+#endif
 }
 
 void UMassRepresentationProcessor::Execute(FMassEntityManager& InEntityManager, FMassExecutionContext& Context)
diff --git a/Engine/Plugins/Runtime/MassGameplay/Source/MassRepresentation/Private/MassVisualizationLODProcessor.cpp b/Engine/Plugins/Runtime/MassGameplay/Source/MassRepresentation/Private/MassVisualizationLODProcessor.cpp
--- a/Engine/Plugins/Runtime/MassGameplay/Source/MassRepresentation/Private/MassVisualizationLODProcessor.cpp
+++ b/Engine/Plugins/Runtime/MassGameplay/Source/MassRepresentation/Private/MassVisualizationLODProcessor.cpp
@@ -1,15 +1,9 @@
 // Copyright Epic Games, Inc. All Rights Reserved.
 
 #include "MassVisualizationLODProcessor.h"
+#include "MassRepresentationDebug.h"
 #include "MassExecutionContext.h"
 
-namespace UE::Mass::Representation
-{
-	int32 bDebugRepresentationLOD = 0;
-	FAutoConsoleVariableRef CVarDebugRepresentationLOD(TEXT("mass.debug.RepresentationLOD"), bDebugRepresentationLOD, TEXT("Debug representation LOD"), ECVF_Cheat);
-} // UE::Mass::Representation
-
-
 UMassVisualizationLODProcessor::UMassVisualizationLODProcessor()
 {
 	bAutoRegisterWithProcessingPhases = false;
@@ -112,8 +106,9 @@ void UMassVisualizationLODProcessor::Execute(FMassEntityManager& EntityManager,
 		// Far entities do not need to maximize count
 	}
 
+#if WITH_MASSGAMEPLAY_DEBUG
 	// Optional debug display
-	if (UE::Mass::Representation::bDebugRepresentationLOD)
+	if (UE::Mass::Representation::Debug::DebugRepresentationLOD == 1 || UE::Mass::Representation::Debug::DebugRepresentationLOD >= 3)
 	{
 		TRACE_CPUPROFILER_EVENT_SCOPE(DebugDisplayLOD)
 		UWorld* World = EntityManager.GetWorld();
@@ -125,4 +120,17 @@ void UMassVisualizationLODProcessor::Execute(FMassEntityManager& EntityManager,
 			LODSharedFragment.LODCalculator.DebugDisplayLOD(Context, RepresentationLODList, TransformList, World);
 		});
 	}
+	// Optional vislog
+	if (UE::Mass::Representation::Debug::DebugRepresentationLOD >= 2)
+	{
+		TRACE_CPUPROFILER_EVENT_SCOPE(VisLogLOD)
+		DebugEntityQuery.ForEachEntityChunk(EntityManager, Context, [this](FMassExecutionContext& Context)
+		{
+			FMassVisualizationLODSharedFragment& LODSharedFragment = Context.GetMutableSharedFragment<FMassVisualizationLODSharedFragment>();
+			TConstArrayView<FMassRepresentationLODFragment> RepresentationLODList = Context.GetFragmentView<FMassRepresentationLODFragment>();
+			TConstArrayView<FTransformFragment> TransformList = Context.GetFragmentView<FTransformFragment>();
+			LODSharedFragment.LODCalculator.VisLogLOD(Context, RepresentationLODList, TransformList, this);
+		});
+	}
+#endif // WITH_MASSGAMEPLAY_DEBUG
 }
diff --git a/Engine/Plugins/Runtime/MassGameplay/Source/MassRepresentation/Public/MassRepresentationDebug.h b/Engine/Plugins/Runtime/MassGameplay/Source/MassRepresentation/Public/MassRepresentationDebug.h
new file mode 100644
--- /dev/null
+++ b/Engine/Plugins/Runtime/MassGameplay/Source/MassRepresentation/Public/MassRepresentationDebug.h
@@ -0,0 +1,29 @@
+// Copyright Epic Games, Inc. All Rights Reserved.
+
+#pragma once
+
+#include "MassCommonTypes.h"
+#include "Containers/ArrayView.h"
+
+#if WITH_MASSGAMEPLAY_DEBUG
+
+class UWorld;
+struct FColor;
+struct FMassExecutionContext;
+struct FMassRepresentationFragment;
+struct FTransformFragment;
+
+namespace UE::Mass::Representation::Debug
+{
+	extern MASSREPRESENTATION_API FColor RepresentationColors[];
+
+	extern MASSREPRESENTATION_API int32 DebugRepresentation;
+
+	extern MASSREPRESENTATION_API int32 DebugRepresentationLOD;
+
+	void DebugDisplayRepresentation(FMassExecutionContext& Context, TConstArrayView<FMassRepresentationFragment> RepresentationList, TConstArrayView<FTransformFragment> LocationList, UWorld* World);
+
+	void VisLogRepresentation(FMassExecutionContext& Context, TConstArrayView<FMassRepresentationFragment> RepresentationList, TConstArrayView<FTransformFragment> LocationList, UObject* LogOwner);
+}
+
+#endif // WITH_MASSGAMEPLAY_DEBUG
