Added a limit to number of iterations in UMassActorSpawnerSubsystem::ProcessPendingSpawningRequest to avoid needlessly exhausting time limit when only "waiting for actor to spawn" requests remain.

diff --git a/Engine/Plugins/Runtime/MassGameplay/Source/MassActors/Private/MassActorSpawnerSubsystem.cpp b/Engine/Plugins/Runtime/MassGameplay/Source/MassActors/Private/MassActorSpawnerSubsystem.cpp
--- a/Engine/Plugins/Runtime/MassGameplay/Source/MassActors/Private/MassActorSpawnerSubsystem.cpp
+++ b/Engine/Plugins/Runtime/MassGameplay/Source/MassActors/Private/MassActorSpawnerSubsystem.cpp
@@ -254,7 +254,10 @@ void UMassActorSpawnerSubsystem::ProcessPendingSpawningRequest(const double MaxT
 
 	const double TimeSliceEnd = FPlatformTime::Seconds() + MaxTimeSlicePerTick;
 
-	while (FPlatformTime::Seconds() < TimeSliceEnd)
+	const int32 IterationsLimit = SpawnRequestHandleManager.CalcNumUsedHandles();
+	int32 IterationsCount = 0;
+
+	while (FPlatformTime::Seconds() < TimeSliceEnd && IterationsCount++ < IterationsLimit)
 	{
 		FMassActorSpawnRequestHandle SpawnRequestHandle = GetNextRequestToSpawn();
 		if (!SpawnRequestHandle.IsValid() ||
