[Mass] Made UMassActorSpawnerSubsystem::ReleaseActorToPool a virtual method to allow subclasses to customize actor pooling behavior.

diff --git a/Engine/Plugins/Runtime/MassGameplay/Source/MassActors/Private/MassActorSpawnerSubsystem.cpp b/Engine/Plugins/Runtime/MassGameplay/Source/MassActors/Private/MassActorSpawnerSubsystem.cpp
--- a/Engine/Plugins/Runtime/MassGameplay/Source/MassActors/Private/MassActorSpawnerSubsystem.cpp
+++ b/Engine/Plugins/Runtime/MassGameplay/Source/MassActors/Private/MassActorSpawnerSubsystem.cpp
@@ -82,7 +82,7 @@ void UMassActorSpawnerSubsystem::DestroyActor(AActor* Actor, bool bImmediate /*=
 
 bool UMassActorSpawnerSubsystem::ReleaseActorToPool(AActor* Actor)
 {
-	if (!UE::MassActors::bUseActorPooling || !bActorPoolingEnabled)
+	if (!IsActorPoolingEnabled())
 	{
 		return false;
 	}
@@ -182,7 +182,7 @@ ESpawnRequestStatus UMassActorSpawnerSubsystem::SpawnOrRetrieveFromPool(FConstSt
 {
 	const FMassActorSpawnRequest& SpawnRequest = SpawnRequestView.Get<const FMassActorSpawnRequest>();
 
-	if (UE::MassActors::bUseActorPooling != 0 && bActorPoolingEnabled)
+	if (IsActorPoolingEnabled())
 	{
 		auto* Pool = PooledActors.Find(SpawnRequest.Template);
 
@@ -433,6 +433,11 @@ void UMassActorSpawnerSubsystem::DisableActorPooling()
 
 }
 
+bool UMassActorSpawnerSubsystem::IsActorPoolingEnabled()
+{
+	return UE::MassActors::bUseActorPooling && bActorPoolingEnabled;
+}
+
 void UMassActorSpawnerSubsystem::ReleaseAllResources()
 {
 	if (UWorld* World = GetWorld())
diff --git a/Engine/Plugins/Runtime/MassGameplay/Source/MassActors/Public/MassActorSpawnerSubsystem.h b/Engine/Plugins/Runtime/MassGameplay/Source/MassActors/Public/MassActorSpawnerSubsystem.h
--- a/Engine/Plugins/Runtime/MassGameplay/Source/MassActors/Public/MassActorSpawnerSubsystem.h
+++ b/Engine/Plugins/Runtime/MassGameplay/Source/MassActors/Public/MassActorSpawnerSubsystem.h
@@ -180,6 +180,7 @@ public:
 	
 	void EnableActorPooling();
 	void DisableActorPooling();
+	bool IsActorPoolingEnabled();
 
 	void ReleaseAllResources();
 
@@ -213,7 +214,7 @@ protected:
 	/** Try releasing this actor to pool if possible 
 	 * @param Actor to release to the bool
 	 * @return true if the actor was actually released to the pool */
-	bool ReleaseActorToPool(AActor* Actor);
+	virtual bool ReleaseActorToPool(AActor* Actor);
 
 	/** Internal generic request actor spawn to make sure the request derives from FMassActorSpawnRequest 
 	 *  @param SpawnRequest the spawn request parameters, We are allowing any type of spawn request, let's store it internally as a FInstancedStruct. This parameter is the FStructView over provide user struct */
