Added a new tag, FMassVisualizationLODProcessorTag, and made it a requirement to have UMassVisualizationLODProcessor operate on a given archetype. The tag has been added to support temporary disabling of processing for individual entities (by removing and later re-adding the tag).

diff --git a/Engine/Plugins/Runtime/MassGameplay/Source/MassRepresentation/Private/MassVisualizationLODProcessor.cpp b/Engine/Plugins/Runtime/MassGameplay/Source/MassRepresentation/Private/MassVisualizationLODProcessor.cpp
--- a/Engine/Plugins/Runtime/MassGameplay/Source/MassRepresentation/Private/MassVisualizationLODProcessor.cpp
+++ b/Engine/Plugins/Runtime/MassGameplay/Source/MassRepresentation/Private/MassVisualizationLODProcessor.cpp
@@ -15,6 +15,7 @@ UMassVisualizationLODProcessor::UMassVisualizationLODProcessor()
 void UMassVisualizationLODProcessor::ConfigureQueries()
 {
 	FMassEntityQuery BaseQuery;
+	BaseQuery.AddTagRequirement<FMassVisualizationLODProcessorTag>(EMassFragmentPresence::All);
 	BaseQuery.AddRequirement<FMassViewerInfoFragment>(EMassFragmentAccess::ReadOnly);
 	BaseQuery.AddRequirement<FMassRepresentationLODFragment>(EMassFragmentAccess::ReadWrite);
 	BaseQuery.AddRequirement<FTransformFragment>(EMassFragmentAccess::ReadOnly);
diff --git a/Engine/Plugins/Runtime/MassGameplay/Source/MassRepresentation/Private/MassVisualizationTrait.cpp b/Engine/Plugins/Runtime/MassGameplay/Source/MassRepresentation/Private/MassVisualizationTrait.cpp
--- a/Engine/Plugins/Runtime/MassGameplay/Source/MassRepresentation/Private/MassVisualizationTrait.cpp
+++ b/Engine/Plugins/Runtime/MassGameplay/Source/MassRepresentation/Private/MassVisualizationTrait.cpp
@@ -10,6 +10,7 @@
 #include "MassLODFragments.h"
 #include "MassActorSubsystem.h"
 #include "MassEntityUtils.h"
+#include "MassVisualizationLODProcessor.h"
 
 
 UMassVisualizationTrait::UMassVisualizationTrait()
@@ -95,6 +96,8 @@ void UMassVisualizationTrait::BuildTemplate(FMassEntityTemplateBuildContext& Bui
 	BuildContext.AddFragment<FMassRepresentationLODFragment>();
 	BuildContext.AddTag<FMassVisibilityCulledByDistanceTag>();
 	BuildContext.AddChunkFragment<FMassVisualizationChunkFragment>();
+
+	BuildContext.AddTag<FMassVisualizationLODProcessorTag>();
 }
 
 
diff --git a/Engine/Plugins/Runtime/MassGameplay/Source/MassRepresentation/Public/MassVisualizationLODProcessor.h b/Engine/Plugins/Runtime/MassGameplay/Source/MassRepresentation/Public/MassVisualizationLODProcessor.h
--- a/Engine/Plugins/Runtime/MassGameplay/Source/MassRepresentation/Public/MassVisualizationLODProcessor.h
+++ b/Engine/Plugins/Runtime/MassGameplay/Source/MassRepresentation/Public/MassVisualizationLODProcessor.h
@@ -9,6 +9,17 @@
 
 #include "MassVisualizationLODProcessor.generated.h"
 
+
+/** 
+ * Tag required by Visualization LOD Processor to update LOD information. Removing the tag allows to support temporary 
+ * disabling of processing for individual entities.
+ */
+USTRUCT()
+struct MASSREPRESENTATION_API FMassVisualizationLODProcessorTag : public FMassTag
+{
+	GENERATED_BODY();
+};
+
 UCLASS()
 class MASSREPRESENTATION_API UMassVisualizationLODProcessor : public UMassProcessor
 {
