Added FMassEntityTemplateData::RemoveTag to support tag removal to extend functionality of "modified templates" introduced earlier this year.

diff --git a/Engine/Plugins/Runtime/MassGameplay/Source/MassSpawner/Public/MassEntityTemplate.h b/Engine/Plugins/Runtime/MassGameplay/Source/MassSpawner/Public/MassEntityTemplate.h
--- a/Engine/Plugins/Runtime/MassGameplay/Source/MassSpawner/Public/MassEntityTemplate.h
+++ b/Engine/Plugins/Runtime/MassGameplay/Source/MassSpawner/Public/MassEntityTemplate.h
@@ -182,6 +182,19 @@ struct MASSSPAWNER_API FMassEntityTemplateData
 		Composition.Tags.Add(TagType);
 	}
 
+	template<typename T>
+	void RemoveTag()
+	{
+		static_assert(TIsDerivedFrom<T, FMassTag>::IsDerived, "Given struct doesn't represent a valid mass tag type. Make sure to inherit from FMassTag or one of its child-types.");
+		Composition.Tags.Remove<T>();
+	}
+
+	void RemoveTag(const UScriptStruct& TagType)
+	{
+		checkf(TagType.IsChildOf(FMassTag::StaticStruct()), TEXT("Given struct doesn't represent a valid mass tag type. Make sure to inherit from FMassTag or one of its child-types."));
+		Composition.Tags.Remove(TagType);
+	}
+
 	const FMassTagBitSet& GetTags() const { return Composition.Tags; }
 	FMassTagBitSet& GetMutableTags() { return Composition.Tags; }
 
