Moved some code around to support reusing parts of UMassRepresentationProcessor's and UMassStationaryISMSwitcherProcessor's logic. Also added tags that gate archetypes being processed by these processors.

diff --git a/Engine/Plugins/Runtime/MassGameplay/Source/MassRepresentation/Private/MassRepresentationProcessor.cpp b/Engine/Plugins/Runtime/MassGameplay/Source/MassRepresentation/Private/MassRepresentationProcessor.cpp
--- a/Engine/Plugins/Runtime/MassGameplay/Source/MassRepresentation/Private/MassRepresentationProcessor.cpp
+++ b/Engine/Plugins/Runtime/MassGameplay/Source/MassRepresentation/Private/MassRepresentationProcessor.cpp
@@ -231,7 +231,7 @@ void UMassRepresentationProcessor::UpdateRepresentation(FMassExecutionContext& C
 	if (UE::Mass::Representation::Debug::DebugRepresentation >= 2)
 	{
 		TRACE_CPUPROFILER_EVENT_SCOPE(VisLogRepresentation)
-		UE::Mass::Representation::Debug::VisLogRepresentation(Context, RepresentationList, TransformList, this);
+		UE::Mass::Representation::Debug::VisLogRepresentation(Context, RepresentationList, TransformList, RepresentationSubsystem);
 	}
 #endif
 }
@@ -313,6 +313,7 @@ void UMassVisualizationProcessor::ConfigureQueries()
 {
 	Super::ConfigureQueries();
 	EntityQuery.AddChunkRequirement<FMassVisualizationChunkFragment>(EMassFragmentAccess::ReadWrite);
+	EntityQuery.AddTagRequirement<FMassVisualizationProcessorTag>(EMassFragmentPresence::All);
 }
 
 FMassVisualizationChunkFragment& UMassVisualizationProcessor::UpdateChunkVisibility(FMassExecutionContext& Context) const
diff --git a/Engine/Plugins/Runtime/MassGameplay/Source/MassRepresentation/Private/MassStationaryISMSwitcherProcessor.cpp b/Engine/Plugins/Runtime/MassGameplay/Source/MassRepresentation/Private/MassStationaryISMSwitcherProcessor.cpp
--- a/Engine/Plugins/Runtime/MassGameplay/Source/MassRepresentation/Private/MassStationaryISMSwitcherProcessor.cpp
+++ b/Engine/Plugins/Runtime/MassGameplay/Source/MassRepresentation/Private/MassStationaryISMSwitcherProcessor.cpp
@@ -23,116 +23,119 @@ void UMassStationaryISMSwitcherProcessor::ConfigureQueries()
 	EntityQuery.AddConstSharedRequirement<FMassRepresentationParameters>();
 	EntityQuery.AddSharedRequirement<FMassRepresentationSubsystemSharedFragment>(EMassFragmentAccess::ReadWrite);
 	EntityQuery.AddTagRequirement<FMassStaticRepresentationTag>(EMassFragmentPresence::All);
+	EntityQuery.AddTagRequirement<FMassStationaryISMSwitcherProcessorTag>(EMassFragmentPresence::All);
 	EntityQuery.AddSubsystemRequirement<UMassSignalSubsystem>(EMassFragmentAccess::ReadWrite);
 }
 
 void UMassStationaryISMSwitcherProcessor::Execute(FMassEntityManager& EntityManager, FMassExecutionContext& Context)
 {
-	EntityQuery.ForEachEntityChunk(EntityManager, Context, [this](FMassExecutionContext& Context)
-	{
-		UMassSignalSubsystem& SignalSubsystem = Context.GetMutableSubsystemChecked<UMassSignalSubsystem>();
+	EntityQuery.ForEachEntityChunk(EntityManager, Context, &UMassStationaryISMSwitcherProcessor::ProcessContext);
+}
+	
+void UMassStationaryISMSwitcherProcessor::ProcessContext(FMassExecutionContext& Context)
+{
+	UMassSignalSubsystem& SignalSubsystem = Context.GetMutableSubsystemChecked<UMassSignalSubsystem>();
 
-		UMassRepresentationSubsystem* RepresentationSubsystem = Context.GetMutableSharedFragment<FMassRepresentationSubsystemSharedFragment>().RepresentationSubsystem;
-		check(RepresentationSubsystem);
-		FMassInstancedStaticMeshInfoArrayView ISMInfosView = RepresentationSubsystem->GetMutableInstancedStaticMeshInfos();
+	UMassRepresentationSubsystem* RepresentationSubsystem = Context.GetMutableSharedFragment<FMassRepresentationSubsystemSharedFragment>().RepresentationSubsystem;
+	check(RepresentationSubsystem);
+	FMassInstancedStaticMeshInfoArrayView ISMInfosView = RepresentationSubsystem->GetMutableInstancedStaticMeshInfos();
 
-		const TConstArrayView<FMassRepresentationLODFragment> RepresentationLODList = Context.GetFragmentView<FMassRepresentationLODFragment>();
-		const TConstArrayView<FTransformFragment> TransformList = Context.GetFragmentView<FTransformFragment>();
-		const TArrayView<FMassRepresentationFragment> RepresentationList = Context.GetMutableFragmentView<FMassRepresentationFragment>();
+	const TConstArrayView<FMassRepresentationLODFragment> RepresentationLODList = Context.GetFragmentView<FMassRepresentationLODFragment>();
+	const TConstArrayView<FTransformFragment> TransformList = Context.GetFragmentView<FTransformFragment>();
+	const TArrayView<FMassRepresentationFragment> RepresentationList = Context.GetMutableFragmentView<FMassRepresentationFragment>();
 
-		const FMassRepresentationParameters& RepresentationParams = Context.GetConstSharedFragment<FMassRepresentationParameters>();
-		const bool bDoKeepActorExtraFrame = UE::Mass::Representation::bAllowKeepActorExtraFrame ? RepresentationParams.bKeepLowResActors : false;
+	const FMassRepresentationParameters& RepresentationParams = Context.GetConstSharedFragment<FMassRepresentationParameters>();
+	const bool bDoKeepActorExtraFrame = UE::Mass::Representation::bAllowKeepActorExtraFrame ? RepresentationParams.bKeepLowResActors : false;
 
-		const int32 NumEntities = Context.GetNumEntities();
-		for (int32 EntityIdx = 0; EntityIdx < NumEntities; EntityIdx++)
-		{
-			const FMassEntityHandle EntityHandle = Context.GetEntity(EntityIdx);
-			const FMassRepresentationLODFragment& RepresentationLOD = RepresentationLODList[EntityIdx];
-			const FTransformFragment& TransformFragment = TransformList[EntityIdx];
-			FMassRepresentationFragment& Representation = RepresentationList[EntityIdx];
+	const int32 NumEntities = Context.GetNumEntities();
+	for (int32 EntityIdx = 0; EntityIdx < NumEntities; EntityIdx++)
+	{
+		const FMassEntityHandle EntityHandle = Context.GetEntity(EntityIdx);
+		const FMassRepresentationLODFragment& RepresentationLOD = RepresentationLODList[EntityIdx];
+		const FTransformFragment& TransformFragment = TransformList[EntityIdx];
+		FMassRepresentationFragment& Representation = RepresentationList[EntityIdx];
 
-			if (Representation.PrevRepresentation != EMassRepresentationType::StaticMeshInstance
-				&& Representation.CurrentRepresentation != EMassRepresentationType::StaticMeshInstance)
-			{
-				// nothing to do here
-				continue;
-			}
+		if (Representation.PrevRepresentation != EMassRepresentationType::StaticMeshInstance
+			&& Representation.CurrentRepresentation != EMassRepresentationType::StaticMeshInstance)
+		{
+			// nothing to do here
+			continue;
+		}
 
-			FMassInstancedStaticMeshInfo& ISMInfo = ISMInfosView[Representation.StaticMeshDescIndex];
+		FMassInstancedStaticMeshInfo& ISMInfo = ISMInfosView[Representation.StaticMeshDescIndex];
 
-			if (Representation.PrevRepresentation == EMassRepresentationType::StaticMeshInstance
-				&& Representation.CurrentRepresentation != EMassRepresentationType::StaticMeshInstance)
-			{
-				const int32 EntityId = GetTypeHash(EntityHandle);
+		if (Representation.PrevRepresentation == EMassRepresentationType::StaticMeshInstance
+			&& Representation.CurrentRepresentation != EMassRepresentationType::StaticMeshInstance)
+		{
+			const int32 EntityId = GetTypeHash(EntityHandle);
 
-				// note that we're using the PrevLODSignificance here, and the reason for it is that the Prev value matches the 
-				// PrevRepresentation - thus we need to remove from the "previously" used LODSignificance range.
-				ISMInfo.RemoveInstance(EntityId, Representation.PrevLODSignificance);
+			// note that we're using the PrevLODSignificance here, and the reason for it is that the Prev value matches the 
+			// PrevRepresentation - thus we need to remove from the "previously" used LODSignificance range.
+			ISMInfo.RemoveInstance(EntityId, Representation.PrevLODSignificance);
 
-				// consume "prev" data
-				Representation.PrevRepresentation = Representation.CurrentRepresentation;
+			// consume "prev" data
+			Representation.PrevRepresentation = Representation.CurrentRepresentation;
 
-				if (Representation.PrevRepresentation != EMassRepresentationType::None)
-				{
-					SignalSubsystem.SignalEntity(UE::Mass::Signals::SwitchedToActor, EntityHandle);
-				}
-			}
-			else if (Representation.PrevRepresentation != EMassRepresentationType::StaticMeshInstance
-				&& Representation.CurrentRepresentation == EMassRepresentationType::StaticMeshInstance)
+			if (Representation.PrevRepresentation != EMassRepresentationType::None)
 			{
-				const int32 EntityId = GetTypeHash(EntityHandle);
-				const FTransform& Transform = TransformFragment.GetTransform();
-				const FTransform& PrevTransform = Representation.PrevTransform;
-				const float LODSignificance = RepresentationLOD.LODSignificance;
-				const float PrevLODSignificance = Representation.PrevLODSignificance;
+				SignalSubsystem.SignalEntity(UE::Mass::Signals::SwitchedToActor, EntityHandle);
+			}
+		}
+		else if (Representation.PrevRepresentation != EMassRepresentationType::StaticMeshInstance
+			&& Representation.CurrentRepresentation == EMassRepresentationType::StaticMeshInstance)
+		{
+			const int32 EntityId = GetTypeHash(EntityHandle);
+			const FTransform& Transform = TransformFragment.GetTransform();
+			const FTransform& PrevTransform = Representation.PrevTransform;
+			const float LODSignificance = RepresentationLOD.LODSignificance;
+			const float PrevLODSignificance = Representation.PrevLODSignificance;
 				
-				if (FMassLODSignificanceRange* NewRange = ISMInfo.GetLODSignificanceRange(RepresentationLOD.LODSignificance))
+			if (FMassLODSignificanceRange* NewRange = ISMInfo.GetLODSignificanceRange(RepresentationLOD.LODSignificance))
+			{
+				if (ISMInfo.ShouldUseTransformOffset())
 				{
-					if (ISMInfo.ShouldUseTransformOffset())
-					{
-						const FTransform& TransformOffset = ISMInfo.GetTransformOffset();
-						const FTransform SMTransform = TransformOffset * Transform;
-						NewRange->AddInstance(EntityId, SMTransform);
-					}
-					else
-					{
-						NewRange->AddInstance(EntityId, Transform);
-					}
+					const FTransform& TransformOffset = ISMInfo.GetTransformOffset();
+					const FTransform SMTransform = TransformOffset * Transform;
+					NewRange->AddInstance(EntityId, SMTransform);
 				}
-
-				// consume "prev" data
-				// @note crazy hacky, but we don't want to consume if bDoKeepActorExtraFrame is true. In that case 
-				// UMassRepresentationProcessor::UpdateRepresentation expects the "prev" state not to be consumed 
-				// a frame longer so that it can do the consuming (and call "disable actor").
-				if (bDoKeepActorExtraFrame == false)
+				else
 				{
-					Representation.PrevRepresentation = Representation.CurrentRepresentation;
+					NewRange->AddInstance(EntityId, Transform);
 				}
+			}
 
-				SignalSubsystem.SignalEntity(UE::Mass::Signals::SwitchedToISM, EntityHandle);
+			// consume "prev" data
+			// @note crazy hacky, but we don't want to consume if bDoKeepActorExtraFrame is true. In that case 
+			// UMassRepresentationProcessor::UpdateRepresentation expects the "prev" state not to be consumed 
+			// a frame longer so that it can do the consuming (and call "disable actor").
+			if (bDoKeepActorExtraFrame == false)
+			{
+				Representation.PrevRepresentation = Representation.CurrentRepresentation;
 			}
-			else if (ISMInfo.GetLODSignificanceRangesNum() > 1 && Representation.PrevLODSignificance != RepresentationLOD.LODSignificance)
+
+			SignalSubsystem.SignalEntity(UE::Mass::Signals::SwitchedToISM, EntityHandle);
+		}
+		else if (ISMInfo.GetLODSignificanceRangesNum() > 1 && Representation.PrevLODSignificance != RepresentationLOD.LODSignificance)
+		{
+			// we remain in ISM land, but LODSignificance changed and we have multiple LODSignificance ranges for this entity
+			FMassLODSignificanceRange* OldRange = ISMInfo.GetLODSignificanceRange(Representation.PrevLODSignificance);
+			FMassLODSignificanceRange* NewRange = ISMInfo.GetLODSignificanceRange(RepresentationLOD.LODSignificance);
+			if (OldRange != NewRange)
 			{
-				// we remain in ISM land, but LODSignificance changed and we have multiple LODSignificance ranges for this entity
-				FMassLODSignificanceRange* OldRange = ISMInfo.GetLODSignificanceRange(Representation.PrevLODSignificance);
-				FMassLODSignificanceRange* NewRange = ISMInfo.GetLODSignificanceRange(RepresentationLOD.LODSignificance);
-				if (OldRange != NewRange)
+				const int32 EntityId = GetTypeHash(EntityHandle);
+				if (OldRange)
 				{
-					const int32 EntityId = GetTypeHash(EntityHandle);
-					if (OldRange)
-					{
-						OldRange->RemoveInstance(EntityId);
-					}
-					if (NewRange)
-					{
-						const FTransform& Transform = TransformFragment.GetTransform();
-						NewRange->AddInstance(EntityId, Transform);
-					}
+					OldRange->RemoveInstance(EntityId);
+				}
+				if (NewRange)
+				{
+					const FTransform& Transform = TransformFragment.GetTransform();
+					NewRange->AddInstance(EntityId, Transform);
 				}
 			}
-
-			// consume "prev" data
-			Representation.PrevLODSignificance = RepresentationLOD.LODSignificance;
 		}
-	});
+
+		// consume "prev" data
+		Representation.PrevLODSignificance = RepresentationLOD.LODSignificance;
+	}
 }
diff --git a/Engine/Plugins/Runtime/MassGameplay/Source/MassRepresentation/Private/MassVisualizationTrait.cpp b/Engine/Plugins/Runtime/MassGameplay/Source/MassRepresentation/Private/MassVisualizationTrait.cpp
--- a/Engine/Plugins/Runtime/MassGameplay/Source/MassRepresentation/Private/MassVisualizationTrait.cpp
+++ b/Engine/Plugins/Runtime/MassGameplay/Source/MassRepresentation/Private/MassVisualizationTrait.cpp
@@ -11,6 +11,7 @@
 #include "MassActorSubsystem.h"
 #include "MassEntityUtils.h"
 #include "MassVisualizationLODProcessor.h"
+#include "MassRepresentationProcessor.h"
 
 
 UMassVisualizationTrait::UMassVisualizationTrait()
@@ -98,6 +99,7 @@ void UMassVisualizationTrait::BuildTemplate(FMassEntityTemplateBuildContext& Bui
 	BuildContext.AddChunkFragment<FMassVisualizationChunkFragment>();
 
 	BuildContext.AddTag<FMassVisualizationLODProcessorTag>();
+	BuildContext.AddTag<FMassVisualizationProcessorTag>();
 }
 
 
diff --git a/Engine/Plugins/Runtime/MassGameplay/Source/MassRepresentation/Public/MassRepresentationProcessor.h b/Engine/Plugins/Runtime/MassGameplay/Source/MassRepresentation/Public/MassRepresentationProcessor.h
--- a/Engine/Plugins/Runtime/MassGameplay/Source/MassRepresentation/Public/MassRepresentationProcessor.h
+++ b/Engine/Plugins/Runtime/MassGameplay/Source/MassRepresentation/Public/MassRepresentationProcessor.h
@@ -27,6 +27,12 @@ class MASSREPRESENTATION_API UMassRepresentationProcessor : public UMassProcesso
 public:
 	UMassRepresentationProcessor();
 
+	/*
+	 * Update representation type for each entity, must be called within a ForEachEntityChunk
+	 * @param Context of the execution from the entity sub system
+	 */
+	static void UpdateRepresentation(FMassExecutionContext& Context);
+
 protected:
 
 	/** Configure the owned FMassEntityQuery instances to express processor's requirements */
@@ -50,19 +56,24 @@ protected:
 	 * @param bCancelSpawningOnly tell to only cancel the existing spawning request and to not release the associated actor it any.
 	 * @return if the actor was release or the spawning was canceled.
 	 */
-	bool ReleaseActorOrCancelSpawning(UMassRepresentationSubsystem& RepresentationSubsystem, UMassActorSubsystem* MassActorSubsystem
+	static bool ReleaseActorOrCancelSpawning(UMassRepresentationSubsystem& RepresentationSubsystem, UMassActorSubsystem* MassActorSubsystem
 		, const FMassEntityHandle MassAgent, FMassActorFragment& ActorInfo, const int16 TemplateActorIndex, FMassActorSpawnRequestHandle& SpawnRequestHandle
 		, FMassCommandBuffer& CommandBuffer, const bool bCancelSpawningOnly = false);
 
-	/*
-	 * Update representation type for each entity, must be called within a ForEachEntityChunk
-	 * @param Context of the execution from the entity sub system
-	 */
-	void UpdateRepresentation(FMassExecutionContext& Context);
-
 	FMassEntityQuery EntityQuery;
 };
 
+	
+/**
+ * Tag required by Visualization Processor to process given archetype. Removing the tag allows to support temporary
+ * disabling of processing for individual entities of given archetype.
+ */
+USTRUCT()
+struct MASSREPRESENTATION_API FMassVisualizationProcessorTag : public FMassTag
+{
+	GENERATED_BODY();
+};
+
 UCLASS()
 class MASSREPRESENTATION_API UMassVisualizationProcessor : public UMassRepresentationProcessor
 {
diff --git a/Engine/Plugins/Runtime/MassGameplay/Source/MassRepresentation/Public/MassStationaryISMSwitcherProcessor.h b/Engine/Plugins/Runtime/MassGameplay/Source/MassRepresentation/Public/MassStationaryISMSwitcherProcessor.h
--- a/Engine/Plugins/Runtime/MassGameplay/Source/MassRepresentation/Public/MassStationaryISMSwitcherProcessor.h
+++ b/Engine/Plugins/Runtime/MassGameplay/Source/MassRepresentation/Public/MassStationaryISMSwitcherProcessor.h
@@ -5,12 +5,25 @@
 #include "MassProcessor.h"
 #include "MassStationaryISMSwitcherProcessor.generated.h"
 
+
+struct FMassExecutionContext;
+
 namespace UE::Mass::Signals
 {
 	const FName SwitchedToActor = FName(TEXT("SwitchedToActor"));
 	const FName SwitchedToISM = FName(TEXT("SwitchedToISM"));
 }
 
+/**
+ * Tag required by UMassStationaryISMSwitcherProcessor to process given archetype. Removing the tag allows to support temporary
+ * disabling of processing for individual entities of given archetype.
+ */
+USTRUCT()
+struct MASSREPRESENTATION_API FMassStationaryISMSwitcherProcessorTag : public FMassTag
+{
+	GENERATED_BODY();
+};
+
 /** 
  * This processor's sole responsibility is to process all entities tagged with FMassStaticRepresentationTag
  * and check if they've switched to or away from EMassRepresentationType::StaticMeshInstance; and acordingly add or remove 
@@ -24,6 +37,8 @@ class MASSREPRESENTATION_API UMassStationaryISMSwitcherProcessor : public UMassP
 public:
 	UMassStationaryISMSwitcherProcessor(const FObjectInitializer& ObjectInitializer = FObjectInitializer::Get());
 
+	static void ProcessContext(FMassExecutionContext& Context);
+
 protected:
 	// UMassProcessor overrides begin
 	virtual void ConfigureQueries() override;
