Removed commented-out code from UMassVisualizationComponent::HandleChangesWithExternalIDTracking and added additional check for catching duplicated addition of instances.

diff --git a/Engine/Plugins/Runtime/MassGameplay/Source/MassRepresentation/Private/MassVisualizationComponent.cpp b/Engine/Plugins/Runtime/MassGameplay/Source/MassRepresentation/Private/MassVisualizationComponent.cpp
--- a/Engine/Plugins/Runtime/MassGameplay/Source/MassRepresentation/Private/MassVisualizationComponent.cpp
+++ b/Engine/Plugins/Runtime/MassGameplay/Source/MassRepresentation/Private/MassVisualizationComponent.cpp
@@ -297,47 +297,8 @@ void UMassVisualizationComponent::HandleChangesWithExternalIDTracking(UInstanced
 			ISMComponent.PerInstanceIds[NewIndices[i]] = InstanceIds[i];
 		}
 
+		checkf(ISMComponent.InstanceIdToInstanceIndexMap.Num() == ISMComponent.PerInstanceIds.Num(), TEXT("Duplicates have been added to the ISMComponents. Things will go down hill from here."));
 		ensureMsgf(CustomFloatData.Num() == 0, TEXT("Custom floats not supported with this set up just yet."));
-#if 0 // CustomFloatData support below
-		if (CustomFloatData.Num() && ISMComponent.Mobility != EComponentMobility::Static)
-		{
-			checkf(ISMComponent.NumCustomDataFloats == InNumCustomDataFloats, TEXT("Adding instances with a Custrom Floats count inconsisntent with previously added instances"));
-
-			for (int32 i = 0; i < NewIndices.Num(); ++i)
-			{
-				const int32 InstanceIndex = NewIndices[i];
-				const int32 TargetCustomDataOffset = InstanceIndex * ISMComponent.NumCustomDataFloats;
-				const int32 SrcCustomDataOffset = i * ISMComponent.NumCustomDataFloats;
-				for (int32 FloatIndex = 0; FloatIndex < ISMComponent.NumCustomDataFloats; ++FloatIndex)
-				{
-					// we're making a change only if any of the input data differs of the currently stored ones
-					if (FMath::Abs(CustomFloatData[SrcCustomDataOffset + FloatIndex] - ISMComponent.PerInstanceSMCustomData[TargetCustomDataOffset + FloatIndex]) > EqualTolerance)
-					{
-						// Update the component's data in place.
-						FMemory::Memcpy(&ISMComponent.PerInstanceSMCustomData[TargetCustomDataOffset], &CustomFloatData[SrcCustomDataOffset], ISMComponent.NumCustomDataFloats * sizeof(float));
-
-						// Record in a command buffer for future use.
-						// Using AddInstance here rather than SetCustomData because AddInstancesInternal we used to create 
-						// instances doesn't add commands to InstanceUpdateCmdBuffer
-						ISMComponent.InstanceUpdateCmdBuffer.AddInstance(InstanceIds[i], InstanceTransforms[i].ToMatrixWithScale(), FMatrix()
-							, MakeArrayView((const float*)&CustomFloatData[SrcCustomDataOffset], ISMComponent.NumCustomDataFloats));
-
-						break;
-					}
-				}
-			}
-		}
-		else 
-		{
-			const FTransform& ComponentTransform = ISMComponent.GetComponentTransform();
-			// since AddInstancesInternal called above doesn't add any commands we need to add them here.
-			// The "there are custom floats" path is using a different, command flavor
-			for (const FTransform& InstanceTransform : InstanceTransforms)
-			{
-				ISMComponent.InstanceUpdateCmdBuffer.AddInstance(InstanceTransform.GetRelativeTransform(ComponentTransform).ToMatrixWithScale());
-			}
-		}
-#endif // 0
 	}
 
 	if (SharedData.GetRemoveInstanceIds().Num())
