[Mass] Added mass.debug.Representation.MaxSignificance and mass.debug.RepresentationLOD.MaxSignificance CVars to limit debug drawing for performance and readability with large agent counts.

diff --git a/Engine/Plugins/Runtime/MassGameplay/Source/MassLOD/Public/MassLODCalculator.h b/Engine/Plugins/Runtime/MassGameplay/Source/MassLOD/Public/MassLODCalculator.h
--- a/Engine/Plugins/Runtime/MassGameplay/Source/MassLOD/Public/MassLODCalculator.h
+++ b/Engine/Plugins/Runtime/MassGameplay/Source/MassLOD/Public/MassLODCalculator.h
@@ -163,6 +163,17 @@ public:
 	template <typename TLODFragment, typename TTransformFragment>
 	void DebugDisplayLOD(FMassExecutionContext& Context, TConstArrayView<TLODFragment> LODList, TConstArrayView<TTransformFragment> LocationList, UWorld* World);
 
+	/**
+	 * Debug draw the current state of each agent as a color coded square, within MaxLODSignificance range
+	 * @Param Context of the chunk execution
+	 * @Param LODList is the fragment where calculation are stored
+	 * @Param LocationList is the fragment transforms of the entities
+	 * @Param World where the debug display should be drawn
+	 * @Param MaxLODSignificance is the max allowed value of LODList[i].LODSignificance for an agent's state to debug draw
+	 */
+	template <typename TLODFragment, typename TTransformFragment>
+	void DebugDisplaySignificantLOD(FMassExecutionContext& Context, TConstArrayView<TLODFragment> LODList, TConstArrayView<TTransformFragment> LocationList, UWorld* World, float MaxLODSignificance);
+
 	/**
 	 * Add Visual Log entries for the current state of each agent as a color coded location
 	 * @Param Context of the chunk execution
@@ -172,6 +183,17 @@ public:
 	 */
 	template <typename TLODFragment, typename TTransformFragment>
 	void VisLogLOD(FMassExecutionContext& Context, TConstArrayView<TLODFragment> LODList, TConstArrayView<TTransformFragment> LocationList, UObject* LogOwner);
+
+	/**
+	 * Add Visual Log entries for the current state of each agent as a color coded location, within MaxLODSignificance range
+	 * @Param Context of the chunk execution
+	 * @Param LODList is the fragment where calculation are stored
+	 * @Param LocationList is the fragment transforms of the entities
+	 * @Param World where the debug display should be drawn
+	 * @Param MaxLODSignificance is the max allowed value of LODList[i].LODSignificance for an agent's state to vislog
+	 */
+	template <typename TLODFragment, typename TTransformFragment>
+	void VisLogSignificantLOD(FMassExecutionContext& Context, TConstArrayView<TLODFragment> LODList, TConstArrayView<TTransformFragment> LocationList, UObject* LogOwner, float MaxLODSignificance);
 #endif // WITH_MASSGAMEPLAY_DEBUG
 
 	/**
@@ -780,6 +802,23 @@ void TMassLODCalculator<FLODLogic>::DebugDisplayLOD(FMassExecutionContext& Conte
 	}
 }
 
+template <typename FLODLogic>
+template <typename TLODFragment, typename TTransformFragment>
+void TMassLODCalculator<FLODLogic>::DebugDisplaySignificantLOD(FMassExecutionContext& Context, TConstArrayView<TLODFragment> LODList, TConstArrayView<TTransformFragment> LocationList, UWorld* World, const float MaxLODSignificance)
+{
+	const int32 NumEntities = Context.GetNumEntities();
+	for (int EntityIdx = 0; EntityIdx < NumEntities; EntityIdx++)
+	{
+		const TLODFragment& EntityLOD = LODList[EntityIdx];
+		if (EntityLOD.LODSignificance <= MaxLODSignificance)
+		{
+			const TTransformFragment& EntityLocation = LocationList[EntityIdx];
+			int32 LODIdx = (int32)EntityLOD.LOD;
+			DrawDebugSolidBox(World, EntityLocation.GetTransform().GetLocation() + FVector(0.0f, 0.0f, 120.0f), FVector(25.0f), UE::MassLOD::LODColors[LODIdx]);
+		}
+	}
+}
+
 template <typename FLODLogic>
 template <typename TLODFragment, typename TTransformFragment>
 void TMassLODCalculator<FLODLogic>::VisLogLOD(FMassExecutionContext& Context, TConstArrayView<TLODFragment> LODList, TConstArrayView<TTransformFragment> LocationList, UObject* LogOwner)
@@ -795,6 +834,25 @@ void TMassLODCalculator<FLODLogic>::VisLogLOD(FMassExecutionContext& Context, TC
 	}
 #endif
 }
+
+template <typename FLODLogic>
+template <typename TLODFragment, typename TTransformFragment>
+void TMassLODCalculator<FLODLogic>::VisLogSignificantLOD(FMassExecutionContext& Context, TConstArrayView<TLODFragment> LODList, TConstArrayView<TTransformFragment> LocationList, UObject* LogOwner, const float MaxLODSignificance)
+{
+#if ENABLE_VISUAL_LOG
+	const int32 NumEntities = Context.GetNumEntities();
+	for (int EntityIdx = 0; EntityIdx < NumEntities; EntityIdx++)
+	{
+		const TLODFragment& EntityLOD = LODList[EntityIdx];
+		if (EntityLOD.LODSignificance <= MaxLODSignificance)
+		{
+			const TTransformFragment& EntityLocation = LocationList[EntityIdx];
+			int32 LODIdx = (int32)EntityLOD.LOD;
+			UE_VLOG_LOCATION(LogOwner, LogMassLOD, Verbose, EntityLocation.GetTransform().GetLocation(), 20.0f, UE::MassLOD::LODColors[LODIdx], TEXT("%s %d"), *Context.GetEntity(EntityIdx).DebugGetDescription(), LODIdx);
+		}
+	}
+#endif
+}
 #endif // WITH_MASSGAMEPLAY_DEBUG
 
 template <typename FLODLogic>
diff --git a/Engine/Plugins/Runtime/MassGameplay/Source/MassRepresentation/Private/MassRepresentationDebug.cpp b/Engine/Plugins/Runtime/MassGameplay/Source/MassRepresentation/Private/MassRepresentationDebug.cpp
--- a/Engine/Plugins/Runtime/MassGameplay/Source/MassRepresentation/Private/MassRepresentationDebug.cpp
+++ b/Engine/Plugins/Runtime/MassGameplay/Source/MassRepresentation/Private/MassRepresentationDebug.cpp
@@ -24,36 +24,50 @@ namespace UE::Mass::Representation::Debug
 	int32 DebugRepresentation = 0;
 	FAutoConsoleVariableRef CVarDebugRepresentation(TEXT("mass.debug.Representation"), DebugRepresentation, TEXT("Debug current representation 0 = Off, 1 = Debug Draw, 2 = VisLog, 3 = Both"), ECVF_Cheat);
 
+	float DebugRepresentationMaxSignificance = float(EMassLOD::Max);
+	FAutoConsoleVariableRef CVarDebugRepresentationMaxSignificance(TEXT("mass.debug.Representation.MaxSignificance"), DebugRepresentationMaxSignificance, TEXT("Max LODSignificance for entities to draw / vislog with mass.debug.Representation enabled"), ECVF_Cheat);
+
 	int32 DebugRepresentationLOD = 0;
 	FAutoConsoleVariableRef CVarDebugRepresentationLOD(TEXT("mass.debug.RepresentationLOD"), DebugRepresentationLOD, TEXT("Debug representation LOD 0 = Off, 1 = Debug Draw, 2 = VisLog, 3 = Both"), ECVF_Cheat);
 
-	void DebugDisplayRepresentation(FMassExecutionContext& Context, TConstArrayView<FMassRepresentationFragment> RepresentationList, TConstArrayView<FTransformFragment> LocationList, UWorld* World)
+	float DebugRepresentationLODMaxSignificance = float(EMassLOD::Max);
+	FAutoConsoleVariableRef CVarDebugRepresentationLODMaxSignificance(TEXT("mass.debug.RepresentationLOD.MaxSignificance"), DebugRepresentationLODMaxSignificance, TEXT("Max LODSignificance for entities to draw / vislog with mass.debug.RepresentationLOD enabled"), ECVF_Cheat);
+
+	void DebugDisplayRepresentation(FMassExecutionContext& Context, TConstArrayView<FMassRepresentationLODFragment> RepresentationLODList, TConstArrayView<FMassRepresentationFragment> RepresentationList, TConstArrayView<FTransformFragment> LocationList, UWorld* World)
 	{
 #if UE_ENABLE_DEBUG_DRAWING
 		const int32 NumEntities = Context.GetNumEntities();
 		for (int EntityIdx = 0; EntityIdx < NumEntities; EntityIdx++)
 		{
-			const FTransformFragment& EntityLocation = LocationList[EntityIdx];
-			const FMassRepresentationFragment& Representation = RepresentationList[EntityIdx];
-			int32 CurrentRepresentationIdx = (int32)Representation.CurrentRepresentation;
-			DrawDebugSolidBox(World, EntityLocation.GetTransform().GetLocation() + FVector(0.0f, 0.0f, 150.0f), FVector(25.0f), UE::Mass::Representation::Debug::RepresentationColors[CurrentRepresentationIdx]);
+			const FMassRepresentationLODFragment& RepresentationLOD = RepresentationLODList[EntityIdx];
+			if (RepresentationLOD.LODSignificance <= UE::Mass::Representation::Debug::DebugRepresentationMaxSignificance)
+			{
+				const FMassRepresentationFragment& Representation = RepresentationList[EntityIdx];
+				const FTransformFragment& EntityLocation = LocationList[EntityIdx];
+				int32 CurrentRepresentationIdx = (int32)Representation.CurrentRepresentation;
+				DrawDebugSolidBox(World, EntityLocation.GetTransform().GetLocation() + FVector(0.0f, 0.0f, 150.0f), FVector(25.0f), UE::Mass::Representation::Debug::RepresentationColors[CurrentRepresentationIdx]);
+			}
 		}
 #endif // UE_ENABLE_DEBUG_DRAWING
 	}
 
-	void VisLogRepresentation(FMassExecutionContext& Context, TConstArrayView<FMassRepresentationFragment> RepresentationList, TConstArrayView<FTransformFragment> LocationList, UObject* LogOwner)
+	void VisLogRepresentation(FMassExecutionContext& Context, TConstArrayView<FMassRepresentationLODFragment> RepresentationLODList, TConstArrayView<FMassRepresentationFragment> RepresentationList, TConstArrayView<FTransformFragment> LocationList, UObject* LogOwner)
 	{
 #if ENABLE_VISUAL_LOG
 		const int32 NumEntities = Context.GetNumEntities();
 		for (int EntityIdx = 0; EntityIdx < NumEntities; EntityIdx++)
 		{
-			const FTransformFragment& EntityLocation = LocationList[EntityIdx];
-			const FMassRepresentationFragment& Representation = RepresentationList[EntityIdx];
-			int32 CurrentRepresentationIdx = (int32)Representation.CurrentRepresentation;
-			int32 PrevRepresentationIdx = (int32)Representation.PrevRepresentation;
-			// Add 20cm +Z offset to draw above mass.debug.RepresentationLOD so they can be viewed together
-			UE_CVLOG_LOCATION(Representation.CurrentRepresentation == Representation.PrevRepresentation, LogOwner, LogMassRepresentation, Verbose, EntityLocation.GetTransform().GetLocation() + FVector(0.0f, 0.0f, 20.0f), 20.0f, UE::Mass::Representation::Debug::RepresentationColors[CurrentRepresentationIdx], TEXT("%s %d"), *Context.GetEntity(EntityIdx).DebugGetDescription(), CurrentRepresentationIdx);
-			UE_CVLOG_LOCATION(Representation.CurrentRepresentation != Representation.PrevRepresentation, LogOwner, LogMassRepresentation, Verbose, EntityLocation.GetTransform().GetLocation() + FVector(0.0f, 0.0f, 20.0f), 20.0f, UE::Mass::Representation::Debug::RepresentationColors[CurrentRepresentationIdx], TEXT("%s %d -> %d"), *Context.GetEntity(EntityIdx).DebugGetDescription(), PrevRepresentationIdx, CurrentRepresentationIdx);
+			const FMassRepresentationLODFragment& RepresentationLOD = RepresentationLODList[EntityIdx];
+			if (RepresentationLOD.LODSignificance <= UE::Mass::Representation::Debug::DebugRepresentationMaxSignificance)
+			{
+				const FTransformFragment& EntityLocation = LocationList[EntityIdx];
+				const FMassRepresentationFragment& Representation = RepresentationList[EntityIdx];
+				int32 CurrentRepresentationIdx = (int32)Representation.CurrentRepresentation;
+				int32 PrevRepresentationIdx = (int32)Representation.PrevRepresentation;
+				// Add 20cm +Z offset to draw above mass.debug.RepresentationLOD so they can be viewed together
+				UE_CVLOG_LOCATION(Representation.CurrentRepresentation == Representation.PrevRepresentation, LogOwner, LogMassRepresentation, Verbose, EntityLocation.GetTransform().GetLocation() + FVector(0.0f, 0.0f, 20.0f), 20.0f, UE::Mass::Representation::Debug::RepresentationColors[CurrentRepresentationIdx], TEXT("%s %d"), *Context.GetEntity(EntityIdx).DebugGetDescription(), CurrentRepresentationIdx);
+				UE_CVLOG_LOCATION(Representation.CurrentRepresentation != Representation.PrevRepresentation, LogOwner, LogMassRepresentation, Verbose, EntityLocation.GetTransform().GetLocation() + FVector(0.0f, 0.0f, 20.0f), 20.0f, UE::Mass::Representation::Debug::RepresentationColors[CurrentRepresentationIdx], TEXT("%s %d -> %d"), *Context.GetEntity(EntityIdx).DebugGetDescription(), PrevRepresentationIdx, CurrentRepresentationIdx);
+			}
 		}
 #endif // ENABLE_VISUAL_LOG
 	}
diff --git a/Engine/Plugins/Runtime/MassGameplay/Source/MassRepresentation/Private/MassRepresentationProcessor.cpp b/Engine/Plugins/Runtime/MassGameplay/Source/MassRepresentation/Private/MassRepresentationProcessor.cpp
--- a/Engine/Plugins/Runtime/MassGameplay/Source/MassRepresentation/Private/MassRepresentationProcessor.cpp
+++ b/Engine/Plugins/Runtime/MassGameplay/Source/MassRepresentation/Private/MassRepresentationProcessor.cpp
@@ -225,13 +225,13 @@ void UMassRepresentationProcessor::UpdateRepresentation(FMassExecutionContext& C
 	{
 		TRACE_CPUPROFILER_EVENT_SCOPE(DebugDisplayRepresentation)
 		UWorld* World = CachedEntityManager.GetWorld();
-		UE::Mass::Representation::Debug::DebugDisplayRepresentation(Context, RepresentationList, TransformList, World);
+		UE::Mass::Representation::Debug::DebugDisplayRepresentation(Context, RepresentationLODList, RepresentationList, TransformList, World);
 	}
 	// Optional vislog
 	if (UE::Mass::Representation::Debug::DebugRepresentation >= 2)
 	{
 		TRACE_CPUPROFILER_EVENT_SCOPE(VisLogRepresentation)
-		UE::Mass::Representation::Debug::VisLogRepresentation(Context, RepresentationList, TransformList, RepresentationSubsystem);
+		UE::Mass::Representation::Debug::VisLogRepresentation(Context, RepresentationLODList, RepresentationList, TransformList, RepresentationSubsystem);
 	}
 #endif
 }
diff --git a/Engine/Plugins/Runtime/MassGameplay/Source/MassRepresentation/Private/MassVisualizationLODProcessor.cpp b/Engine/Plugins/Runtime/MassGameplay/Source/MassRepresentation/Private/MassVisualizationLODProcessor.cpp
--- a/Engine/Plugins/Runtime/MassGameplay/Source/MassRepresentation/Private/MassVisualizationLODProcessor.cpp
+++ b/Engine/Plugins/Runtime/MassGameplay/Source/MassRepresentation/Private/MassVisualizationLODProcessor.cpp
@@ -118,7 +118,7 @@ void UMassVisualizationLODProcessor::Execute(FMassEntityManager& EntityManager,
 			FMassVisualizationLODSharedFragment& LODSharedFragment = Context.GetMutableSharedFragment<FMassVisualizationLODSharedFragment>();
 			TConstArrayView<FMassRepresentationLODFragment> RepresentationLODList = Context.GetFragmentView<FMassRepresentationLODFragment>();
 			TConstArrayView<FTransformFragment> TransformList = Context.GetFragmentView<FTransformFragment>();
-			LODSharedFragment.LODCalculator.DebugDisplayLOD(Context, RepresentationLODList, TransformList, World);
+			LODSharedFragment.LODCalculator.DebugDisplaySignificantLOD(Context, RepresentationLODList, TransformList, World, UE::Mass::Representation::Debug::DebugRepresentationLODMaxSignificance);
 		});
 	}
 	// Optional vislog
@@ -130,7 +130,7 @@ void UMassVisualizationLODProcessor::Execute(FMassEntityManager& EntityManager,
 			FMassVisualizationLODSharedFragment& LODSharedFragment = Context.GetMutableSharedFragment<FMassVisualizationLODSharedFragment>();
 			TConstArrayView<FMassRepresentationLODFragment> RepresentationLODList = Context.GetFragmentView<FMassRepresentationLODFragment>();
 			TConstArrayView<FTransformFragment> TransformList = Context.GetFragmentView<FTransformFragment>();
-			LODSharedFragment.LODCalculator.VisLogLOD(Context, RepresentationLODList, TransformList, this);
+			LODSharedFragment.LODCalculator.VisLogSignificantLOD(Context, RepresentationLODList, TransformList, this, UE::Mass::Representation::Debug::DebugRepresentationLODMaxSignificance);
 		});
 	}
 #endif // WITH_MASSGAMEPLAY_DEBUG
diff --git a/Engine/Plugins/Runtime/MassGameplay/Source/MassRepresentation/Public/MassRepresentationDebug.h b/Engine/Plugins/Runtime/MassGameplay/Source/MassRepresentation/Public/MassRepresentationDebug.h
--- a/Engine/Plugins/Runtime/MassGameplay/Source/MassRepresentation/Public/MassRepresentationDebug.h
+++ b/Engine/Plugins/Runtime/MassGameplay/Source/MassRepresentation/Public/MassRepresentationDebug.h
@@ -11,6 +11,7 @@ class UWorld;
 struct FColor;
 struct FMassExecutionContext;
 struct FMassRepresentationFragment;
+struct FMassRepresentationLODFragment;
 struct FTransformFragment;
 
 namespace UE::Mass::Representation::Debug
@@ -18,12 +19,14 @@ namespace UE::Mass::Representation::Debug
 	extern MASSREPRESENTATION_API FColor RepresentationColors[];
 
 	extern MASSREPRESENTATION_API int32 DebugRepresentation;
+	extern MASSREPRESENTATION_API float DebugRepresentationMaxSignificance;
 
 	extern MASSREPRESENTATION_API int32 DebugRepresentationLOD;
+	extern MASSREPRESENTATION_API float DebugRepresentationLODMaxSignificance;
 
-	void DebugDisplayRepresentation(FMassExecutionContext& Context, TConstArrayView<FMassRepresentationFragment> RepresentationList, TConstArrayView<FTransformFragment> LocationList, UWorld* World);
+	void DebugDisplayRepresentation(FMassExecutionContext& Context, TConstArrayView<FMassRepresentationLODFragment> RepresentationLODList, TConstArrayView<FMassRepresentationFragment> RepresentationList, TConstArrayView<FTransformFragment> LocationList, UWorld* World);
 
-	void VisLogRepresentation(FMassExecutionContext& Context, TConstArrayView<FMassRepresentationFragment> RepresentationList, TConstArrayView<FTransformFragment> LocationList, UObject* LogOwner);
+	void VisLogRepresentation(FMassExecutionContext& Context, TConstArrayView<FMassRepresentationLODFragment> RepresentationLODList, TConstArrayView<FMassRepresentationFragment> RepresentationList, TConstArrayView<FTransformFragment> LocationList, UObject* LogOwner);
 }
 
 #endif // WITH_MASSGAMEPLAY_DEBUG
