Massaged entity compaction code to get information on how many entities got actually moved around.
While in that code I renamed the console variable gating compaction to drop the `ai.` prefix.

diff --git a/Engine/Plugins/Runtime/MassEntity/Source/MassEntity/Private/MassArchetypeData.cpp b/Engine/Plugins/Runtime/MassEntity/Source/MassEntity/Private/MassArchetypeData.cpp
--- a/Engine/Plugins/Runtime/MassEntity/Source/MassEntity/Private/MassArchetypeData.cpp
+++ b/Engine/Plugins/Runtime/MassEntity/Source/MassEntity/Private/MassArchetypeData.cpp
@@ -503,8 +503,9 @@ void FMassArchetypeData::ExecutionFunctionForChunk(FMassExecutionContext RunCont
 	}
 }
 
-void FMassArchetypeData::CompactEntities(const double TimeAllowed)
+int32 FMassArchetypeData::CompactEntities(const double TimeAllowed)
 {
+	int32 TotalEntitiesMoved = 0;
 	const double TimeAllowedEnd = FPlatformTime::Seconds() + TimeAllowed;
 
 	TMap<uint32, TArray<FMassArchetypeChunk*>> SortedChunksBySharedValues;
@@ -581,8 +582,12 @@ void FMassArchetypeData::CompactEntities(const double TimeAllowed)
 			{
 				EntityMap.FindChecked(ToEntity->Index) = AbsoluteIndex + i;
 			}
+
+			TotalEntitiesMoved += NumberOfEntitiesToMove;
 		}
 	}
+
+	return TotalEntitiesMoved;
 }
 
 void FMassArchetypeData::GetRequirementsFragmentMapping(TConstArrayView<FMassFragmentRequirementDescription> Requirements, FMassFragmentIndicesMapping& OutFragmentIndices) const
diff --git a/Engine/Plugins/Runtime/MassEntity/Source/MassEntity/Private/MassArchetypeData.h b/Engine/Plugins/Runtime/MassEntity/Source/MassEntity/Private/MassArchetypeData.h
--- a/Engine/Plugins/Runtime/MassEntity/Source/MassEntity/Private/MassArchetypeData.h
+++ b/Engine/Plugins/Runtime/MassEntity/Source/MassEntity/Private/MassArchetypeData.h
@@ -252,8 +252,9 @@ public:
 
 	/**
 	 * Compacts entities to fill up chunks as much as possible
+	 * @return number of entities moved around
 	 */
-	void CompactEntities(const double TimeAllowed);
+	int32 CompactEntities(const double TimeAllowed);
 
 	/**
 	 * Moves the entity from this archetype to another, will only copy all matching fragment types
diff --git a/Engine/Plugins/Runtime/MassEntity/Source/MassEntity/Private/MassEntityManager.cpp b/Engine/Plugins/Runtime/MassEntity/Source/MassEntity/Private/MassEntityManager.cpp
--- a/Engine/Plugins/Runtime/MassEntity/Source/MassEntity/Private/MassEntityManager.cpp
+++ b/Engine/Plugins/Runtime/MassEntity/Source/MassEntity/Private/MassEntityManager.cpp
@@ -358,6 +358,7 @@ void FMassEntityManager::ForEachArchetypeFragmentType(const FMassArchetypeHandle
 
 void FMassEntityManager::DoEntityCompaction(const double TimeAllowed)
 {
+	int32 TotalEntitiesMoved = 0;
 	const double TimeAllowedEnd = FPlatformTime::Seconds() + TimeAllowed;
 
 	bool bReachedTimeLimit = false;
@@ -369,15 +370,17 @@ void FMassEntityManager::DoEntityCompaction(const double TimeAllowed)
 			bReachedTimeLimit = TimeAllowedLeft <= 0.0;
 			if (bReachedTimeLimit)
 			{
-				break;
+ 				break;
 			}
-			ArchetypePtr->CompactEntities(TimeAllowedLeft);
+			TotalEntitiesMoved += ArchetypePtr->CompactEntities(TimeAllowedLeft);
 		}
 		if (bReachedTimeLimit)
 		{
 			break;
 		}
 	}
+
+	UE_CVLOG(TotalEntitiesMoved, GetOwner(), LogMass, Verbose, TEXT("Entity Compaction: moved %d entities"), TotalEntitiesMoved);
 }
 
 FMassEntityHandle FMassEntityManager::CreateEntity(const FMassArchetypeHandle& ArchetypeHandle, const FMassArchetypeSharedFragmentValues& SharedFragmentValues)
diff --git a/Engine/Plugins/Runtime/MassGameplay/Source/MassSimulation/Private/MassSimulationSubsystem.cpp b/Engine/Plugins/Runtime/MassGameplay/Source/MassSimulation/Private/MassSimulationSubsystem.cpp
--- a/Engine/Plugins/Runtime/MassGameplay/Source/MassSimulation/Private/MassSimulationSubsystem.cpp
+++ b/Engine/Plugins/Runtime/MassGameplay/Source/MassSimulation/Private/MassSimulationSubsystem.cpp
@@ -16,7 +16,7 @@ DEFINE_LOG_CATEGORY(LogMassSim);
 namespace UE::MassSimulation
 {
 	int32 bDoEntityCompaction = 1;
-	FAutoConsoleVariableRef CVarEntityCompaction(TEXT("ai.mass.EntityCompaction"), bDoEntityCompaction, TEXT("Maximize the nubmer of entities per chunk"), ECVF_Cheat);
+	FAutoConsoleVariableRef CVarEntityCompaction(TEXT("mass.EntityCompaction"), bDoEntityCompaction, TEXT("Maximize the number of entities per chunk"), ECVF_Cheat);
 }
 
 //----------------------------------------------------------------------//
@@ -178,10 +178,10 @@ void UMassSimulationSubsystem::OnProcessingPhaseStarted(const float DeltaSeconds
 	{
 		case EMassProcessingPhase::PrePhysics:
 			{
-				TRACE_CPUPROFILER_EVENT_SCOPE(DoEntityCompation);
-				check(EntityManager);
-				if (UE::MassSimulation::bDoEntityCompaction)
+				if (UE::MassSimulation::bDoEntityCompaction && GET_MASSSIMULATION_CONFIG_VALUE(bEntityCompactionEnabled))
 				{
+					TRACE_CPUPROFILER_EVENT_SCOPE(DoEntityCompaction);
+					check(EntityManager);
 					EntityManager->DoEntityCompaction(GET_MASSSIMULATION_CONFIG_VALUE(DesiredEntityCompactionTimeSlicePerTick));
 				}
 			}
diff --git a/Engine/Plugins/Runtime/MassGameplay/Source/MassSimulation/Public/MassSimulationSettings.h b/Engine/Plugins/Runtime/MassGameplay/Source/MassSimulation/Public/MassSimulationSettings.h
--- a/Engine/Plugins/Runtime/MassGameplay/Source/MassSimulation/Public/MassSimulationSettings.h
+++ b/Engine/Plugins/Runtime/MassGameplay/Source/MassSimulation/Public/MassSimulationSettings.h
@@ -24,6 +24,15 @@ public:
 	UPROPERTY(EditDefaultsOnly, Category = "Runtime", config)
 	double DesiredActorDestructionTimeSlicePerTick = 0.0005;
 
+	/** 
+	 * Whether the Entity Compaction is allowed. When enabled all the archetypes will be processed on a regular basis
+	 * and the entities will be moved around to optimize archetype chunk memory utilization. Note that this might not be
+	 * desired if your set up results in individual chunks being often not compatible due to the number and value of chunk 
+	 * and shared fragments.
+	 */
+	UPROPERTY(EditDefaultsOnly, Category = "Runtime", config)
+	bool bEntityCompactionEnabled = true;
+
 	/** The desired budget in seconds allowed to do entity compaction per frame */
 	UPROPERTY(EditDefaultsOnly, Category = "Runtime", config)
 	double DesiredEntityCompactionTimeSlicePerTick = 0.005;
