Fixed System.Mass.EntityTemplate.Trivial test

diff --git a/Engine/Plugins/Runtime/MassGameplay/Source/MassGameplayTestSuite/Private/MassEntityTemplateTest.cpp b/Engine/Plugins/Runtime/MassGameplay/Source/MassGameplayTestSuite/Private/MassEntityTemplateTest.cpp
--- a/Engine/Plugins/Runtime/MassGameplay/Source/MassGameplayTestSuite/Private/MassEntityTemplateTest.cpp
+++ b/Engine/Plugins/Runtime/MassGameplay/Source/MassGameplayTestSuite/Private/MassEntityTemplateTest.cpp
@@ -108,10 +108,11 @@ struct FEntityTemplate_Trivial : FEntityTemplateBase
 
 		TemplateDataA.AddFragment<FTestFragment_Int>();
 		TemplateDataB.AddFragment<FTestFragment_Int>();
+		const FGuid NonEmptyTemplateGuid = FGuid::NewGuid();
 
 		const TSharedRef<FMassEntityTemplate>& FinalizedEmptyTemplate = TemplateRegistry.FindOrAddTemplate(FMassEntityTemplateIDFactory::Make(FGuid::NewGuid()), MoveTemp(EmptyTemplateData));
-		const TSharedRef<FMassEntityTemplate>& FinalizedTemplateA = TemplateRegistry.FindOrAddTemplate(FMassEntityTemplateIDFactory::Make(FGuid::NewGuid()), MoveTemp(TemplateDataA));
-		const TSharedRef<FMassEntityTemplate>& FinalizedTemplateB = TemplateRegistry.FindOrAddTemplate(FMassEntityTemplateIDFactory::Make(FGuid::NewGuid()), MoveTemp(TemplateDataA));
+		const TSharedRef<FMassEntityTemplate>& FinalizedTemplateA = TemplateRegistry.FindOrAddTemplate(FMassEntityTemplateIDFactory::Make(NonEmptyTemplateGuid), MoveTemp(TemplateDataA));
+		const TSharedRef<FMassEntityTemplate>& FinalizedTemplateB = TemplateRegistry.FindOrAddTemplate(FMassEntityTemplateIDFactory::Make(NonEmptyTemplateGuid), MoveTemp(TemplateDataB));
 
 		AITEST_NOT_EQUAL("Non-empty template data should result in a finalized template different from the empty one", FinalizedEmptyTemplate->GetTemplateID(), FinalizedTemplateA->GetTemplateID());
 		AITEST_EQUAL("Non-empty template data should result in the very same finalized template", FinalizedTemplateA, FinalizedTemplateB);
