Made adding duplicated InstanceIds to Mass-controller ISM components a less dangerous action. Also changed the order of operations - we remove instances first and only then add new ones.

diff --git a/Engine/Plugins/Runtime/MassGameplay/Source/MassRepresentation/Private/MassVisualizationComponent.cpp b/Engine/Plugins/Runtime/MassGameplay/Source/MassRepresentation/Private/MassVisualizationComponent.cpp
--- a/Engine/Plugins/Runtime/MassGameplay/Source/MassRepresentation/Private/MassVisualizationComponent.cpp
+++ b/Engine/Plugins/Runtime/MassGameplay/Source/MassRepresentation/Private/MassVisualizationComponent.cpp
@@ -253,7 +253,7 @@ void UMassVisualizationComponent::BeginVisualChanges()
 	}
 }
 
-void UMassVisualizationComponent::HandleChangesWithExternalIDTracking(UInstancedStaticMeshComponent& ISMComponent, const FMassISMCSharedData& SharedData)
+void UMassVisualizationComponent::HandleChangesWithExternalIDTracking(UInstancedStaticMeshComponent& ISMComponent, FMassISMCSharedData& SharedData)
 {
 	constexpr float EqualTolerance = 1e-6;
 
@@ -265,42 +265,8 @@ void UMassVisualizationComponent::HandleChangesWithExternalIDTracking(UInstanced
 	UHierarchicalInstancedStaticMeshComponent* HISMComp = Cast<UHierarchicalInstancedStaticMeshComponent>(&ISMComponent);
 	bool bAutoReset = HISMComp ? HISMComp->bAutoRebuildTreeOnInstanceChanges : false;
 
-	if (SharedData.GetUpdateInstanceIds().Num())
-	{
-		TConstArrayView<int32> InstanceIds = SharedData.GetUpdateInstanceIds();
-		const TArray<FTransform>& InstanceTransforms = SharedData.GetStaticMeshInstanceTransformsArray();
-		const int32 InNumCustomDataFloats = SharedData.GetStaticMeshInstanceCustomFloats().Num();
-		TConstArrayView<float> CustomFloatData = SharedData.GetStaticMeshInstanceCustomFloats();
-
-		const int32 StartingCount = ISMComponent.PerInstanceSMData.Num();
-		check(ISMComponent.InstanceIdToInstanceIndexMap.Num() == StartingCount);
-
-		// if these are the first entities we're adding we need to set NumCustomDataFloats so that the PerInstanceSMCustomData
-		// gets populated properly by the AddInstancesInternal call below
-		if (StartingCount == 0 && CustomFloatData.Num() && ISMComponent.Mobility != EComponentMobility::Static)
-		{
-			ISMComponent.NumCustomDataFloats = InNumCustomDataFloats;
-		}
-
-		check(InstanceIds.Num() == InstanceTransforms.Num());
-		TArray<int32> NewIndices = ISMComponent.AddInstances(InstanceTransforms, /*bShouldReturnIndices=*/true, /*bWorldSpace=*/true);
-		
-		check(InstanceIds.Num() == NewIndices.Num());
-		ISMComponent.PerInstanceIds.AddDefaulted(ISMComponent.PerInstanceSMData.Num() - ISMComponent.PerInstanceIds.Num());
-
-		for (int32 i = 0; i < InstanceIds.Num(); ++i)
-		{
-			checkfSlow(ISMComponent.InstanceIdToInstanceIndexMap.Find(InstanceIds[i]) == nullptr
-				, TEXT("This occuring signals trouble. None of the InstanceIds is expected to have already been added to this MassISM component instance."));
-
-			ISMComponent.InstanceIdToInstanceIndexMap.Add(InstanceIds[i], NewIndices[i]);
-			ISMComponent.PerInstanceIds[NewIndices[i]] = InstanceIds[i];
-		}
-
-		checkf(ISMComponent.InstanceIdToInstanceIndexMap.Num() == ISMComponent.PerInstanceIds.Num(), TEXT("Duplicates have been added to the ISMComponents. Things will go down hill from here."));
-		ensureMsgf(CustomFloatData.Num() == 0, TEXT("Custom floats not supported with this set up just yet."));
-	}
-
+	// removing instances first, since this operation is more resilient to duplicates. Plus we make an arbitrary decision 
+	// that it's better to have redundant things visible than not seeing required things
 	if (SharedData.GetRemoveInstanceIds().Num())
 	{
 		//RemoveInstanceWithIds(SharedData.GetRemoveInstanceIds());
@@ -438,6 +404,75 @@ void UMassVisualizationComponent::HandleChangesWithExternalIDTracking(UInstanced
 		}
 	}
 
+	if (SharedData.GetUpdateInstanceIds().Num())
+	{
+		TConstArrayView<int32> InstanceIds = SharedData.GetUpdateInstanceIds();
+		const TArray<FTransform>& InstanceTransforms = SharedData.GetStaticMeshInstanceTransformsArray();
+		int32 InNumCustomDataFloats = SharedData.GetStaticMeshInstanceCustomFloats().Num();
+		TConstArrayView<float> CustomFloatData = SharedData.GetStaticMeshInstanceCustomFloats();
+
+		const int32 StartingCount = ISMComponent.PerInstanceSMData.Num();
+		check(ISMComponent.InstanceIdToInstanceIndexMap.Num() == StartingCount);
+
+		{
+			// @todo want to track this for some time since it's quite possible this will cost us, without any actual 
+			// gain (duplicates should not happen if everything runs as expected).
+			TRACE_CPUPROFILER_EVENT_SCOPE_STR("MassVisualizationComponent_RemovingDuplicates");
+
+			// This part makes sure we're not trying to add the same InstanceId twice
+			// Note that we're going to ignore new data since updating is non-trivial. The user is expected to remove the old data first.
+			bool bDataRemoved = false;
+			if (ISMComponent.InstanceIdToInstanceIndexMap.Num())
+			{
+				for (int32 IDIndex = InstanceIds.Num() - 1; IDIndex >= 0; --IDIndex)
+				{
+					if (!ensure(!ISMComponent.InstanceIdToInstanceIndexMap.Find(InstanceIds[IDIndex])))
+					{
+						SharedData.RemoveUpdatedInstanceIdsAtSwap(IDIndex);
+						bDataRemoved = true;
+					}
+				}
+			}
+
+			if (bDataRemoved)
+			{
+				// if we indeed removed data we need to update the array views
+				InstanceIds = SharedData.GetUpdateInstanceIds();
+				InNumCustomDataFloats = SharedData.GetStaticMeshInstanceCustomFloats().Num();
+				CustomFloatData = SharedData.GetStaticMeshInstanceCustomFloats();
+			}
+		}
+
+		// at this point it's possible we removed all the data - just make sure there's anything to add
+		if (InstanceIds.Num())
+		{
+			// if these are the first entities we're adding we need to set NumCustomDataFloats so that the PerInstanceSMCustomData
+			// gets populated properly by the AddInstancesInternal call below
+			if (StartingCount == 0 && CustomFloatData.Num() && ISMComponent.Mobility != EComponentMobility::Static)
+			{
+				ISMComponent.NumCustomDataFloats = InNumCustomDataFloats;
+			}
+
+			check(InstanceIds.Num() == InstanceTransforms.Num());
+			TArray<int32> NewIndices = ISMComponent.AddInstances(InstanceTransforms, /*bShouldReturnIndices=*/true, /*bWorldSpace=*/true);
+
+			check(InstanceIds.Num() == NewIndices.Num());
+			ISMComponent.PerInstanceIds.AddDefaulted(ISMComponent.PerInstanceSMData.Num() - ISMComponent.PerInstanceIds.Num());
+
+			for (int32 i = 0; i < InstanceIds.Num(); ++i)
+			{
+				checkf(ISMComponent.InstanceIdToInstanceIndexMap.Find(InstanceIds[i]) == nullptr
+					, TEXT("This occuring signals trouble. None of the InstanceIds is expected to have already been added to this ISM component instance."));
+
+				ISMComponent.InstanceIdToInstanceIndexMap.Add(InstanceIds[i], NewIndices[i]);
+				ISMComponent.PerInstanceIds[NewIndices[i]] = InstanceIds[i];
+			}
+
+			checkf(ISMComponent.InstanceIdToInstanceIndexMap.Num() == ISMComponent.PerInstanceIds.Num(), TEXT("Duplicates have been added to the ISMComponents. Things will go down hill from here."));
+			ensureMsgf(CustomFloatData.Num() == 0, TEXT("Custom floats not supported with this set up just yet."));
+		}
+	}
+	
 	ISMComponent.MarkRenderStateDirty();
 
 	if (bNavigationRelevant && ISMComponent.GetInstanceCount() == 0)
diff --git a/Engine/Plugins/Runtime/MassGameplay/Source/MassRepresentation/Public/MassRepresentationTypes.h b/Engine/Plugins/Runtime/MassGameplay/Source/MassRepresentation/Public/MassRepresentationTypes.h
--- a/Engine/Plugins/Runtime/MassGameplay/Source/MassRepresentation/Public/MassRepresentationTypes.h
+++ b/Engine/Plugins/Runtime/MassGameplay/Source/MassRepresentation/Public/MassRepresentationTypes.h
@@ -171,6 +171,17 @@ struct MASSREPRESENTATION_API FMassISMCSharedData
 		WriteIterator = 0;
 	}
 
+	void RemoveUpdatedInstanceIdsAtSwap(const int32 InstanceIDIndex)
+	{
+		UpdateInstanceIds.RemoveAtSwap(InstanceIDIndex);
+		StaticMeshInstanceTransforms.RemoveAtSwap(InstanceIDIndex);
+		StaticMeshInstancePrevTransforms.RemoveAtSwap(InstanceIDIndex);
+		if (StaticMeshInstanceCustomFloats.Num())
+		{
+			StaticMeshInstanceCustomFloats.RemoveAtSwap(InstanceIDIndex);
+		}
+	}
+
 	bool HasUpdatesToApply() const { return UpdateInstanceIds.Num() || RemoveInstanceIds.Num(); }
 	TConstArrayView<int32> GetUpdateInstanceIds() const { return UpdateInstanceIds; }
 	TConstArrayView<FTransform> GetStaticMeshInstanceTransforms() const { return StaticMeshInstanceTransforms; }
diff --git a/Engine/Plugins/Runtime/MassGameplay/Source/MassRepresentation/Public/MassVisualizationComponent.h b/Engine/Plugins/Runtime/MassGameplay/Source/MassRepresentation/Public/MassVisualizationComponent.h
--- a/Engine/Plugins/Runtime/MassGameplay/Source/MassRepresentation/Public/MassVisualizationComponent.h
+++ b/Engine/Plugins/Runtime/MassGameplay/Source/MassRepresentation/Public/MassVisualizationComponent.h
@@ -61,7 +61,7 @@ protected:
 	 * Applies changes accumulated in SharedData while manually updating the Instance ID mapping. This approach is done in preparation 
 	 * to upcoming ISM changes to keep the mapping management more secure (by making mapping private and fully component-owned).
 	 */
-	void HandleChangesWithExternalIDTracking(UInstancedStaticMeshComponent& ISMComponent, const FMassISMCSharedData& SharedData);
+	void HandleChangesWithExternalIDTracking(UInstancedStaticMeshComponent& ISMComponent, FMassISMCSharedData& SharedData);
 
 	/** Recreate all the static mesh components from the InstancedStaticMeshInfos */
 	void ConstructStaticMeshComponents();
