Trivial change to MassVisualizationTrait that lets subclasses block default registration of StaticMeshInstanceDesc.

diff --git a/Engine/Plugins/Runtime/MassGameplay/Source/MassRepresentation/Private/MassVisualizationTrait.cpp b/Engine/Plugins/Runtime/MassGameplay/Source/MassRepresentation/Private/MassVisualizationTrait.cpp
--- a/Engine/Plugins/Runtime/MassGameplay/Source/MassRepresentation/Private/MassVisualizationTrait.cpp
+++ b/Engine/Plugins/Runtime/MassGameplay/Source/MassRepresentation/Private/MassVisualizationTrait.cpp
@@ -83,7 +83,10 @@ void UMassVisualizationTrait::BuildTemplate(FMassEntityTemplateBuildContext& Bui
 	BuildContext.AddConstSharedFragment(ParamsFragment);
 
 	FMassRepresentationFragment& RepresentationFragment = BuildContext.AddFragment_GetRef<FMassRepresentationFragment>();
-	RepresentationFragment.StaticMeshDescIndex = RepresentationSubsystem->FindOrAddStaticMeshDesc(StaticMeshInstanceDesc);
+	if (bRegisterStaticMeshDesc)
+	{
+		RepresentationFragment.StaticMeshDescIndex = RepresentationSubsystem->FindOrAddStaticMeshDesc(StaticMeshInstanceDesc);
+	}
 	RepresentationFragment.HighResTemplateActorIndex = HighResTemplateActor.Get() ? RepresentationSubsystem->FindOrAddTemplateActor(HighResTemplateActor.Get()) : INDEX_NONE;
 	RepresentationFragment.LowResTemplateActorIndex = LowResTemplateActor.Get() ? RepresentationSubsystem->FindOrAddTemplateActor(LowResTemplateActor.Get()) : INDEX_NONE;
 
diff --git a/Engine/Plugins/Runtime/MassGameplay/Source/MassRepresentation/Public/MassVisualizationTrait.h b/Engine/Plugins/Runtime/MassGameplay/Source/MassRepresentation/Public/MassVisualizationTrait.h
--- a/Engine/Plugins/Runtime/MassGameplay/Source/MassRepresentation/Public/MassVisualizationTrait.h
+++ b/Engine/Plugins/Runtime/MassGameplay/Source/MassRepresentation/Public/MassVisualizationTrait.h
@@ -57,4 +57,12 @@ public:
 	UPROPERTY(EditDefaultsOnly, Category = "Mass|Visual")
 	bool bCanModifyRepresentationSubsystemClass = true;
 #endif // WITH_EDITORONLY_DATA
+
+protected:
+	/** 
+	 * Controls whether StaticMeshInstanceDesc gets registered via FindOrAddStaticMeshDesc call. Setting it to `false` 
+	 * can be useful for subclasses to avoid needlessly creating visualization data in RepresentationSubsystem, 
+	 * data that will never be used.
+	 */
+	bool bRegisterStaticMeshDesc = true;
 };
