Mass visualization changes aimed at visualization data reuse in scenarios where an ISMComponent gets removed and then added again later (the same object path, but different pointers).

In addition I changed how we initialize newly added InstancedStaticMeshInfos entries - we no longer linearly go through the whole array to detect uninitialized entries, instead upon creation we collect indices to be processed when the time comes

diff --git a/Engine/Plugins/Runtime/MassGameplay/Source/MassRepresentation/Private/MassRepresentationTypes.cpp b/Engine/Plugins/Runtime/MassGameplay/Source/MassRepresentation/Private/MassRepresentationTypes.cpp
--- a/Engine/Plugins/Runtime/MassGameplay/Source/MassRepresentation/Private/MassRepresentationTypes.cpp
+++ b/Engine/Plugins/Runtime/MassGameplay/Source/MassRepresentation/Private/MassRepresentationTypes.cpp
@@ -75,6 +75,32 @@ default: \
 }
 } // UE::Mass::Representation
 
+//-----------------------------------------------------------------------------
+// FMassInstancedStaticMeshInfo
+//-----------------------------------------------------------------------------
+void FMassInstancedStaticMeshInfo::AddISMComponent(FMassISMCSharedData& SharedData)
+{
+	if (ensure(SharedData.GetISMComponent()))
+	{
+		const uint32 ISMComponentPathHash = GetTypeHash(SharedData.GetISMComponentChecked().GetPathName());
+		AddISMComponent(ISMComponentPathHash, SharedData);
+	}
+}
+
+void FMassInstancedStaticMeshInfo::ReplaceISMComponent(const uint32 PathHash, UInstancedStaticMeshComponent& ISMComponent)
+{
+	const int32 EntryIndex = PathHashes.Find(PathHash);
+	if (EntryIndex != INDEX_NONE)
+	{
+		// using GetValid rather than TObjectPtr.Get since the latter won't fail if the object in question is marked as Garbage
+		const UInstancedStaticMeshComponent* PrevInstance = GetValid(InstancedStaticMeshComponents[EntryIndex]);
+		ensure(PrevInstance == nullptr || PrevInstance == &ISMComponent);
+		if (PrevInstance != &ISMComponent)
+		{
+			InstancedStaticMeshComponents[EntryIndex] = &ISMComponent;
+		}
+	}
+}
 
 //-----------------------------------------------------------------------------
 // FMassStaticMeshInstanceVisualizationMeshDesc
diff --git a/Engine/Plugins/Runtime/MassGameplay/Source/MassRepresentation/Private/MassVisualizationComponent.cpp b/Engine/Plugins/Runtime/MassGameplay/Source/MassRepresentation/Private/MassVisualizationComponent.cpp
--- a/Engine/Plugins/Runtime/MassGameplay/Source/MassRepresentation/Private/MassVisualizationComponent.cpp
+++ b/Engine/Plugins/Runtime/MassGameplay/Source/MassRepresentation/Private/MassVisualizationComponent.cpp
@@ -55,7 +55,7 @@ int16 UMassVisualizationComponent::FindOrAddVisualDesc(const FStaticMeshInstance
 			VisualIndex = InstancedStaticMeshInfos.Emplace(Desc);
 			BuildLODSignificanceForInfo(InstancedStaticMeshInfos[VisualIndex]);
 
-			bNeedStaticMeshComponentConstruction = true;
+			InstancedSMComponentsRequiringConstructing.Add(VisualIndex);
 		}
 	}
 	checkf(VisualIndex < INT16_MAX, TEXT("%hs resulting VisualIndex is out of expected bounds"), __FUNCTION__);
@@ -76,13 +76,27 @@ int16 UMassVisualizationComponent::AddVisualDescWithISMComponent(const FStaticMe
 	
 	UE_MT_SCOPED_WRITE_ACCESS(InstancedStaticMeshInfosDetector);
 
-	const int32 VisualIndex = InstancedStaticMeshInfos.Emplace(Desc);
-	const uint32 MeshDescHash = GetTypeHash(MeshDesc);
-	const uint32 CombinedHash = PointerHash(&ISMComponent, MeshDescHash);
-	FMassISMCSharedData& NewData = ISMCSharedData.FindOrAdd(CombinedHash, FMassISMCSharedData(&ISMComponent, 1));
-	InstancedStaticMeshInfos[VisualIndex].AddISMComponent(NewData);
+	int32 VisualIndex = INDEX_NONE;
 
-	BuildLODSignificanceForInfo(InstancedStaticMeshInfos[VisualIndex], CombinedHash);
+	const uint32 ISMComponentPathHash = GetTypeHash(ISMComponent.GetPathName());
+	if (const int32* ExistingVisualIndex = ISMComponentMap.Find(ISMComponentPathHash))
+	{
+		VisualIndex = *ExistingVisualIndex;
+		// we need to update the ISM component - it's possible the previous instance got unloaded and that's why we use 
+		// path-hash rather than a pointer hash
+		InstancedStaticMeshInfos[VisualIndex].ReplaceISMComponent(ISMComponentPathHash, ISMComponent);
+	}
+	else
+	{
+		VisualIndex = InstancedStaticMeshInfos.Emplace(Desc);
+		const uint32 MeshDescHash = GetTypeHash(MeshDesc);
+		const uint32 CombinedHash = HashCombine(ISMComponentPathHash, MeshDescHash);
+		FMassISMCSharedData& NewData = ISMCSharedData.FindOrAdd(CombinedHash, FMassISMCSharedData(&ISMComponent, 1));
+		InstancedStaticMeshInfos[VisualIndex].AddISMComponent(ISMComponentPathHash, NewData);
+		BuildLODSignificanceForInfo(InstancedStaticMeshInfos[VisualIndex], CombinedHash);
+
+		ISMComponentMap.Add(ISMComponentPathHash, VisualIndex);
+	}
 
 	checkf(VisualIndex < INT16_MAX, TEXT("%hs resulting VisualIndex is out of expected bounds"), __FUNCTION__);
 	return (int16)VisualIndex;
@@ -94,8 +108,17 @@ void UMassVisualizationComponent::ConstructStaticMeshComponents()
 	check(ActorOwner);
 	
 	UE_MT_SCOPED_WRITE_ACCESS(InstancedStaticMeshInfosDetector);
-	for (FMassInstancedStaticMeshInfo& Info : InstancedStaticMeshInfos)
+	for (const int32 VisualIndex : InstancedSMComponentsRequiringConstructing)
 	{
+		if (!ensureMsgf(InstancedStaticMeshInfos.IsValidIndex(VisualIndex)
+			, TEXT("InstancedStaticMeshInfos (size: %d) is never expected to shrink, so VisualIndex (value: %d) being invalid indicates it was wrong from the start.")
+			, InstancedStaticMeshInfos.Num(), VisualIndex))
+		{
+			continue;
+		}
+
+		FMassInstancedStaticMeshInfo& Info = InstancedStaticMeshInfos[VisualIndex];
+
 		// Check if it is already created
 		if (!Info.InstancedStaticMeshComponents.IsEmpty())
 		{
@@ -146,6 +169,11 @@ void UMassVisualizationComponent::ConstructStaticMeshComponents()
 				}
 			}
 
+			check(ISMC);
+			const uint32 ISMComponentPathHash = GetTypeHash(ISMC->GetPathName());
+			ensureMsgf(ISMComponentMap.Find(ISMComponentPathHash) == nullptr, TEXT("We've just created the ISMC that's being used here, so this check failing indicates hash-clash."));
+			ISMComponentMap.Add(ISMComponentPathHash, VisualIndex); 
+
 			check(SharedData);
 			Info.AddISMComponent(*SharedData);
 		}
@@ -246,10 +274,10 @@ void UMassVisualizationComponent::BeginVisualChanges()
 	TRACE_CPUPROFILER_EVENT_SCOPE_STR("MassVisualizationComponent BeginVisualChanges")
 
 	// Conditionally construct static mesh components
-	if (bNeedStaticMeshComponentConstruction)
+	if (InstancedSMComponentsRequiringConstructing.Num())
 	{
 		ConstructStaticMeshComponents();
-		bNeedStaticMeshComponentConstruction = false;
+		InstancedSMComponentsRequiringConstructing.Reset();
 	}
 }
 
diff --git a/Engine/Plugins/Runtime/MassGameplay/Source/MassRepresentation/Public/MassRepresentationTypes.h b/Engine/Plugins/Runtime/MassGameplay/Source/MassRepresentation/Public/MassRepresentationTypes.h
--- a/Engine/Plugins/Runtime/MassGameplay/Source/MassRepresentation/Public/MassRepresentationTypes.h
+++ b/Engine/Plugins/Runtime/MassGameplay/Source/MassRepresentation/Public/MassRepresentationTypes.h
@@ -158,6 +158,7 @@ struct MASSREPRESENTATION_API FMassISMCSharedData
 	}
 
 	UInstancedStaticMeshComponent* GetISMComponent() { return ISMC; }
+	UInstancedStaticMeshComponent& GetISMComponentChecked() { check(ISMC); return *ISMC; }
 	int32 StoreReference() { return ++RefCount; }
 	int32 ReleaseReference() { ensure(RefCount >= 0); return --RefCount; }
 
@@ -390,15 +391,27 @@ public:
 		}
 	}
 
-	void AddISMComponent(FMassISMCSharedData& SharedData)
+	/** 
+	 * Stores the ISM component tied to SharedData
+	 * @param PathHash is used to quickly and uniquely identify the relevant ISM component. Note that we're using object's
+	 *	path rather than a pointer to make up for objects reloaded with streaming having the same path but different pointers. 
+	 */
+	void AddISMComponent(const uint32 PathHash, FMassISMCSharedData& SharedData)
 	{
-		if (ensure(SharedData.GetISMComponent()))
-		{
-			InstancedStaticMeshComponents.Add(SharedData.GetISMComponent());
-			SharedData.StoreReference();
-		}
+		ensureMsgf(PathHashes.Find(PathHash) == INDEX_NONE, TEXT("Given PathHash has already been added. Consider using ReplaceISMComponent instead"));
+		PathHashes.Add(PathHash);
+		InstancedStaticMeshComponents.Add(&SharedData.GetISMComponentChecked());
+		SharedData.StoreReference();
 	}
 
+	void AddISMComponent(FMassISMCSharedData& SharedData);
+
+	/** 
+	 * Looks for the previous instance of the given component (that has been potentially unloaded), as indicated 
+	 * by path hash, and replaces it with the given instance
+	 */
+	void ReplaceISMComponent(const uint32 PathHash, UInstancedStaticMeshComponent& ISMComponent);
+
 	int32 GetLODSignificanceRangesNum() const { return LODSignificanceRanges.Num(); }
 
 protected:
@@ -414,6 +427,9 @@ protected:
 	UPROPERTY(VisibleAnywhere, Category = "Mass/Debug")
 	TArray<TObjectPtr<UInstancedStaticMeshComponent>> InstancedStaticMeshComponents;
 
+	/** Hashes of object paths for ISMComponents stored in InstancedStaticMeshComponents*/
+	TArray<uint32> PathHashes;
+
 	UPROPERTY(VisibleAnywhere, Category = "Mass/Debug")
 	TArray<FMassLODSignificanceRange> LODSignificanceRanges;
 
diff --git a/Engine/Plugins/Runtime/MassGameplay/Source/MassRepresentation/Public/MassVisualizationComponent.h b/Engine/Plugins/Runtime/MassGameplay/Source/MassRepresentation/Public/MassVisualizationComponent.h
--- a/Engine/Plugins/Runtime/MassGameplay/Source/MassRepresentation/Public/MassVisualizationComponent.h
+++ b/Engine/Plugins/Runtime/MassGameplay/Source/MassRepresentation/Public/MassVisualizationComponent.h
@@ -81,8 +81,11 @@ protected:
 	TArray<FMassInstancedStaticMeshInfo> InstancedStaticMeshInfos;
 	UE_MT_DECLARE_RW_RECURSIVE_ACCESS_DETECTOR(InstancedStaticMeshInfosDetector);
 
+	/** Mapping from ISMComponent object path to corresponding VisualIndex */
+	TMap<uint32, int32> ISMComponentMap;
+
 	FMassISMCSharedDataMap ISMCSharedData;
 
-	/** Whether there is a need to create a StaticMeshComponent */
-	bool bNeedStaticMeshComponentConstruction = false;
+	/** Indicies to InstancedStaticMeshInfos that need their SMComponent constructed */
+	TArray<int32> InstancedSMComponentsRequiringConstructing;
 };
