Quick fix to an issue stemming from Mass Visualization's wrong assumption that all the ISMComponents used are under its control. If external components are used it's possible that the component instance will go away without the system knowing about it. The fix is a stop-gap until we massage the MassVisualization to remove faulty assumption.

diff --git a/Engine/Plugins/Runtime/MassGameplay/Source/MassRepresentation/Private/MassVisualizationComponent.cpp b/Engine/Plugins/Runtime/MassGameplay/Source/MassRepresentation/Private/MassVisualizationComponent.cpp
--- a/Engine/Plugins/Runtime/MassGameplay/Source/MassRepresentation/Private/MassVisualizationComponent.cpp
+++ b/Engine/Plugins/Runtime/MassGameplay/Source/MassRepresentation/Private/MassVisualizationComponent.cpp
@@ -527,7 +527,11 @@ void UMassVisualizationComponent::EndVisualChanges()
 	for (auto It = ISMCSharedData.CreateIterator(); It; ++It)
 	{
 		FMassISMCSharedData& SharedData = It.Value();
-		if (UInstancedStaticMeshComponent* InstancedStaticMeshComponent = SharedData.GetISMComponent())
+		UInstancedStaticMeshComponent* InstancedStaticMeshComponent = SharedData.GetISMComponent();
+		// @todo need to check validity this way since Mass used to rely on the assumption that all the ISM components used were
+		// under its control. That's no longer the case, but the system has not been updated to take that into consideration.
+		// This is a temporary fix. 
+		if (IsValid(InstancedStaticMeshComponent))
 		{
 			if (SharedData.RequiresExternalInstanceIDTracking())
 			{
