Optimized initialization of newly created FMassInstancedStaticMeshInfo

diff --git a/Engine/Plugins/Runtime/MassGameplay/Source/MassRepresentation/Private/MassVisualizationComponent.cpp b/Engine/Plugins/Runtime/MassGameplay/Source/MassRepresentation/Private/MassVisualizationComponent.cpp
--- a/Engine/Plugins/Runtime/MassGameplay/Source/MassRepresentation/Private/MassVisualizationComponent.cpp
+++ b/Engine/Plugins/Runtime/MassGameplay/Source/MassRepresentation/Private/MassVisualizationComponent.cpp
@@ -55,7 +55,7 @@ int16 UMassVisualizationComponent::FindOrAddVisualDesc(const FStaticMeshInstance
 			VisualIndex = InstancedStaticMeshInfos.Emplace(Desc);
 			BuildLODSignificanceForInfo(InstancedStaticMeshInfos[VisualIndex]);
 
-			bNeedStaticMeshComponentConstruction = true;
+			InstancedSMComponentsRequiringConstructing.Add(VisualIndex);
 		}
 	}
 	checkf(VisualIndex < INT16_MAX, TEXT("%hs resulting VisualIndex is out of expected bounds"), __FUNCTION__);
@@ -94,8 +94,17 @@ void UMassVisualizationComponent::ConstructStaticMeshComponents()
 	check(ActorOwner);
 	
 	UE_MT_SCOPED_WRITE_ACCESS(InstancedStaticMeshInfosDetector);
-	for (FMassInstancedStaticMeshInfo& Info : InstancedStaticMeshInfos)
+	for (const int32 VisualIndex : InstancedSMComponentsRequiringConstructing)
 	{
+		if (!ensureMsgf(InstancedStaticMeshInfos.IsValidIndex(VisualIndex)
+			, TEXT("InstancedStaticMeshInfos (size: %d) is never expected to shrink, so VisualIndex (value: %d) being invalid indicates it was wrong from the start.")
+			, InstancedStaticMeshInfos.Num(), VisualIndex))
+		{
+			continue;
+		}
+
+		FMassInstancedStaticMeshInfo& Info = InstancedStaticMeshInfos[VisualIndex];
+
 		// Check if it is already created
 		if (!Info.InstancedStaticMeshComponents.IsEmpty())
 		{
@@ -227,6 +236,7 @@ void UMassVisualizationComponent::ClearAllVisualInstances()
 	}
 
 	ISMCSharedData.Reset();
+	InstancedSMComponentsRequiringConstructing.Reset();
 }
 
 void UMassVisualizationComponent::DirtyVisuals()
@@ -246,10 +256,10 @@ void UMassVisualizationComponent::BeginVisualChanges()
 	TRACE_CPUPROFILER_EVENT_SCOPE_STR("MassVisualizationComponent BeginVisualChanges")
 
 	// Conditionally construct static mesh components
-	if (bNeedStaticMeshComponentConstruction)
+	if (InstancedSMComponentsRequiringConstructing.Num())
 	{
 		ConstructStaticMeshComponents();
-		bNeedStaticMeshComponentConstruction = false;
+		InstancedSMComponentsRequiringConstructing.Reset();
 	}
 }
 
diff --git a/Engine/Plugins/Runtime/MassGameplay/Source/MassRepresentation/Public/MassVisualizationComponent.h b/Engine/Plugins/Runtime/MassGameplay/Source/MassRepresentation/Public/MassVisualizationComponent.h
--- a/Engine/Plugins/Runtime/MassGameplay/Source/MassRepresentation/Public/MassVisualizationComponent.h
+++ b/Engine/Plugins/Runtime/MassGameplay/Source/MassRepresentation/Public/MassVisualizationComponent.h
@@ -83,6 +83,6 @@ protected:
 
 	FMassISMCSharedDataMap ISMCSharedData;
 
-	/** Whether there is a need to create a StaticMeshComponent */
-	bool bNeedStaticMeshComponentConstruction = false;
+	/** Indicies to InstancedStaticMeshInfos that need their SMComponent constructed */
+	TArray<int32> InstancedSMComponentsRequiringConstructing;
 };
