Added a way to remove ISM Components from MassVisualizationComponent

diff --git a/Engine/Plugins/Runtime/MassGameplay/Source/MassRepresentation/Private/MassRepresentationSubsystem.cpp b/Engine/Plugins/Runtime/MassGameplay/Source/MassRepresentation/Private/MassRepresentationSubsystem.cpp
--- a/Engine/Plugins/Runtime/MassGameplay/Source/MassRepresentation/Private/MassRepresentationSubsystem.cpp
+++ b/Engine/Plugins/Runtime/MassGameplay/Source/MassRepresentation/Private/MassRepresentationSubsystem.cpp
@@ -32,6 +32,12 @@ int16 UMassRepresentationSubsystem::AddVisualDescWithISMComponent(const FStaticM
 	return VisualizationComponent->AddVisualDescWithISMComponent(Desc, ISMComponent);
 }
 
+void UMassRepresentationSubsystem::RemoveISMComponent(UInstancedStaticMeshComponent& ISMComponent)
+{
+	check(VisualizationComponent);
+	return VisualizationComponent->RemoveISMComponent(ISMComponent);
+}
+
 FMassInstancedStaticMeshInfoArrayView UMassRepresentationSubsystem::GetMutableInstancedStaticMeshInfos()
 {
 	check(VisualizationComponent);
diff --git a/Engine/Plugins/Runtime/MassGameplay/Source/MassRepresentation/Private/MassRepresentationTypes.cpp b/Engine/Plugins/Runtime/MassGameplay/Source/MassRepresentation/Private/MassRepresentationTypes.cpp
--- a/Engine/Plugins/Runtime/MassGameplay/Source/MassRepresentation/Private/MassRepresentationTypes.cpp
+++ b/Engine/Plugins/Runtime/MassGameplay/Source/MassRepresentation/Private/MassRepresentationTypes.cpp
@@ -75,6 +75,15 @@ default: \
 }
 } // UE::Mass::Representation
 
+//-----------------------------------------------------------------------------
+// FMassInstancedStaticMeshInfo
+//-----------------------------------------------------------------------------
+void FMassInstancedStaticMeshInfo::Reset()
+{
+	Desc.Reset();
+	InstancedStaticMeshComponents.Reset();
+	LODSignificanceRanges.Reset();
+}
 
 //-----------------------------------------------------------------------------
 // FMassStaticMeshInstanceVisualizationMeshDesc
diff --git a/Engine/Plugins/Runtime/MassGameplay/Source/MassRepresentation/Private/MassVisualizationComponent.cpp b/Engine/Plugins/Runtime/MassGameplay/Source/MassRepresentation/Private/MassVisualizationComponent.cpp
--- a/Engine/Plugins/Runtime/MassGameplay/Source/MassRepresentation/Private/MassVisualizationComponent.cpp
+++ b/Engine/Plugins/Runtime/MassGameplay/Source/MassRepresentation/Private/MassVisualizationComponent.cpp
@@ -33,6 +33,22 @@ void UMassVisualizationComponent::PostInitProperties()
 	}
 }
 
+int32 UMassVisualizationComponent::AddInstancedStaticMeshInfo(const FStaticMeshInstanceVisualizationDesc& Desc)
+{
+	int32 Index = INDEX_NONE;
+	if (InstancedStaticMeshInfosFreeIndices.Num() > 0)
+	{
+		Index = InstancedStaticMeshInfosFreeIndices.Pop(/*bAllowShrinking=*/false);
+		new(&InstancedStaticMeshInfos[Index]) FMassInstancedStaticMeshInfo(Desc);
+	}
+	else
+	{
+		Index = InstancedStaticMeshInfos.Emplace(Desc);
+	}
+
+	return Index;
+}
+
 int16 UMassVisualizationComponent::FindOrAddVisualDesc(const FStaticMeshInstanceVisualizationDesc& Desc)
 {
 	UE_MT_SCOPED_WRITE_ACCESS(InstancedStaticMeshInfosDetector);
@@ -52,7 +68,7 @@ int16 UMassVisualizationComponent::FindOrAddVisualDesc(const FStaticMeshInstance
 
 		if (bValidDescription)
 		{
-			VisualIndex = InstancedStaticMeshInfos.Emplace(Desc);
+			VisualIndex = AddInstancedStaticMeshInfo(Desc);
 			BuildLODSignificanceForInfo(InstancedStaticMeshInfos[VisualIndex]);
 
 			InstancedSMComponentsRequiringConstructing.Add(VisualIndex);
@@ -76,18 +92,33 @@ int16 UMassVisualizationComponent::AddVisualDescWithISMComponent(const FStaticMe
 	
 	UE_MT_SCOPED_WRITE_ACCESS(InstancedStaticMeshInfosDetector);
 
-	const int32 VisualIndex = InstancedStaticMeshInfos.Emplace(Desc);
-	const uint32 MeshDescHash = GetTypeHash(MeshDesc);
-	const uint32 CombinedHash = PointerHash(&ISMComponent, MeshDescHash);
-	FMassISMCSharedData& NewData = ISMCSharedData.FindOrAdd(CombinedHash, FMassISMCSharedData(&ISMComponent, 1));
+	const int32 VisualIndex = AddInstancedStaticMeshInfo(Desc);
+
+	const uint32 ISMComponentPathHash = GetTypeHash(ISMComponent.GetPathName()); 
+	FMassISMCSharedData& NewData = ISMCSharedData.FindOrAdd(ISMComponentPathHash, FMassISMCSharedData(&ISMComponent, /*bInRequiresExternalInstanceIDTracking=*/true));
 	InstancedStaticMeshInfos[VisualIndex].AddISMComponent(NewData);
 
-	BuildLODSignificanceForInfo(InstancedStaticMeshInfos[VisualIndex], CombinedHash);
+	BuildLODSignificanceForInfo(InstancedStaticMeshInfos[VisualIndex], ISMComponentPathHash);
+	
+	ISMComponentMap.Add(ISMComponentPathHash, VisualIndex);
 
 	checkf(VisualIndex < INT16_MAX, TEXT("%hs resulting VisualIndex is out of expected bounds"), __FUNCTION__);
 	return (int16)VisualIndex;
 }
 
+void UMassVisualizationComponent::RemoveISMComponent(UInstancedStaticMeshComponent& ISMComponent)
+{
+	const uint32 ISMComponentPathHash = GetTypeHash(ISMComponent.GetPathName());
+	const int32* VisualIndexPtr = ISMComponentMap.Find(ISMComponentPathHash);
+	if (ensure(VisualIndexPtr))
+	{
+		ISMCSharedData.Remove(ISMComponentPathHash);
+
+		InstancedStaticMeshInfos[*VisualIndexPtr].Reset();
+		InstancedStaticMeshInfosFreeIndices.Add(*VisualIndexPtr);
+	}
+}
+
 void UMassVisualizationComponent::ConstructStaticMeshComponents()
 {
 	AActor* ActorOwner = GetOwner();
@@ -155,6 +186,11 @@ void UMassVisualizationComponent::ConstructStaticMeshComponents()
 				}
 			}
 
+			check(ISMC);
+			const uint32 ISMComponentPathHash = GetTypeHash(ISMC->GetPathName());
+			ensureMsgf(ISMComponentMap.Find(ISMComponentPathHash) == nullptr, TEXT("We've just created the ISMC that's being used here, so this check failing indicates hash-clash."));
+			ISMComponentMap.Add(ISMComponentPathHash, VisualIndex); 
+
 			check(SharedData);
 			Info.AddISMComponent(*SharedData);
 		}
diff --git a/Engine/Plugins/Runtime/MassGameplay/Source/MassRepresentation/Public/MassRepresentationSubsystem.h b/Engine/Plugins/Runtime/MassGameplay/Source/MassRepresentation/Public/MassRepresentationSubsystem.h
--- a/Engine/Plugins/Runtime/MassGameplay/Source/MassRepresentation/Public/MassRepresentationSubsystem.h
+++ b/Engine/Plugins/Runtime/MassGameplay/Source/MassRepresentation/Public/MassRepresentationSubsystem.h
@@ -43,7 +43,11 @@ public:
 	 */
 	int16 AddVisualDescWithISMComponent(const FStaticMeshInstanceVisualizationDesc& Desc, UInstancedStaticMeshComponent& ISMComponent);
 
-	/** @todo: need to add a release API at some point for static mesh types */
+	/**
+	 * Removes the visualization data associated with the given ISM component. Note that this is safe to do only when
+	 * there are no entities relying on this data. No entity data patching will take place.
+	 */
+	void RemoveISMComponent(UInstancedStaticMeshComponent& ISMComponent);
 
 	/** 
 	 * @return the array of all the static mesh instance component information
diff --git a/Engine/Plugins/Runtime/MassGameplay/Source/MassRepresentation/Public/MassRepresentationTypes.h b/Engine/Plugins/Runtime/MassGameplay/Source/MassRepresentation/Public/MassRepresentationTypes.h
--- a/Engine/Plugins/Runtime/MassGameplay/Source/MassRepresentation/Public/MassRepresentationTypes.h
+++ b/Engine/Plugins/Runtime/MassGameplay/Source/MassRepresentation/Public/MassRepresentationTypes.h
@@ -134,6 +134,11 @@ struct FStaticMeshInstanceVisualizationDesc : public FTableRowBase
 	{
 		return Meshes == Other.Meshes;
 	}
+
+	void Reset()
+	{
+		new(this)FStaticMeshInstanceVisualizationDesc();
+	}
 };
 
 class UInstancedStaticMeshComponent;
@@ -285,11 +290,14 @@ public:
 
 	FMassInstancedStaticMeshInfo() = default;
 
-	FMassInstancedStaticMeshInfo(const FStaticMeshInstanceVisualizationDesc& InDesc)
+	explicit FMassInstancedStaticMeshInfo(const FStaticMeshInstanceVisualizationDesc& InDesc)
 		: Desc(InDesc)
 	{
 	}
 
+	/** Clears out contents so that a given FMassInstancedStaticMeshInfo instance can be reused */
+	void Reset();
+
 	const FStaticMeshInstanceVisualizationDesc& GetDesc() const
 	{
 		return Desc;
diff --git a/Engine/Plugins/Runtime/MassGameplay/Source/MassRepresentation/Public/MassVisualizationComponent.h b/Engine/Plugins/Runtime/MassGameplay/Source/MassRepresentation/Public/MassVisualizationComponent.h
--- a/Engine/Plugins/Runtime/MassGameplay/Source/MassRepresentation/Public/MassVisualizationComponent.h
+++ b/Engine/Plugins/Runtime/MassGameplay/Source/MassRepresentation/Public/MassVisualizationComponent.h
@@ -35,7 +35,11 @@ public:
 	 */
 	int16 AddVisualDescWithISMComponent(const FStaticMeshInstanceVisualizationDesc& Desc, UInstancedStaticMeshComponent& ISMComponent);
 
-	/** @todo: need to add removal API at some point for visual types */
+	/** 
+	 * Removes the visualization data associated with the given ISM component. Note that this is safe to do only when
+	 * there are no entities relying on this data. No entity data patching will take place. 
+	 */
+	void RemoveISMComponent(UInstancedStaticMeshComponent& ISMComponent);
 
 	/** Get the array of all visual instance informations */
 	FMassInstancedStaticMeshInfoArrayView GetMutableVisualInfos()
@@ -76,11 +80,20 @@ protected:
 	 */
 	void BuildLODSignificanceForInfo(FMassInstancedStaticMeshInfo& Info, const uint32 ForcedStaticMeshRefKey = 0);
 
-	/** The information of all the instanced static meshes */
+	/** Either adds an element to InstancedStaticMeshInfos or reuses an existing entry based on InstancedStaticMeshInfosFreeIndices*/
+	int32 AddInstancedStaticMeshInfo(const FStaticMeshInstanceVisualizationDesc& Desc);
+
+	/** The information of all the instanced static meshes. Make sure to use AddInstancedStaticMeshInfo to add elements to it */
 	UPROPERTY(Transient)
 	TArray<FMassInstancedStaticMeshInfo> InstancedStaticMeshInfos;
 	UE_MT_DECLARE_RW_RECURSIVE_ACCESS_DETECTOR(InstancedStaticMeshInfosDetector);
 
+	/** Indices to InstancedStaticMeshInfos that have been released and can be reused */
+	TArray<int32> InstancedStaticMeshInfosFreeIndices;
+
+	/** Mapping from ISMComponent object path to corresponding VisualIndex */
+	TMap<uint32, int32> ISMComponentMap;
+
 	FMassISMCSharedDataMap ISMCSharedData;
 
 	/** Indicies to InstancedStaticMeshInfos that need their SMComponent constructed */
