Made it possible for UMassRepresentationProcessor to not test whether world collision is loaded while considering actor spawning

diff --git a/Engine/Plugins/Runtime/MassGameplay/Source/MassRepresentation/Private/MassRepresentationProcessor.cpp b/Engine/Plugins/Runtime/MassGameplay/Source/MassRepresentation/Private/MassRepresentationProcessor.cpp
--- a/Engine/Plugins/Runtime/MassGameplay/Source/MassRepresentation/Private/MassRepresentationProcessor.cpp
+++ b/Engine/Plugins/Runtime/MassGameplay/Source/MassRepresentation/Private/MassRepresentationProcessor.cpp
@@ -43,7 +43,7 @@ void UMassRepresentationProcessor::ConfigureQueries()
 	EntityQuery.AddSubsystemRequirement<UMassActorSubsystem>(EMassFragmentAccess::ReadWrite);
 }
 
-void UMassRepresentationProcessor::UpdateRepresentation(FMassExecutionContext& Context)
+void UMassRepresentationProcessor::UpdateRepresentation(FMassExecutionContext& Context, const FMassRepresentationUpdateParams& Params)
 {
 	UMassRepresentationSubsystem* RepresentationSubsystem = Context.GetMutableSharedFragment<FMassRepresentationSubsystemSharedFragment>().RepresentationSubsystem;
 	check(RepresentationSubsystem);
@@ -78,8 +78,9 @@ void UMassRepresentationProcessor::UpdateRepresentation(FMassExecutionContext& C
 		EMassRepresentationType WantedRepresentationType = RepresentationParams.LODRepresentation[FMath::Min((int32)RepresentationLOD.LOD, (int32)EMassLOD::Off)];
 
 		// Make sure we do not have actor spawned in areas not fully loaded
-		if ((WantedRepresentationType == EMassRepresentationType::HighResSpawnedActor || WantedRepresentationType == EMassRepresentationType::LowResSpawnedActor) &&
-			!RepresentationSubsystem->IsCollisionLoaded(RepresentationParams.WorldPartitionGridNameContainingCollision, TransformFragment.GetTransform()))
+		if (Params.bTestCollisionAvailibilityForActorVisualization
+			&& (WantedRepresentationType == EMassRepresentationType::HighResSpawnedActor || WantedRepresentationType == EMassRepresentationType::LowResSpawnedActor) 
+			&& !RepresentationSubsystem->IsCollisionLoaded(RepresentationParams.WorldPartitionGridNameContainingCollision, TransformFragment.GetTransform()))
 		{
 			WantedRepresentationType = RepresentationParams.CachedDefaultRepresentationType;
 		}
@@ -241,7 +242,7 @@ void UMassRepresentationProcessor::Execute(FMassEntityManager& InEntityManager,
 	// Update entities representation
 	EntityQuery.ForEachEntityChunk(InEntityManager, Context, [this](FMassExecutionContext& Context)
 	{
-		UpdateRepresentation(Context);
+		UpdateRepresentation(Context, UpdateParams);
 	});
 }
 
@@ -293,7 +294,7 @@ void UMassVisualizationProcessor::UpdateVisualization(FMassExecutionContext& Con
 		return;
 	}
 
-	UpdateRepresentation(Context);
+	UpdateRepresentation(Context, UpdateParams);
 
 	// Update entity visibility
 	const TArrayView<FMassRepresentationFragment> RepresentationList = Context.GetMutableFragmentView<FMassRepresentationFragment>();
@@ -374,11 +375,15 @@ void UMassVisualizationProcessor::UpdateEntityVisibility(const FMassEntityHandle
 
 void UMassVisualizationProcessor::Execute(FMassEntityManager& InEntityManager, FMassExecutionContext& Context)
 {
+	int32 TotalEntitiesProcessed = 0;
 	// Update entities visualization
-	EntityQuery.ForEachEntityChunk(InEntityManager, Context, [this](FMassExecutionContext& Context)
+	EntityQuery.ForEachEntityChunk(InEntityManager, Context, [this, &TotalEntitiesProcessed](FMassExecutionContext& Context)
 	{
+		TotalEntitiesProcessed += Context.GetNumEntities();
 		UpdateVisualization(Context);
 	});
+
+	UE_VLOG(this, LogMassRepresentation, Verbose, TEXT("UMassVisualizationProcessor::Execute processed %d entites"), TotalEntitiesProcessed);
 }
 
 //----------------------------------------------------------------------//
diff --git a/Engine/Plugins/Runtime/MassGameplay/Source/MassRepresentation/Public/MassRepresentationProcessor.h b/Engine/Plugins/Runtime/MassGameplay/Source/MassRepresentation/Public/MassRepresentationProcessor.h
--- a/Engine/Plugins/Runtime/MassGameplay/Source/MassRepresentation/Public/MassRepresentationProcessor.h
+++ b/Engine/Plugins/Runtime/MassGameplay/Source/MassRepresentation/Public/MassRepresentationProcessor.h
@@ -19,6 +19,20 @@ namespace UE::Mass::Representation
 	extern MASSREPRESENTATION_API int32 bAllowKeepActorExtraFrame;
 }
 
+USTRUCT()
+struct FMassRepresentationUpdateParams
+{
+	GENERATED_BODY()
+
+	/** 
+	 * Controls whether UpdateRepresentation run will ask the RepresentationSubsystem whether the relevant world 
+	 * collision has been already loaded while considering switching an entity to a actor-level representation. Note that
+	 * the test is unnecessary for static nor stationary entities.
+	 */
+	UPROPERTY(config, EditDefaultsOnly, Category = "Mass")
+	bool bTestCollisionAvailibilityForActorVisualization = true;
+};
+
 UCLASS()
 class MASSREPRESENTATION_API UMassRepresentationProcessor : public UMassProcessor
 {
@@ -31,7 +45,7 @@ public:
 	 * Update representation type for each entity, must be called within a ForEachEntityChunk
 	 * @param Context of the execution from the entity sub system
 	 */
-	static void UpdateRepresentation(FMassExecutionContext& Context);
+	static void UpdateRepresentation(FMassExecutionContext& Context, const FMassRepresentationUpdateParams& Params);
 
 protected:
 
@@ -61,6 +75,9 @@ protected:
 		, FMassCommandBuffer& CommandBuffer, const bool bCancelSpawningOnly = false);
 
 	FMassEntityQuery EntityQuery;
+
+	UPROPERTY(config, EditDefaultsOnly, Category = "Mass")
+	FMassRepresentationUpdateParams UpdateParams;
 };
 
 	
