Fixed the handling of removal of duplicated InstanceId data in UMassVisualizationComponent::HandleChangesWithExternalIDTracking

diff --git a/Engine/Plugins/Runtime/MassGameplay/Source/MassRepresentation/Private/MassVisualizationComponent.cpp b/Engine/Plugins/Runtime/MassGameplay/Source/MassRepresentation/Private/MassVisualizationComponent.cpp
--- a/Engine/Plugins/Runtime/MassGameplay/Source/MassRepresentation/Private/MassVisualizationComponent.cpp
+++ b/Engine/Plugins/Runtime/MassGameplay/Source/MassRepresentation/Private/MassVisualizationComponent.cpp
@@ -472,8 +472,10 @@ void UMassVisualizationComponent::HandleChangesWithExternalIDTracking(UInstanced
 			{
 				for (int32 IDIndex = InstanceIds.Num() - 1; IDIndex >= 0; --IDIndex)
 				{
-					if (!ensure(!ISMComponent.InstanceIdToInstanceIndexMap.Find(InstanceIds[IDIndex])))
+					if (ISMComponent.InstanceIdToInstanceIndexMap.Find(InstanceIds[IDIndex]))
 					{
+						// note that it's safe to remove elements, even though we iterate over and array view 
+						// because RemoveUpdatedInstanceIdsAtSwap makes sure no memory is reallocated (using bAllowShrinking = false)
 						SharedData.RemoveUpdatedInstanceIdsAtSwap(IDIndex);
 						bDataRemoved = true;
 					}
diff --git a/Engine/Plugins/Runtime/MassGameplay/Source/MassRepresentation/Public/MassRepresentationTypes.h b/Engine/Plugins/Runtime/MassGameplay/Source/MassRepresentation/Public/MassRepresentationTypes.h
--- a/Engine/Plugins/Runtime/MassGameplay/Source/MassRepresentation/Public/MassRepresentationTypes.h
+++ b/Engine/Plugins/Runtime/MassGameplay/Source/MassRepresentation/Public/MassRepresentationTypes.h
@@ -178,12 +178,12 @@ struct MASSREPRESENTATION_API FMassISMCSharedData
 
 	void RemoveUpdatedInstanceIdsAtSwap(const int32 InstanceIDIndex)
 	{
-		UpdateInstanceIds.RemoveAtSwap(InstanceIDIndex);
-		StaticMeshInstanceTransforms.RemoveAtSwap(InstanceIDIndex);
-		StaticMeshInstancePrevTransforms.RemoveAtSwap(InstanceIDIndex);
+		UpdateInstanceIds.RemoveAtSwap(InstanceIDIndex, 1, /*bAllowShrinking=*/false);
+		StaticMeshInstanceTransforms.RemoveAtSwap(InstanceIDIndex, 1, /*bAllowShrinking=*/false);
+		StaticMeshInstancePrevTransforms.RemoveAtSwap(InstanceIDIndex, 1, /*bAllowShrinking=*/false);
 		if (StaticMeshInstanceCustomFloats.Num())
 		{
-			StaticMeshInstanceCustomFloats.RemoveAtSwap(InstanceIDIndex);
+			StaticMeshInstanceCustomFloats.RemoveAtSwap(InstanceIDIndex, 1, /*bAllowShrinking=*/false);
 		}
 	}
 
