Fix to UMassActorSpawnerSubsystem::GetNextRequestToSpawn potentially repetedly returning the same handle

diff --git a/Engine/Plugins/Runtime/MassGameplay/Source/MassActors/Private/MassActorSpawnerSubsystem.cpp b/Engine/Plugins/Runtime/MassGameplay/Source/MassActors/Private/MassActorSpawnerSubsystem.cpp
--- a/Engine/Plugins/Runtime/MassGameplay/Source/MassActors/Private/MassActorSpawnerSubsystem.cpp
+++ b/Engine/Plugins/Runtime/MassGameplay/Source/MassActors/Private/MassActorSpawnerSubsystem.cpp
@@ -139,14 +139,29 @@ FMassActorSpawnRequestHandle UMassActorSpawnerSubsystem::RequestActorSpawnIntern
 	return SpawnRequestHandle;
 }
 
-FMassActorSpawnRequestHandle UMassActorSpawnerSubsystem::GetNextRequestToSpawn() const
+// @todo investigate whether storing requests in a sorted array would improve overall perf - if AllHandles was sorted we
+// wouldn't need to do any tests other than just checking if a thing is valid and not completed (i.e. pending or retry-pending).
+FMassActorSpawnRequestHandle UMassActorSpawnerSubsystem::GetNextRequestToSpawn(int32& InOutHandleIndex) const
 {
+	const TArray<FMassActorSpawnRequestHandle>& AllHandles = SpawnRequestHandleManager.GetHandles();
+	if (AllHandles.Num() == 0)
+	{
+		InOutHandleIndex = INDEX_NONE;
+		return FMassActorSpawnRequestHandle();
+	}
+
 	FMassActorSpawnRequestHandle BestSpawnRequestHandle;
 	float BestPriority = MAX_FLT;
 	bool bBestIsPending = false;
 	uint32 BestSerialNumber = MAX_uint32;
-	for (const FMassActorSpawnRequestHandle SpawnRequestHandle : SpawnRequestHandleManager.GetHandles())
+	int32 BestIndex = INDEX_NONE;
+		
+	int32 HandleIndex = InOutHandleIndex != INDEX_NONE ? ((InOutHandleIndex + 1) % AllHandles.Num()) : 0;
+	const int32 IterationsLimit = (InOutHandleIndex == INDEX_NONE) ? AllHandles.Num() : (AllHandles.Num() - 1);
+	
+	for (int32 IterationIndex = 0; IterationIndex < IterationsLimit; ++IterationIndex, HandleIndex = (HandleIndex + 1) % AllHandles.Num())
 	{
+		const FMassActorSpawnRequestHandle SpawnRequestHandle = AllHandles[HandleIndex];
 		if (!SpawnRequestHandle.IsValid())
 		{
 			continue;
@@ -162,6 +177,7 @@ FMassActorSpawnRequestHandle UMassActorSpawnerSubsystem::GetNextRequestToSpawn()
 				BestSerialNumber = SpawnRequest.SerialNumber;
 				BestPriority = SpawnRequest.Priority;
 				bBestIsPending = true;
+				BestIndex = HandleIndex;
 			}
 		}
 		else if (!bBestIsPending && SpawnRequest.SpawnStatus == ESpawnRequestStatus::RetryPending)
@@ -171,10 +187,12 @@ FMassActorSpawnRequestHandle UMassActorSpawnerSubsystem::GetNextRequestToSpawn()
 			{
 				BestSpawnRequestHandle = SpawnRequestHandle;
 				BestSerialNumber = SpawnRequest.SerialNumber;
+				BestIndex = HandleIndex;
 			}
 		}
 	}
 
+	InOutHandleIndex = BestIndex;
 	return BestSpawnRequestHandle;
 }
 
@@ -250,18 +268,22 @@ ESpawnRequestStatus UMassActorSpawnerSubsystem::SpawnOrRetrieveFromPool(FConstSt
 void UMassActorSpawnerSubsystem::ProcessPendingSpawningRequest(const double MaxTimeSlicePerTick)
 {
 	TRACE_CPUPROFILER_EVENT_SCOPE(UMassActorSpawnerSubsystem::ProcessPendingSpawningRequest);
+
 	SpawnRequestHandleManager.ShrinkHandles();
 
 	const double TimeSliceEnd = FPlatformTime::Seconds() + MaxTimeSlicePerTick;
 
 	const int32 IterationsLimit = SpawnRequestHandleManager.CalcNumUsedHandles();
 	int32 IterationsCount = 0;
+	int32 StartingHandleIndex = INDEX_NONE;
 
 	while (FPlatformTime::Seconds() < TimeSliceEnd && IterationsCount++ < IterationsLimit)
 	{
-		FMassActorSpawnRequestHandle SpawnRequestHandle = GetNextRequestToSpawn();
-		if (!SpawnRequestHandle.IsValid() ||
-			!ensureMsgf(SpawnRequestHandleManager.IsValidHandle(SpawnRequestHandle), TEXT("GetNextRequestToSpawn returned an invalid handle, expecting an empty one or a valid one.")))
+		const FMassActorSpawnRequestHandle SpawnRequestHandle = GetNextRequestToSpawn(StartingHandleIndex);
+
+		// getting an invalid handle is fine - it indicates no more handles are there to be considered. 
+		if (!SpawnRequestHandle.IsValid() 
+			|| !SpawnRequestHandleManager.IsValidHandle(SpawnRequestHandle))
 		{
 			return;
 		}
@@ -310,6 +332,11 @@ void UMassActorSpawnerSubsystem::ProcessPendingSpawningRequest(const double MaxT
 				ensureMsgf(SpawnRequestHandleManager.RemoveHandle(SpawnRequestHandle), TEXT("When providing a delegate, the spawn request gets automatically removed, no need to remove it on your side"));
 			}
 		}
+		else
+		{
+			// lower priority
+			SpawnRequest.SpawnStatus = ESpawnRequestStatus::RetryPending;
+		}
 	}
 }
 
@@ -458,3 +485,12 @@ void UMassActorSpawnerSubsystem::ReleaseAllResources()
 	CSV_CUSTOM_STAT(MassActors, NumSpawned, NumActorSpawned, ECsvCustomStatOp::Accumulate);
 	CSV_CUSTOM_STAT(MassActors, NumPooled, NumActorPooled, ECsvCustomStatOp::Accumulate);
 }
+
+//-----------------------------------------------------------------------------
+// DEPRECATED
+//-----------------------------------------------------------------------------
+FMassActorSpawnRequestHandle UMassActorSpawnerSubsystem::GetNextRequestToSpawn() const
+{
+	int32 DummyIndex = 0;
+	return GetNextRequestToSpawn(DummyIndex);
+}
diff --git a/Engine/Plugins/Runtime/MassGameplay/Source/MassActors/Public/MassActorSpawnerSubsystem.h b/Engine/Plugins/Runtime/MassGameplay/Source/MassActors/Public/MassActorSpawnerSubsystem.h
--- a/Engine/Plugins/Runtime/MassGameplay/Source/MassActors/Public/MassActorSpawnerSubsystem.h
+++ b/Engine/Plugins/Runtime/MassGameplay/Source/MassActors/Public/MassActorSpawnerSubsystem.h
@@ -192,10 +192,16 @@ protected:
 	/** Called at the end of the PrePhysics mass processing phase and calls ProcessPendingDestruction */ 
 	void OnPrePhysicsPhaseFinished(const float DeltaSeconds); 
 	
-	/** Retrieve what would be the next best spawning request to spawn, can be overridden to have different logic
+	/** 
+	 *  Retrieve what would be the next best spawning request to spawn, can be overridden to have different logic
 	 *  Default implementation is the first valid request in the list, no interesting logic yet
-	 *  @return the next best handle to spawn. */
-	virtual FMassActorSpawnRequestHandle GetNextRequestToSpawn() const;
+	 *  @param InOutHandleIndex used to start the search in subsequent locations. Also the index ensures the same handle 
+	 *    won't get returned twice in a row. InOutHandleIndex being INDEX_NONE indicates this is the first run, so all 
+	 *    handles are to be considered. If it's a  valid index then we iterate all but one to not even consider the 
+	 *    handle indicated by InOutHandleIndex.
+	 *  @return the next best handle to spawn. 
+	 */
+	virtual FMassActorSpawnRequestHandle GetNextRequestToSpawn(int32& InOutHandleIndex) const;
 
 	virtual ESpawnRequestStatus SpawnOrRetrieveFromPool(FConstStructView SpawnRequestView, TObjectPtr<AActor>& OutSpawnedActor);
 
@@ -241,4 +247,8 @@ protected:
 
 	mutable int32 NumActorSpawned = 0;
 	mutable int32 NumActorPooled = 0;
+
+public:
+	UE_DEPRECATED(5.4, "This flavor of GetNextRequestToSpawn is deprecated. Use the alternative taking an int32& parameter")
+	virtual FMassActorSpawnRequestHandle GetNextRequestToSpawn() const final;
 };
