Fixed the issue with MassVisualizationComponent::HandleChangesWithExternalIDTracking's "duplicate addition" handling that wasn't handling the case where there were duplicates in the incoming data (it was only checking the incoming data vs the existing data)

diff --git a/Engine/Plugins/Runtime/MassGameplay/Source/MassRepresentation/Private/MassVisualizationComponent.cpp b/Engine/Plugins/Runtime/MassGameplay/Source/MassRepresentation/Private/MassVisualizationComponent.cpp
--- a/Engine/Plugins/Runtime/MassGameplay/Source/MassRepresentation/Private/MassVisualizationComponent.cpp
+++ b/Engine/Plugins/Runtime/MassGameplay/Source/MassRepresentation/Private/MassVisualizationComponent.cpp
@@ -450,69 +450,50 @@ void UMassVisualizationComponent::HandleChangesWithExternalIDTracking(UInstanced
 		}
 	}
 
-	if (SharedData.GetUpdateInstanceIds().Num())
+	TArray<int32>& InstanceIds = SharedData.UpdateInstanceIds;
+	if (InstanceIds.Num())
 	{
-		TConstArrayView<int32> InstanceIds = SharedData.GetUpdateInstanceIds();
-		const TArray<FTransform>& InstanceTransforms = SharedData.GetStaticMeshInstanceTransformsArray();
-		int32 InNumCustomDataFloats = SharedData.GetStaticMeshInstanceCustomFloats().Num();
-		TConstArrayView<float> CustomFloatData = SharedData.GetStaticMeshInstanceCustomFloats();
-
-		const int32 StartingCount = ISMComponent.PerInstanceSMData.Num();
-		check(ISMComponent.InstanceIdToInstanceIndexMap.Num() == StartingCount);
+		check(ISMComponent.InstanceIdToInstanceIndexMap.Num() == ISMComponent.PerInstanceSMData.Num());
 
+		// We first add all the unique IDs, while removing incoming data duplicating the existing data.
+		// we need to identify the duplicates as the very first thing, since the ISMComponent.AddInstances
+		// call below actually adds data to the ISMComponent and by then we'd already have to many instances added.
+		// Once we're done here we can simply assume that all the InstanceIds are valid keys to ISMComponent.InstanceIdToInstanceIndexMap.
+		for (int32 IDIndex = InstanceIds.Num() - 1; IDIndex >= 0; --IDIndex)
 		{
-			// @todo want to track this for some time since it's quite possible this will cost us, without any actual 
-			// gain (duplicates should not happen if everything runs as expected).
-			TRACE_CPUPROFILER_EVENT_SCOPE_STR("MassVisualizationComponent_RemovingDuplicates");
-
-			// This part makes sure we're not trying to add the same InstanceId twice
-			// Note that we're going to ignore new data since updating is non-trivial. The user is expected to remove the old data first.
-			bool bDataRemoved = false;
-			if (ISMComponent.InstanceIdToInstanceIndexMap.Num())
-			{
-				for (int32 IDIndex = InstanceIds.Num() - 1; IDIndex >= 0; --IDIndex)
-				{
-					if (ISMComponent.InstanceIdToInstanceIndexMap.Find(InstanceIds[IDIndex]))
-					{
-						// note that it's safe to remove elements, even though we iterate over and array view 
-						// because RemoveUpdatedInstanceIdsAtSwap makes sure no memory is reallocated (using bAllowShrinking = false)
-						SharedData.RemoveUpdatedInstanceIdsAtSwap(IDIndex);
-						bDataRemoved = true;
-					}
-				}
-			}
-
-			if (bDataRemoved)
+			const int32 InstanceIndex = ISMComponent.InstanceIdToInstanceIndexMap.FindOrAdd(InstanceIds[IDIndex], INDEX_NONE);
+			if (InstanceIndex != INDEX_NONE)
 			{
-				// if we indeed removed data we need to update the array views
-				InstanceIds = SharedData.GetUpdateInstanceIds();
-				InNumCustomDataFloats = SharedData.GetStaticMeshInstanceCustomFloats().Num();
-				CustomFloatData = SharedData.GetStaticMeshInstanceCustomFloats();
+				SharedData.RemoveUpdatedInstanceIdsAtSwap(IDIndex);
 			}
 		}
 
-		// at this point it's possible we removed all the data - just make sure there's anything to add
+		// it's possible the loop above removed all the data, so we do one last check
 		if (InstanceIds.Num())
 		{
+			const TArray<FTransform>& InstanceTransforms = SharedData.GetStaticMeshInstanceTransformsArray();
+			const int32 InNumCustomDataFloats = SharedData.GetStaticMeshInstanceCustomFloats().Num();
+			TConstArrayView<float> CustomFloatData = SharedData.GetStaticMeshInstanceCustomFloats();
+
 			// if these are the first entities we're adding we need to set NumCustomDataFloats so that the PerInstanceSMCustomData
 			// gets populated properly by the AddInstancesInternal call below
+			const int32 StartingCount = ISMComponent.PerInstanceSMData.Num();
 			if (StartingCount == 0 && CustomFloatData.Num() && ISMComponent.Mobility != EComponentMobility::Static)
 			{
 				ISMComponent.NumCustomDataFloats = InNumCustomDataFloats;
 			}
 
 			check(InstanceIds.Num() == InstanceTransforms.Num());
-			TArray<int32> NewIndices = ISMComponent.AddInstances(InstanceTransforms, /*bShouldReturnIndices=*/true, /*bWorldSpace=*/true);
+			const TArray<int32> NewIndices = ISMComponent.AddInstances(InstanceTransforms, /*bShouldReturnIndices=*/true, /*bWorldSpace=*/true);
 
 			check(InstanceIds.Num() == NewIndices.Num());
+			// note that in this case ISMComponent.PerInstanceSMData.Num() is always going to be greater than ISMComponent.PerInstanceIds
+			// since the former gets added to as part of the ISMComponent.AddInstances above.
 			ISMComponent.PerInstanceIds.AddDefaulted(ISMComponent.PerInstanceSMData.Num() - ISMComponent.PerInstanceIds.Num());
 
 			for (int32 i = 0; i < InstanceIds.Num(); ++i)
 			{
-				checkf(ISMComponent.InstanceIdToInstanceIndexMap.Find(InstanceIds[i]) == nullptr
-					, TEXT("This occuring signals trouble. None of the InstanceIds is expected to have already been added to this ISM component instance."));
-
-				ISMComponent.InstanceIdToInstanceIndexMap.Add(InstanceIds[i], NewIndices[i]);
+				ISMComponent.InstanceIdToInstanceIndexMap[InstanceIds[i]] = NewIndices[i];
 				ISMComponent.PerInstanceIds[NewIndices[i]] = InstanceIds[i];
 			}
 
