Kick off of "stat Mass" support

diff --git a/Engine/Plugins/Runtime/MassEntity/Source/MassEntity/Private/MassProcessingPhaseManager.cpp b/Engine/Plugins/Runtime/MassEntity/Source/MassEntity/Private/MassProcessingPhaseManager.cpp
--- a/Engine/Plugins/Runtime/MassEntity/Source/MassEntity/Private/MassProcessingPhaseManager.cpp
+++ b/Engine/Plugins/Runtime/MassEntity/Source/MassEntity/Private/MassProcessingPhaseManager.cpp
@@ -15,7 +15,7 @@
 
 #define LOCTEXT_NAMESPACE "Mass"
 
-DECLARE_CYCLE_STAT(TEXT("Mass Phase Done"), STAT_MassPhaseDone, STATGROUP_TaskGraphTasks);
+DECLARE_CYCLE_STAT(TEXT("Mass Phase Tick"), STAT_Mass_PhaseTick, STATGROUP_Mass);
 
 namespace UE::Mass::Tweakables
 {
@@ -56,6 +56,8 @@ void FMassProcessingPhase::ExecuteTick(float DeltaTime, ELevelTick TickType, ENa
 		return;
 	}
 
+	SCOPE_CYCLE_COUNTER(STAT_Mass_PhaseTick);
+
 	checkf(PhaseManager, TEXT("Manager is null which is not a supported case. Either this FMassProcessingPhase has not been initialized properly or it's been left dangling after the FMassProcessingPhase owner got destroyed."));
 
 	TRACE_CPUPROFILER_EVENT_SCOPE_TEXT(*FString::Printf(TEXT("FMassProcessingPhase::ExecuteTick %s"), *UEnum::GetValueAsString(Phase)));
diff --git a/Engine/Plugins/Runtime/MassEntity/Source/MassEntity/Public/MassEntityTypes.h b/Engine/Plugins/Runtime/MassEntity/Source/MassEntity/Public/MassEntityTypes.h
--- a/Engine/Plugins/Runtime/MassEntity/Source/MassEntity/Public/MassEntityTypes.h
+++ b/Engine/Plugins/Runtime/MassEntity/Source/MassEntity/Public/MassEntityTypes.h
@@ -13,6 +13,8 @@
 
 MASSENTITY_API DECLARE_LOG_CATEGORY_EXTERN(LogMass, Warning, All);
 
+DECLARE_STATS_GROUP(TEXT("Mass"), STATGROUP_Mass, STATCAT_Advanced);
+
 // This is the base class for all lightweight fragments
 USTRUCT()
 struct FMassFragment
diff --git a/Engine/Plugins/Runtime/MassGameplay/Source/MassRepresentation/Private/MassRepresentationProcessor.cpp b/Engine/Plugins/Runtime/MassGameplay/Source/MassRepresentation/Private/MassRepresentationProcessor.cpp
--- a/Engine/Plugins/Runtime/MassGameplay/Source/MassRepresentation/Private/MassRepresentationProcessor.cpp
+++ b/Engine/Plugins/Runtime/MassGameplay/Source/MassRepresentation/Private/MassRepresentationProcessor.cpp
@@ -15,6 +15,8 @@
 #include "MassExecutionContext.h"
 
 
+DECLARE_CYCLE_STAT(TEXT("Visualization Processor Execute"), STAT_Mass_VisProcessor_Execute, STATGROUP_Mass);
+
 namespace UE::Mass::Representation
 {
 	int32 bAllowKeepActorExtraFrame = 1;
@@ -375,6 +377,8 @@ void UMassVisualizationProcessor::UpdateEntityVisibility(const FMassEntityHandle
 
 void UMassVisualizationProcessor::Execute(FMassEntityManager& InEntityManager, FMassExecutionContext& Context)
 {
+	SCOPE_CYCLE_COUNTER(STAT_Mass_VisProcessor_Execute);
+
 	int32 TotalEntitiesProcessed = 0;
 	// Update entities visualization
 	EntityQuery.ForEachEntityChunk(InEntityManager, Context, [this, &TotalEntitiesProcessed](FMassExecutionContext& Context)
diff --git a/Engine/Plugins/Runtime/MassGameplay/Source/MassRepresentation/Private/MassVisualizationComponent.cpp b/Engine/Plugins/Runtime/MassGameplay/Source/MassRepresentation/Private/MassVisualizationComponent.cpp
--- a/Engine/Plugins/Runtime/MassGameplay/Source/MassRepresentation/Private/MassVisualizationComponent.cpp
+++ b/Engine/Plugins/Runtime/MassGameplay/Source/MassRepresentation/Private/MassVisualizationComponent.cpp
@@ -14,6 +14,9 @@
 #include "Rendering/NaniteResources.h"
 #include "AI/NavigationSystemBase.h"
 
+
+DECLARE_CYCLE_STAT(TEXT("MassVisualizationComponent EndVisualChanges"), STAT_MassVisualizationComponent_EndVisualChanges, STATGROUP_Mass);
+
 //---------------------------------------------------------------
 // UMassVisualizationComponent
 //---------------------------------------------------------------
@@ -523,6 +526,7 @@ void UMassVisualizationComponent::HandleChangesWithExternalIDTracking(UInstanced
 void UMassVisualizationComponent::EndVisualChanges()
 {
 	TRACE_CPUPROFILER_EVENT_SCOPE_STR("MassVisualizationComponent EndVisualChanges")
+	SCOPE_CYCLE_COUNTER(STAT_MassVisualizationComponent_EndVisualChanges);
 
 	// Batch update gathered instance transforms
 	for (auto It = ISMCSharedData.CreateIterator(); It; ++It)
