Added GetActorSpawnerSubsystem accessor to MassRepresentationSubsystem

diff --git a/Engine/Plugins/Runtime/MassGameplay/Source/MassRepresentation/Public/MassRepresentationSubsystem.h b/Engine/Plugins/Runtime/MassGameplay/Source/MassRepresentation/Public/MassRepresentationSubsystem.h
--- a/Engine/Plugins/Runtime/MassGameplay/Source/MassRepresentation/Public/MassRepresentationSubsystem.h
+++ b/Engine/Plugins/Runtime/MassGameplay/Source/MassRepresentation/Public/MassRepresentationSubsystem.h
@@ -123,7 +123,9 @@ public:
 	/**
 	 * Release all references to static meshes and template actors
 	 * Use with caution, all entities using this representation subsystem must be destroy otherwise they will point to invalid resources */
-	 void ReleaseAllResources();
+	void ReleaseAllResources();
+
+	UMassActorSpawnerSubsystem* GetActorSpawnerSubsystem() const { return ActorSpawnerSubsystem; }
 
 protected:
 	// USubsystem BEGIN
