Made "last used index" utilized by MassActorSpawnerSubsystem's "find next spawn request handle" process to be a member variable to ensure smoother time-sliced processing of the handles (on subsequent frames picking up whete the process enged last frame).

diff --git a/Engine/Plugins/Runtime/MassGameplay/Source/MassActors/Private/MassActorSpawnerSubsystem.cpp b/Engine/Plugins/Runtime/MassGameplay/Source/MassActors/Private/MassActorSpawnerSubsystem.cpp
--- a/Engine/Plugins/Runtime/MassGameplay/Source/MassActors/Private/MassActorSpawnerSubsystem.cpp
+++ b/Engine/Plugins/Runtime/MassGameplay/Source/MassActors/Private/MassActorSpawnerSubsystem.cpp
@@ -275,7 +275,6 @@ void UMassActorSpawnerSubsystem::ProcessPendingSpawningRequest(const double MaxT
 
 	const int32 IterationsLimit = SpawnRequestHandleManager.CalcNumUsedHandles();
 	int32 IterationsCount = 0;
-	int32 StartingHandleIndex = INDEX_NONE;
 
 	while (FPlatformTime::Seconds() < TimeSliceEnd && IterationsCount++ < IterationsLimit)
 	{
diff --git a/Engine/Plugins/Runtime/MassGameplay/Source/MassActors/Public/MassActorSpawnerSubsystem.h b/Engine/Plugins/Runtime/MassGameplay/Source/MassActors/Public/MassActorSpawnerSubsystem.h
--- a/Engine/Plugins/Runtime/MassGameplay/Source/MassActors/Public/MassActorSpawnerSubsystem.h
+++ b/Engine/Plugins/Runtime/MassGameplay/Source/MassActors/Public/MassActorSpawnerSubsystem.h
@@ -248,6 +248,8 @@ protected:
 	mutable int32 NumActorSpawned = 0;
 	mutable int32 NumActorPooled = 0;
 
+	int32 StartingHandleIndex = INDEX_NONE;
+
 public:
 	UE_DEPRECATED(5.4, "This flavor of GetNextRequestToSpawn is deprecated. Use the alternative taking an int32& parameter")
 	virtual FMassActorSpawnRequestHandle GetNextRequestToSpawn() const final;
