Changes to which data gets processed by MassVisualizationComponent::EndVisualChanges, now it will only touch data that has been dirtied during the frame.

This change skips needlessly processing thousands of dormant entries in UMassVisualizationComponent::ISMCSharedData

diff --git a/Engine/Plugins/Runtime/MassGameplay/Source/MassRepresentation/Private/MassVisualizationComponent.cpp b/Engine/Plugins/Runtime/MassGameplay/Source/MassRepresentation/Private/MassVisualizationComponent.cpp
--- a/Engine/Plugins/Runtime/MassGameplay/Source/MassRepresentation/Private/MassVisualizationComponent.cpp
+++ b/Engine/Plugins/Runtime/MassGameplay/Source/MassRepresentation/Private/MassVisualizationComponent.cpp
@@ -1,6 +1,7 @@
 // Copyright Epic Games, Inc. All Rights Reserved.
 
 #include "MassVisualizationComponent.h"
+#include "CoreGlobals.h"
 #include "Logging/LogMacros.h"
 #include "MassVisualizer.h"
 #include "MassRepresentationTypes.h"
@@ -15,7 +16,12 @@
 #include "AI/NavigationSystemBase.h"
 
 
-DECLARE_CYCLE_STAT(TEXT("MassVisualizationComponent EndVisualChanges"), STAT_MassVisualizationComponent_EndVisualChanges, STATGROUP_Mass);
+DECLARE_CYCLE_STAT(TEXT("VisualizationComp EndVisualChanges"), STAT_Mass_VisualizationComponent_EndVisualChanges, STATGROUP_Mass);
+DECLARE_CYCLE_STAT(TEXT("VisualizationComp Handle Changes"), STAT_Mass_VisualizationComponent_HandleChangesWithExternalIDTracking, STATGROUP_Mass);
+
+DECLARE_DWORD_COUNTER_STAT(TEXT("VisualizationComp Instances Removed"), STAT_Mass_VisualizationComponent_InstancesRemovedNum, STATGROUP_Mass);
+DECLARE_DWORD_COUNTER_STAT(TEXT("VisualizationComp Instances Added"), STAT_Mass_VisualizationComponent_InstancesAddedNum, STATGROUP_Mass);
+
 
 //---------------------------------------------------------------
 // UMassVisualizationComponent
@@ -25,6 +31,10 @@ namespace UE::Mass::Representation
 {
 	int32 GCallUpdateInstances = 1;
 	FAutoConsoleVariableRef  CVarCallUpdateInstances(TEXT("Mass.CallUpdateInstances"), GCallUpdateInstances, TEXT("Toggle between UpdateInstances and BatchUpdateTransform."));
+
+#if STATS
+	uint32 LastStatsResetFrame = 0;
+#endif // STATS
 }  // UE::Mass::Representation
 
 void UMassVisualizationComponent::PostInitProperties()
@@ -181,7 +191,7 @@ void UMassVisualizationComponent::ConstructStaticMeshComponents()
 
 				if (SharedData == nullptr)
 				{
-					SharedData = &ISMCSharedData.Emplace(GetTypeHash(MeshDesc), FMassISMCSharedData(ISMC));
+					SharedData = &ISMCSharedData.Add(GetTypeHash(MeshDesc), FMassISMCSharedData(ISMC));
 				}
 				else
 				{
@@ -265,9 +275,9 @@ void UMassVisualizationComponent::ClearAllVisualInstances()
 	InstancedStaticMeshInfos.Reset();
 	
 	// Pool should already be empty, got a problem if it's not
-	for (auto It = ISMCSharedData.CreateIterator(); It; ++It)
+	for (int32 SharedDataIndex = 0; SharedDataIndex < ISMCSharedData.Num(); ++SharedDataIndex)
 	{
-		if (UInstancedStaticMeshComponent* InstancedStaticMeshComponent = It.Value().GetISMComponent())
+		if (UInstancedStaticMeshComponent* InstancedStaticMeshComponent = ISMCSharedData.GetAtIndex(SharedDataIndex).GetISMComponent())
 		{
 			InstancedStaticMeshComponent->ClearInstances();
 			InstancedStaticMeshComponent->DestroyComponent();
@@ -306,11 +316,13 @@ void UMassVisualizationComponent::HandleChangesWithExternalIDTracking(UInstanced
 {
 	constexpr float EqualTolerance = 1e-6;
 
-	if (SharedData.HasUpdatesToApply() == false)
+	if (!ensureMsgf(SharedData.HasUpdatesToApply(), TEXT("We're not expected to call this function for SharedData that needs no instance work, as per FMassISMCSharedData::FDirtyIterator used to iterate over data.")))
 	{
 		return;
 	}
 
+	SCOPE_CYCLE_COUNTER(STAT_Mass_VisualizationComponent_HandleChangesWithExternalIDTracking);
+
 	UHierarchicalInstancedStaticMeshComponent* HISMComp = Cast<UHierarchicalInstancedStaticMeshComponent>(&ISMComponent);
 	bool bAutoReset = HISMComp ? HISMComp->bAutoRebuildTreeOnInstanceChanges : false;
 
@@ -318,6 +330,8 @@ void UMassVisualizationComponent::HandleChangesWithExternalIDTracking(UInstanced
 	// that it's better to have redundant things visible than not seeing required things
 	if (SharedData.GetRemoveInstanceIds().Num())
 	{
+		INC_DWORD_STAT_BY(STAT_Mass_VisualizationComponent_InstancesRemovedNum, SharedData.GetRemoveInstanceIds().Num());
+
 		//RemoveInstanceWithIds(SharedData.GetRemoveInstanceIds());
 		TConstArrayView<int32> InstanceIds = SharedData.GetRemoveInstanceIds();
 
@@ -456,6 +470,8 @@ void UMassVisualizationComponent::HandleChangesWithExternalIDTracking(UInstanced
 	TArray<int32>& InstanceIds = SharedData.UpdateInstanceIds;
 	if (InstanceIds.Num())
 	{
+		INC_DWORD_STAT_BY(STAT_Mass_VisualizationComponent_InstancesAddedNum, InstanceIds.Num());
+
 		check(ISMComponent.InstanceIdToInstanceIndexMap.Num() == ISMComponent.PerInstanceSMData.Num());
 
 		// We first add all the unique IDs, while removing incoming data duplicating the existing data.
@@ -526,12 +542,22 @@ void UMassVisualizationComponent::HandleChangesWithExternalIDTracking(UInstanced
 void UMassVisualizationComponent::EndVisualChanges()
 {
 	TRACE_CPUPROFILER_EVENT_SCOPE_STR("MassVisualizationComponent EndVisualChanges")
-	SCOPE_CYCLE_COUNTER(STAT_MassVisualizationComponent_EndVisualChanges);
+	SCOPE_CYCLE_COUNTER(STAT_Mass_VisualizationComponent_EndVisualChanges);
+
+#if STATS
+	if (UE::Mass::Representation::LastStatsResetFrame != GFrameNumber)
+	{
+		SET_DWORD_STAT(STAT_Mass_VisualizationComponent_InstancesRemovedNum, 0);
+		SET_DWORD_STAT(STAT_Mass_VisualizationComponent_InstancesAddedNum, 0);
+		UE::Mass::Representation::LastStatsResetFrame = GFrameNumber;
+	}
+#endif // STATS
 
 	// Batch update gathered instance transforms
-	for (auto It = ISMCSharedData.CreateIterator(); It; ++It)
+	for (FMassISMCSharedDataMap::FDirtyIterator It(ISMCSharedData); It; ++It)
 	{
-		FMassISMCSharedData& SharedData = It.Value();
+		FMassISMCSharedData& SharedData = *It;
+
 		UInstancedStaticMeshComponent* InstancedStaticMeshComponent = SharedData.GetISMComponent();
 		// @todo need to check validity this way since Mass used to rely on the assumption that all the ISM components used were
 		// under its control. That's no longer the case, but the system has not been updated to take that into consideration.
@@ -605,6 +631,8 @@ void UMassVisualizationComponent::EndVisualChanges()
 		
 		SharedData.ResetAccumulatedData();
 	}
+
+	ISMCSharedData.ResetAllDirtyFlags();
 }
 
 //---------------------------------------------------------------
@@ -646,7 +674,7 @@ void FMassLODSignificanceRange::AddBatchedTransform(const int32 InstanceId, cons
 			continue;
 		}
 
-		FMassISMCSharedData& SharedData = (*ISMCSharedDataPtr)[StaticMeshRefs[i]];
+		FMassISMCSharedData& SharedData = ISMCSharedDataPtr->GetAndMarkDirty(StaticMeshRefs[i]);
 
 		SharedData.UpdateInstanceIds.Add(InstanceId);
 		SharedData.StaticMeshInstanceTransforms.Add(Transform);
@@ -664,7 +692,7 @@ void FMassLODSignificanceRange::AddBatchedCustomDataFloats(const TArray<float>&
 			continue;
 		}
 
-		FMassISMCSharedData& SharedData = (*ISMCSharedDataPtr)[StaticMeshRefs[i]];
+		FMassISMCSharedData& SharedData = ISMCSharedDataPtr->GetAndMarkDirty(StaticMeshRefs[i]);
 		SharedData.StaticMeshInstanceCustomFloats.Append(CustomFloats);
 	}
 }
@@ -674,7 +702,7 @@ void FMassLODSignificanceRange::AddInstance(const int32 InstanceId, const FTrans
 	check(ISMCSharedDataPtr);
 	for (int i = 0; i < StaticMeshRefs.Num(); i++)
 	{
-		FMassISMCSharedData& SharedData = (*ISMCSharedDataPtr)[StaticMeshRefs[i]];
+		FMassISMCSharedData& SharedData = ISMCSharedDataPtr->GetAndMarkDirty(StaticMeshRefs[i]);
 		SharedData.UpdateInstanceIds.Add(InstanceId);
 		SharedData.StaticMeshInstanceTransforms.Add(Transform);
 		SharedData.StaticMeshInstancePrevTransforms.Add(Transform);
@@ -686,7 +714,7 @@ void FMassLODSignificanceRange::RemoveInstance(const int32 InstanceId)
 	check(ISMCSharedDataPtr);
 	for (int i = 0; i < StaticMeshRefs.Num(); i++)
 	{
-		FMassISMCSharedData& SharedData = (*ISMCSharedDataPtr)[StaticMeshRefs[i]];
+		FMassISMCSharedData& SharedData = ISMCSharedDataPtr->GetAndMarkDirty(StaticMeshRefs[i]);
 		SharedData.RemoveInstanceIds.Add(InstanceId);
 	}
 }
@@ -701,7 +729,7 @@ void FMassLODSignificanceRange::WriteCustomDataFloatsAtStartIndex(int32 StaticMe
 			return;
 		}
 
-		FMassISMCSharedData& SharedData = (*ISMCSharedDataPtr)[StaticMeshRefs[StaticMeshIndex]];
+		FMassISMCSharedData& SharedData = ISMCSharedDataPtr->GetAndMarkDirty(StaticMeshRefs[StaticMeshIndex]);
 
 		int32 StartIndex = FloatsPerInstance * SharedData.WriteIterator + StartFloatIndex;
 
diff --git a/Engine/Plugins/Runtime/MassGameplay/Source/MassRepresentation/Public/MassRepresentationTypes.h b/Engine/Plugins/Runtime/MassGameplay/Source/MassRepresentation/Public/MassRepresentationTypes.h
--- a/Engine/Plugins/Runtime/MassGameplay/Source/MassRepresentation/Public/MassRepresentationTypes.h
+++ b/Engine/Plugins/Runtime/MassGameplay/Source/MassRepresentation/Public/MassRepresentationTypes.h
@@ -201,6 +201,11 @@ struct MASSREPRESENTATION_API FMassISMCSharedData
 	
 	bool RequiresExternalInstanceIDTracking() const { return bRequiresExternalInstanceIDTracking; }
 
+	void Reset() 
+	{
+		*this = FMassISMCSharedData();
+	}
+
 protected:
 	friend FMassLODSignificanceRange;
 	friend UMassVisualizationComponent;
@@ -228,7 +233,157 @@ protected:
 	uint8 bRequiresExternalInstanceIDTracking : 1;
 };
 
-using FMassISMCSharedDataMap = TMap<uint32, FMassISMCSharedData>;
+
+/** 
+ * The container type hosting FMassISMCSharedData instances and supplying functionality of marking entries that require 
+ * instance-related operations (adding, removing). 
+ * 
+ * To get a FMassISMCSharedData instance to add operations to it call GetAndMarkDirty.
+ * 
+ * Use FDirtyIterator to iterate over just the data that needs processing. 
+ * 
+ * @see UMassVisualizationComponent::EndVisualChanges for iteration
+ * @see FMassLODSignificanceRange methods for performing dirtying operations
+ */
+struct FMassISMCSharedDataMap 
+{
+	struct FDirtyIterator
+	{
+		friend FMassISMCSharedDataMap;
+		explicit FDirtyIterator(FMassISMCSharedDataMap& InContainer)
+			: Container(InContainer), It(InContainer.GetDirtyArray())
+		{
+			if (It && It.GetValue() != bValueToCheck)
+			{
+				// will result in either setting IT to the first bInValue, or making bool(It) == false
+				++(*this);
+			}
+		}
+	public:
+		operator bool() const { return bool(It); }
+
+		FDirtyIterator& operator++()
+		{
+			while (++It)
+			{
+				if (It.GetValue() == bValueToCheck)
+				{
+					break;
+				}
+			}
+			return *this;
+		}
+
+		FMassISMCSharedData& operator*() const
+		{
+			return Container.GetAtIndex(It.GetIndex());
+		}
+
+	private:
+		FMassISMCSharedDataMap& Container;
+		TBitArray<>::FConstIterator It;
+		static constexpr bool bValueToCheck = true;
+	};
+
+	FMassISMCSharedData& GetAndMarkDirty(const uint32 Hash)
+	{
+		const int32 DataIndex = Map[Hash];
+		DirtyData[DataIndex] = true;
+		return Data[DataIndex];
+	}
+	
+	FMassISMCSharedData& FindOrAdd(const uint32 Hash, const FMassISMCSharedData& NewData)
+	{
+		const int32* DataIndex = Map.Find(Hash);
+		if (DataIndex == nullptr)
+		{
+			return Add(Hash, NewData);
+		}
+		check(Data.IsValidIndex(*DataIndex));
+		return Data[*DataIndex];
+	}
+
+	FMassISMCSharedData* Find(const uint32 Hash)
+	{
+		int32* DataIndex = Map.Find(Hash);
+		return (DataIndex == nullptr || *DataIndex == INDEX_NONE) ? (FMassISMCSharedData*)nullptr : &Data[*DataIndex];
+	}
+
+	FMassISMCSharedData& Add(const uint32 Hash, const FMassISMCSharedData& NewData)
+	{
+		const int32 DataIndex = FreeIndices.Num() ? FreeIndices.Pop() : Data.Num();
+		Map.Add(Hash, DataIndex);
+
+		if (DataIndex == Data.Num())
+		{
+			DirtyData.Add(false, DataIndex - DirtyData.Num() + 1);
+			return Data.Add_GetRef(NewData);
+		}
+		else
+		{
+			DirtyData[DataIndex] = false;
+			Data[DataIndex] = NewData;
+			return Data[DataIndex];
+		}
+	}
+
+	void Remove(const uint32 Hash)
+	{
+		int32 DataIndex = INDEX_NONE;
+		if (ensure(Map.RemoveAndCopyValue(Hash, DataIndex)))
+		{
+			DirtyData[DataIndex] = false;
+			Data[DataIndex].Reset();
+			FreeIndices.Add(DataIndex);
+		}
+	}
+
+	FMassISMCSharedData& GetAtIndex(const int32 DataIndex)
+	{
+		return Data[DataIndex];
+	}
+	
+	const TBitArray<>& GetDirtyArray() const 
+	{ 
+		return DirtyData;
+	}
+
+	/** @return total number of entries in Data array. Note that some or all entries could be empty (i.e. already freed) */
+	int32 Num() const
+	{
+		return Data.Num();
+	}
+
+	/** @return number of non-empty entries in Data. */
+	int32 NumValid() const
+	{
+		return Data.Num() - FreeIndices.Num();
+	}
+
+	bool IsDirty(const int32 DataIndex) const
+	{
+		return DirtyData[DataIndex];
+	}
+
+	void ResetAllDirtyFlags()
+	{
+		DirtyData.SetRange(0, DirtyData.Num(), false);
+	}
+
+	void Reset()
+	{
+		*this = FMassISMCSharedDataMap();
+	}
+
+protected:
+	TArray<FMassISMCSharedData> Data;
+	/** Mapping from Hash of data represented by FMassISMCSharedData to an index to Data */
+	TMap<uint32, int32> Map;
+	/** Indicates whether corresponding Data entry has any instance work assigned to it (instance addition or removal) */
+	TBitArray<> DirtyData;
+	/** Indices to Data that are available for reuse */
+	TArray<int32> FreeIndices;
+};
 
 
 USTRUCT()
@@ -255,7 +410,7 @@ public:
 				continue;
 			}
 
-			FMassISMCSharedData& SharedData = (*ISMCSharedDataPtr)[StaticMeshRefs[i]];
+			FMassISMCSharedData& SharedData = (*ISMCSharedDataPtr).GetAndMarkDirty(StaticMeshRefs[i]);
 			const int32 StartIndex = SharedData.StaticMeshInstanceCustomFloats.AddDefaulted(StructSizeInFloats + NumFloatsToPad);
 			InCustomDataType* CustomData = reinterpret_cast<InCustomDataType*>(&SharedData.StaticMeshInstanceCustomFloats[StartIndex]);
 			*CustomData = InCustomData;
