More `stat Mass` stats

Main point, added "Total Mass" stat, that in principle should sum up all of the work performed by Mass as part of regular every-frame work

diff --git a/Engine/Plugins/Runtime/MassEntity/Source/MassEntity/Private/MassArchetypeData.cpp b/Engine/Plugins/Runtime/MassEntity/Source/MassEntity/Private/MassArchetypeData.cpp
--- a/Engine/Plugins/Runtime/MassEntity/Source/MassEntity/Private/MassArchetypeData.cpp
+++ b/Engine/Plugins/Runtime/MassEntity/Source/MassEntity/Private/MassArchetypeData.cpp
@@ -7,6 +7,8 @@
 #include "Misc/StringBuilder.h"
 
 
+DECLARE_CYCLE_STAT(TEXT("Mass Archetype BatchAdd"), STAT_Mass_ArchetypeBatchAdd, STATGROUP_Mass);
+
 namespace UE::Mass
 {
 	namespace Private
@@ -957,6 +959,8 @@ void FMassArchetypeData::REMOVEME_GetArrayViewForFragmentInChunk(int32 ChunkInde
 
 void FMassArchetypeData::BatchAddEntities(TConstArrayView<FMassEntityHandle> Entities, const FMassArchetypeSharedFragmentValues& SharedFragmentValues, TArray<FMassArchetypeEntityCollection::FArchetypeEntityRange>& OutNewRanges)
 {
+	SCOPE_CYCLE_COUNTER(STAT_Mass_ArchetypeBatchAdd);
+
 	FMassArchetypeEntityCollection::FArchetypeEntityRange ResultSubchunk;
 	ResultSubchunk.ChunkIndex = 0;
 	int32 NumberMoved = 0;
diff --git a/Engine/Plugins/Runtime/MassEntity/Source/MassEntity/Private/MassCommandBuffer.cpp b/Engine/Plugins/Runtime/MassEntity/Source/MassEntity/Private/MassCommandBuffer.cpp
--- a/Engine/Plugins/Runtime/MassEntity/Source/MassEntity/Private/MassCommandBuffer.cpp
+++ b/Engine/Plugins/Runtime/MassEntity/Source/MassEntity/Private/MassCommandBuffer.cpp
@@ -10,6 +10,7 @@
 
 CSV_DEFINE_CATEGORY(MassEntities, true);
 CSV_DEFINE_CATEGORY(MassEntitiesCounters, true);
+DECLARE_CYCLE_STAT(TEXT("Mass Flush Commands"), STAT_Mass_FlushCommands, STATGROUP_Mass);
 
 namespace UE::Mass::Command {
 
@@ -97,6 +98,7 @@ void FMassCommandBuffer::Flush(FMassEntityManager& EntityManager)
 	{
 		UE_MT_SCOPED_WRITE_ACCESS(PendingBatchCommandsDetector);
 		LLM_SCOPE_BYNAME(TEXT("Mass/FlushCommands"));
+		SCOPE_CYCLE_COUNTER(STAT_Mass_FlushCommands);
 
 		// array used to group commands depending on their operations. Based on EMassCommandOperationType
 		constexpr int32 CommandTypeOrder[] =
diff --git a/Engine/Plugins/Runtime/MassEntity/Source/MassEntity/Private/MassEntityTypes.cpp b/Engine/Plugins/Runtime/MassEntity/Source/MassEntity/Private/MassEntityTypes.cpp
--- a/Engine/Plugins/Runtime/MassEntity/Source/MassEntity/Private/MassEntityTypes.cpp
+++ b/Engine/Plugins/Runtime/MassEntity/Source/MassEntity/Private/MassEntityTypes.cpp
@@ -5,6 +5,8 @@
 
 #include UE_INLINE_GENERATED_CPP_BY_NAME(MassEntityTypes)
 
+DEFINE_STAT(STAT_Mass_Total);
+
 DEFINE_TYPEBITSET(FMassFragmentBitSet);
 DEFINE_TYPEBITSET(FMassTagBitSet);
 DEFINE_TYPEBITSET(FMassChunkFragmentBitSet);
diff --git a/Engine/Plugins/Runtime/MassEntity/Source/MassEntity/Private/MassExecutor.cpp b/Engine/Plugins/Runtime/MassEntity/Source/MassEntity/Private/MassExecutor.cpp
--- a/Engine/Plugins/Runtime/MassEntity/Source/MassEntity/Private/MassExecutor.cpp
+++ b/Engine/Plugins/Runtime/MassEntity/Source/MassEntity/Private/MassExecutor.cpp
@@ -156,6 +156,7 @@ struct FMassExecutorDoneTask
 	void DoTask(ENamedThreads::Type CurrentThread, const FGraphEventRef& MyCompletionGraphEvent)
 	{
 		TRACE_CPUPROFILER_EVENT_SCOPE_STR("Flush Deferred Commands Parallel");
+		SCOPE_CYCLE_COUNTER(STAT_Mass_Total);
 
 		FMassEntityManager& EntityManagerRef = ExecutionContext.GetEntityManagerChecked();
 
diff --git a/Engine/Plugins/Runtime/MassEntity/Source/MassEntity/Private/MassProcessingPhaseManager.cpp b/Engine/Plugins/Runtime/MassEntity/Source/MassEntity/Private/MassProcessingPhaseManager.cpp
--- a/Engine/Plugins/Runtime/MassEntity/Source/MassEntity/Private/MassProcessingPhaseManager.cpp
+++ b/Engine/Plugins/Runtime/MassEntity/Source/MassEntity/Private/MassProcessingPhaseManager.cpp
@@ -57,6 +57,7 @@ void FMassProcessingPhase::ExecuteTick(float DeltaTime, ELevelTick TickType, ENa
 	}
 
 	SCOPE_CYCLE_COUNTER(STAT_Mass_PhaseTick);
+	SCOPE_CYCLE_COUNTER(STAT_Mass_Total);
 
 	checkf(PhaseManager, TEXT("Manager is null which is not a supported case. Either this FMassProcessingPhase has not been initialized properly or it's been left dangling after the FMassProcessingPhase owner got destroyed."));
 
diff --git a/Engine/Plugins/Runtime/MassEntity/Source/MassEntity/Private/MassProcessor.cpp b/Engine/Plugins/Runtime/MassEntity/Source/MassEntity/Private/MassProcessor.cpp
--- a/Engine/Plugins/Runtime/MassEntity/Source/MassEntity/Private/MassProcessor.cpp
+++ b/Engine/Plugins/Runtime/MassEntity/Source/MassEntity/Private/MassProcessor.cpp
@@ -11,6 +11,7 @@
 #include UE_INLINE_GENERATED_CPP_BY_NAME(MassProcessor)
 
 DECLARE_CYCLE_STAT(TEXT("MassProcessor Group Completed"), Mass_GroupCompletedTask, STATGROUP_TaskGraphTasks);
+DECLARE_CYCLE_STAT(TEXT("Mass Processor Task"), STAT_Mass_DoTask, STATGROUP_Mass);
 
 #if WITH_MASSENTITY_DEBUG
 namespace UE::Mass::Debug
@@ -61,6 +62,8 @@ public:
 		checkf(Processor, TEXT("Expecting a valid processor to execute"));
 
 		PROCESSOR_TASK_LOG(TEXT("+--+ Task %s started on %u"), *Processor->GetProcessorName(), FPlatformTLS::GetCurrentThreadId());
+		SCOPE_CYCLE_COUNTER(STAT_Mass_DoTask);
+		SCOPE_CYCLE_COUNTER(STAT_Mass_Total);
 
 		check(EntityManager);
 		FMassEntityManager& EntityManagerRef = *EntityManager.Get();
diff --git a/Engine/Plugins/Runtime/MassEntity/Source/MassEntity/Public/MassEntityTypes.h b/Engine/Plugins/Runtime/MassEntity/Source/MassEntity/Public/MassEntityTypes.h
--- a/Engine/Plugins/Runtime/MassEntity/Source/MassEntity/Public/MassEntityTypes.h
+++ b/Engine/Plugins/Runtime/MassEntity/Source/MassEntity/Public/MassEntityTypes.h
@@ -14,6 +14,7 @@
 MASSENTITY_API DECLARE_LOG_CATEGORY_EXTERN(LogMass, Warning, All);
 
 DECLARE_STATS_GROUP(TEXT("Mass"), STATGROUP_Mass, STATCAT_Advanced);
+DECLARE_CYCLE_STAT_EXTERN(TEXT("Mass Total Frame Time"), STAT_Mass_Total, STATGROUP_Mass, MASSENTITY_API);
 
 // This is the base class for all lightweight fragments
 USTRUCT()
diff --git a/Engine/Plugins/Runtime/MassGameplay/Source/MassRepresentation/Private/MassRepresentationProcessor.cpp b/Engine/Plugins/Runtime/MassGameplay/Source/MassRepresentation/Private/MassRepresentationProcessor.cpp
--- a/Engine/Plugins/Runtime/MassGameplay/Source/MassRepresentation/Private/MassRepresentationProcessor.cpp
+++ b/Engine/Plugins/Runtime/MassGameplay/Source/MassRepresentation/Private/MassRepresentationProcessor.cpp
@@ -15,7 +15,7 @@
 #include "MassExecutionContext.h"
 
 
-DECLARE_CYCLE_STAT(TEXT("Visualization Processor Execute"), STAT_Mass_VisProcessor_Execute, STATGROUP_Mass);
+DECLARE_CYCLE_STAT(TEXT("Mass Visualization Execute"), STAT_Mass_VisProcessor_Execute, STATGROUP_Mass);
 
 namespace UE::Mass::Representation
 {
diff --git a/Engine/Plugins/Runtime/MassGameplay/Source/MassRepresentation/Private/MassVisualizationComponent.cpp b/Engine/Plugins/Runtime/MassGameplay/Source/MassRepresentation/Private/MassVisualizationComponent.cpp
--- a/Engine/Plugins/Runtime/MassGameplay/Source/MassRepresentation/Private/MassVisualizationComponent.cpp
+++ b/Engine/Plugins/Runtime/MassGameplay/Source/MassRepresentation/Private/MassVisualizationComponent.cpp
@@ -16,8 +16,8 @@
 #include "AI/NavigationSystemBase.h"
 
 
-DECLARE_CYCLE_STAT(TEXT("VisualizationComp EndVisualChanges"), STAT_Mass_VisualizationComponent_EndVisualChanges, STATGROUP_Mass);
-DECLARE_CYCLE_STAT(TEXT("VisualizationComp Handle Changes"), STAT_Mass_VisualizationComponent_HandleChangesWithExternalIDTracking, STATGROUP_Mass);
+DECLARE_CYCLE_STAT(TEXT("Mass Visualization EndVisualChanges"), STAT_Mass_VisualizationComponent_EndVisualChanges, STATGROUP_Mass);
+DECLARE_CYCLE_STAT(TEXT("Mass Visualization HandleIDs"), STAT_Mass_VisualizationComponent_HandleChangesWithExternalIDTracking, STATGROUP_Mass);
 
 DECLARE_DWORD_COUNTER_STAT(TEXT("VisualizationComp Instances Removed"), STAT_Mass_VisualizationComponent_InstancesRemovedNum, STATGROUP_Mass);
 DECLARE_DWORD_COUNTER_STAT(TEXT("VisualizationComp Instances Added"), STAT_Mass_VisualizationComponent_InstancesAddedNum, STATGROUP_Mass);
