Added a straightforward way to disable gathering streaming sources as Mass LOD viewers

diff --git a/Engine/Plugins/Runtime/MassGameplay/Source/MassLOD/Private/MassLODSubsystem.cpp b/Engine/Plugins/Runtime/MassGameplay/Source/MassLOD/Private/MassLODSubsystem.cpp
--- a/Engine/Plugins/Runtime/MassGameplay/Source/MassLOD/Private/MassLODSubsystem.cpp
+++ b/Engine/Plugins/Runtime/MassGameplay/Source/MassLOD/Private/MassLODSubsystem.cpp
@@ -179,26 +179,32 @@ void UMassLODSubsystem::SynchronizeViewers()
 
 	if (World)
 	{
-		// Now go through all current player controllers and add if they do not exist
-		for (FConstPlayerControllerIterator PlayerIterator = World->GetPlayerControllerIterator(); PlayerIterator; ++PlayerIterator)
+		if (bGatherPlayerControllers)
 		{
-			APlayerController* PlayerController = (*PlayerIterator).Get();
-			check(PlayerController);
-
-			// Check if the controller already exists by trying to remove it from the map which was filled up with controllers we were tracking
-			if (LocalViewerMap.Remove(GetTypeHash(PlayerController->GetFName())) == 0)
+			// Now go through all current player controllers and add if they do not exist
+			for (FConstPlayerControllerIterator PlayerIterator = World->GetPlayerControllerIterator(); PlayerIterator; ++PlayerIterator)
 			{
-				// If not add it to the list
-				AddPlayerViewer(*PlayerController);
+				APlayerController* PlayerController = (*PlayerIterator).Get();
+				check(PlayerController);
+
+				// Check if the controller already exists by trying to remove it from the map which was filled up with controllers we were tracking
+				if (LocalViewerMap.Remove(GetTypeHash(PlayerController->GetFName())) == 0)
+				{
+					// If not add it to the list
+					AddPlayerViewer(*PlayerController);
+				}
 			}
 		}
 
-		// Now go through all current streaming source and add if they do not exist
-		for (const FWorldPartitionStreamingSource& StreamingSource : StreamingSources)
+		if (bGatherStreamingSources)
 		{
-			if (LocalViewerMap.Remove(GetTypeHash(StreamingSource.Name)) == 0)
+			// Now go through all current streaming source and add if they do not exist
+			for (const FWorldPartitionStreamingSource& StreamingSource : StreamingSources)
 			{
-				AddStreamingSourceViewer(StreamingSource.Name);
+				if (LocalViewerMap.Remove(GetTypeHash(StreamingSource.Name)) == 0)
+				{
+					AddStreamingSourceViewer(StreamingSource.Name);
+				}
 			}
 		}
 	}
diff --git a/Engine/Plugins/Runtime/MassGameplay/Source/MassLOD/Public/MassLODSubsystem.h b/Engine/Plugins/Runtime/MassGameplay/Source/MassLOD/Public/MassLODSubsystem.h
--- a/Engine/Plugins/Runtime/MassGameplay/Source/MassLOD/Public/MassLODSubsystem.h
+++ b/Engine/Plugins/Runtime/MassGameplay/Source/MassLOD/Public/MassLODSubsystem.h
@@ -60,7 +60,7 @@ DECLARE_MULTICAST_DELEGATE_ThreeParams(FOnViewerRemoved, FMassViewerHandle Viewe
 /*
  * Manager responsible to manage and synchronized available viewers
  */
-UCLASS()
+UCLASS(config = Mass, defaultconfig)
 class MASSLOD_API UMassLODSubsystem : public UTickableWorldSubsystem
 {
 	GENERATED_BODY()
@@ -131,6 +131,15 @@ protected:
 	UFUNCTION()
 	void OnPlayerControllerEndPlay(AActor* Actor, EEndPlayReason::Type EndPlayReason);
 
+protected:
+	/** If true, all PlayerControllers will be gathered as viewers for LOD calcuations. */
+	UPROPERTY(EditDefaultsOnly, Category = "Mass|LOD", config)
+	uint8 bGatherPlayerControllers : 1 = true;
+
+	/** If true, all streaming sources will be gathered as viewers for LOD calcuations. */
+	UPROPERTY(EditDefaultsOnly, Category = "Mass|LOD", config)
+	uint8 bGatherStreamingSources : 1 = true;
+
 private:
 	/** Removes a viewer to the list and send notification about removal */
 	void RemoveViewerInternal(const FMassViewerHandle& ViewerHandle);
