Fixed handling of duplicates addition in MassVisualizationComponent::HandleChangesWithExternalIDTracking. The previous solution was breaking if there were duplicates in the incoming data since we weren't able to distinguish between data just added and the data added a couple of loop iterations earlier.

diff --git a/Engine/Plugins/Runtime/MassGameplay/Source/MassRepresentation/Private/MassVisualizationComponent.cpp b/Engine/Plugins/Runtime/MassGameplay/Source/MassRepresentation/Private/MassVisualizationComponent.cpp
--- a/Engine/Plugins/Runtime/MassGameplay/Source/MassRepresentation/Private/MassVisualizationComponent.cpp
+++ b/Engine/Plugins/Runtime/MassGameplay/Source/MassRepresentation/Private/MassVisualizationComponent.cpp
@@ -467,29 +467,38 @@ void UMassVisualizationComponent::HandleChangesWithExternalIDTracking(UInstanced
 		}
 	}
 
-	TArray<int32>& InstanceIds = SharedData.UpdateInstanceIds;
-	if (InstanceIds.Num())
+	if (SharedData.UpdateInstanceIds.Num())
 	{
-		INC_DWORD_STAT_BY(STAT_Mass_VisualizationComponent_InstancesAddedNum, InstanceIds.Num());
+		INC_DWORD_STAT_BY(STAT_Mass_VisualizationComponent_InstancesAddedNum, SharedData.UpdateInstanceIds.Num());
 
 		check(ISMComponent.InstanceIdToInstanceIndexMap.Num() == ISMComponent.PerInstanceSMData.Num());
 
 		// We first add all the unique IDs, while removing incoming data duplicating the existing data.
-		// we need to identify the duplicates as the very first thing, since the ISMComponent.AddInstances
+		// We need to identify the duplicates as the very first thing, since the ISMComponent.AddInstances
 		// call below actually adds data to the ISMComponent and by then we'd already have to many instances added.
-		// Once we're done here we can simply assume that all the InstanceIds are valid keys to ISMComponent.InstanceIdToInstanceIndexMap.
-		for (int32 IDIndex = InstanceIds.Num() - 1; IDIndex >= 0; --IDIndex)
+		// Once we're done here we can simply assume that all the SharedData.UpdateInstanceIds are valid keys to 
+		// ISMComponent.InstanceIdToInstanceIndexMap.
+		// Note that another effect of the process is that we don't update data for the instances that are already in 
+		// the ISMComponent. This is by design since this code path is intended for stationary entities.
+		
+		// We use TempIndexPrefix to annotate data we've just added and still be able to detect attempts to add duplicates in this very loop
+		constexpr int32 TempIndexPrefix = 0xff000000;
+		for (int32 IDIndex = SharedData.UpdateInstanceIds.Num() - 1; IDIndex >= 0; --IDIndex)
 		{
-			const int32 InstanceIndex = ISMComponent.InstanceIdToInstanceIndexMap.FindOrAdd(InstanceIds[IDIndex], INDEX_NONE);
-			if (InstanceIndex != INDEX_NONE)
+			const int32 TempIndexValue = (TempIndexPrefix | IDIndex);
+			const int32 InstanceIndex = ISMComponent.InstanceIdToInstanceIndexMap.FindOrAdd(SharedData.UpdateInstanceIds[IDIndex], TempIndexValue);
+
+			if (InstanceIndex != TempIndexValue)
 			{
 				SharedData.RemoveUpdatedInstanceIdsAtSwap(IDIndex);
 			}
 		}
 
 		// it's possible the loop above removed all the data, so we do one last check
-		if (InstanceIds.Num())
+		if (SharedData.UpdateInstanceIds.Num())
 		{
+			const TConstArrayView<int32> InstanceIds = SharedData.UpdateInstanceIds;
+
 			const TArray<FTransform>& InstanceTransforms = SharedData.GetStaticMeshInstanceTransformsArray();
 			const int32 InNumCustomDataFloats = SharedData.GetStaticMeshInstanceCustomFloats().Num();
 			TConstArrayView<float> CustomFloatData = SharedData.GetStaticMeshInstanceCustomFloats();
