Fixed in issue in UMassVisualizationComponent::ConstructStaticMeshComponents that resulted in attempts to register the same ISM component multiple times.

diff --git a/Engine/Plugins/Runtime/MassGameplay/Source/MassRepresentation/Private/MassVisualizationComponent.cpp b/Engine/Plugins/Runtime/MassGameplay/Source/MassRepresentation/Private/MassVisualizationComponent.cpp
--- a/Engine/Plugins/Runtime/MassGameplay/Source/MassRepresentation/Private/MassVisualizationComponent.cpp
+++ b/Engine/Plugins/Runtime/MassGameplay/Source/MassRepresentation/Private/MassVisualizationComponent.cpp
@@ -197,12 +197,11 @@ void UMassVisualizationComponent::ConstructStaticMeshComponents()
 				{
 					SharedData->SetISMComponent(*ISMC);
 				}
-			}
 
-			check(ISMC);
-			const uint32 ISMComponentPathHash = GetTypeHash(ISMC->GetPathName());
-			ensureMsgf(ISMComponentMap.Find(ISMComponentPathHash) == nullptr, TEXT("We've just created the ISMC that's being used here, so this check failing indicates hash-clash."));
-			ISMComponentMap.Add(ISMComponentPathHash, VisualIndex); 
+				const uint32 ISMComponentPathHash = GetTypeHash(ISMC->GetPathName());
+				ensureMsgf(ISMComponentMap.Find(ISMComponentPathHash) == nullptr, TEXT("We've just created the ISMC that's being used here, so this check failing indicates hash-clash."));
+				ISMComponentMap.Add(ISMComponentPathHash, VisualIndex);
+			}
 
 			check(SharedData);
 			Info.AddISMComponent(*SharedData);
