Fixed an issue with Mass ISM components resulting in visual ISM-related issues in CitySample.

The core of the issue was that crowd representation relies on ISM IDs magic to remove instances that are no longer being updated, while this information is not explicitly available on MassVisualizationComponent side. Because of that the ISM components that just became "empty" were not marked as "dirty" (because they were being made "empty" by not having anything added to them) and as such were not being processed by UMassVisualizationComponent::EndVisualChanges

diff --git a/Engine/Plugins/Runtime/MassGameplay/Source/MassRepresentation/Private/MassVisualizationComponent.cpp b/Engine/Plugins/Runtime/MassGameplay/Source/MassRepresentation/Private/MassVisualizationComponent.cpp
--- a/Engine/Plugins/Runtime/MassGameplay/Source/MassRepresentation/Private/MassVisualizationComponent.cpp
+++ b/Engine/Plugins/Runtime/MassGameplay/Source/MassRepresentation/Private/MassVisualizationComponent.cpp
@@ -315,8 +315,10 @@ void UMassVisualizationComponent::HandleChangesWithExternalIDTracking(UInstanced
 {
 	constexpr float EqualTolerance = 1e-6;
 
-	if (!ensureMsgf(SharedData.HasUpdatesToApply(), TEXT("We're not expected to call this function for SharedData that needs no instance work, as per FMassISMCSharedData::FDirtyIterator used to iterate over data.")))
+	if (SharedData.HasUpdatesToApply() == false)
 	{
+		// nothing to do here. We most probably were called as the part of the very first tick of this given SharedData
+		// since all the SharedData starts off as `dirty`.
 		return;
 	}
 
@@ -575,6 +577,7 @@ void UMassVisualizationComponent::EndVisualChanges()
 			if (SharedData.RequiresExternalInstanceIDTracking())
 			{
 				HandleChangesWithExternalIDTracking(*InstancedStaticMeshComponent, SharedData);
+				It.ClearDirtyFlag();
 			}
 			else
 			{
@@ -634,13 +637,12 @@ void UMassVisualizationComponent::EndVisualChanges()
 					// Dirty render state
 					InstancedStaticMeshComponent->MarkRenderStateDirty();
 				}
+				// note that we're not clearing the dirty flag on purpose - these components require constant updates
 			}
 		}
 		
 		SharedData.ResetAccumulatedData();
 	}
-
-	ISMCSharedData.ResetAllDirtyFlags();
 }
 
 //---------------------------------------------------------------
diff --git a/Engine/Plugins/Runtime/MassGameplay/Source/MassRepresentation/Public/MassRepresentationTypes.h b/Engine/Plugins/Runtime/MassGameplay/Source/MassRepresentation/Public/MassRepresentationTypes.h
--- a/Engine/Plugins/Runtime/MassGameplay/Source/MassRepresentation/Public/MassRepresentationTypes.h
+++ b/Engine/Plugins/Runtime/MassGameplay/Source/MassRepresentation/Public/MassRepresentationTypes.h
@@ -279,9 +279,14 @@ struct FMassISMCSharedDataMap
 			return Container.GetAtIndex(It.GetIndex());
 		}
 
+		void ClearDirtyFlag()
+		{
+			It.GetValue() = false;
+		}
+
 	private:
 		FMassISMCSharedDataMap& Container;
-		TBitArray<>::FConstIterator It;
+		TBitArray<>::FIterator It;
 		static constexpr bool bValueToCheck = true;
 	};
 
@@ -317,11 +322,12 @@ struct FMassISMCSharedDataMap
 		if (DataIndex == Data.Num())
 		{
 			DirtyData.Add(false, DataIndex - DirtyData.Num() + 1);
+			DirtyData[DataIndex] = true;
 			return Data.Add_GetRef(NewData);
 		}
 		else
 		{
-			DirtyData[DataIndex] = false;
+			DirtyData[DataIndex] = true;
 			Data[DataIndex] = NewData;
 			return Data[DataIndex];
 		}
@@ -343,7 +349,7 @@ struct FMassISMCSharedDataMap
 		return Data[DataIndex];
 	}
 	
-	const TBitArray<>& GetDirtyArray() const 
+	TBitArray<>& GetDirtyArray()
 	{ 
 		return DirtyData;
 	}
@@ -365,11 +371,6 @@ struct FMassISMCSharedDataMap
 		return DirtyData[DataIndex];
 	}
 
-	void ResetAllDirtyFlags()
-	{
-		DirtyData.SetRange(0, DirtyData.Num(), false);
-	}
-
 	void Reset()
 	{
 		*this = FMassISMCSharedDataMap();
