Made UMassStationaryISMSwitcherProcessor::Execute handle invalid representation index properly with some complaining (ensure).

diff --git a/Engine/Plugins/Runtime/MassGameplay/Source/MassRepresentation/Private/MassStationaryISMSwitcherProcessor.cpp b/Engine/Plugins/Runtime/MassGameplay/Source/MassRepresentation/Private/MassStationaryISMSwitcherProcessor.cpp
--- a/Engine/Plugins/Runtime/MassGameplay/Source/MassRepresentation/Private/MassStationaryISMSwitcherProcessor.cpp
+++ b/Engine/Plugins/Runtime/MassGameplay/Source/MassRepresentation/Private/MassStationaryISMSwitcherProcessor.cpp
@@ -62,6 +62,11 @@ void UMassStationaryISMSwitcherProcessor::ProcessContext(FMassExecutionContext&
 			continue;
 		}
 
+		if (!ensureMsgf(Representation.StaticMeshDescIndex != INDEX_NONE && ISMInfosView.IsValidIndex(Representation.StaticMeshDescIndex)
+						, TEXT("Invalid index %d for ISMInfosView"), Representation.StaticMeshDescIndex))
+		{
+			continue;
+		}
 		FMassInstancedStaticMeshInfo& ISMInfo = ISMInfosView[Representation.StaticMeshDescIndex];
 
 		if (Representation.PrevRepresentation == EMassRepresentationType::StaticMeshInstance
diff --git a/Engine/Plugins/Runtime/MassGameplay/Source/MassRepresentation/Public/MassRepresentationTypes.h b/Engine/Plugins/Runtime/MassGameplay/Source/MassRepresentation/Public/MassRepresentationTypes.h
--- a/Engine/Plugins/Runtime/MassGameplay/Source/MassRepresentation/Public/MassRepresentationTypes.h
+++ b/Engine/Plugins/Runtime/MassGameplay/Source/MassRepresentation/Public/MassRepresentationTypes.h
@@ -618,6 +618,11 @@ struct FMassInstancedStaticMeshInfoArrayViewAccessDetector
 		return InstancedStaticMeshInfos[Index];
 	}
 
+	bool IsValidIndex(const int32 Index) const
+	{
+		return InstancedStaticMeshInfos.IsValidIndex(Index);
+	}
+
 private:
 	TArrayView<FMassInstancedStaticMeshInfo> InstancedStaticMeshInfos;
 	const FRWRecursiveAccessDetector* AccessDetector;
