[Mass] Added missing initializations in FMassViewerInfoFragment

diff --git a/Engine/Plugins/Runtime/MassGameplay/Source/MassLOD/Public/MassLODFragments.h b/Engine/Plugins/Runtime/MassGameplay/Source/MassLOD/Public/MassLODFragments.h
--- a/Engine/Plugins/Runtime/MassGameplay/Source/MassLOD/Public/MassLODFragments.h
+++ b/Engine/Plugins/Runtime/MassGameplay/Source/MassLOD/Public/MassLODFragments.h
@@ -40,11 +40,11 @@ struct MASSLOD_API FMassViewerInfoFragment : public FMassFragment
 
 	// Closest viewer distance
 	UPROPERTY()
-	float ClosestViewerDistanceSq;
+	float ClosestViewerDistanceSq = FLT_MAX;
 
 	// Closest distance to frustum
 	UPROPERTY()
-	float ClosestDistanceToFrustum;
+	float ClosestDistanceToFrustum = FLT_MAX;
 };
 
 USTRUCT()
