Added Mass-side support for registering multi-mesh FStaticMeshInstanceVisualizationDesc while providing a list of ISMComponents to handle those meshes (used to support only a single ISMComp registration this way).

diff --git a/Engine/Plugins/Runtime/MassGameplay/Source/MassRepresentation/Private/MassRepresentationSubsystem.cpp b/Engine/Plugins/Runtime/MassGameplay/Source/MassRepresentation/Private/MassRepresentationSubsystem.cpp
--- a/Engine/Plugins/Runtime/MassGameplay/Source/MassRepresentation/Private/MassRepresentationSubsystem.cpp
+++ b/Engine/Plugins/Runtime/MassGameplay/Source/MassRepresentation/Private/MassRepresentationSubsystem.cpp
@@ -32,6 +32,12 @@ int16 UMassRepresentationSubsystem::AddVisualDescWithISMComponent(const FStaticM
 	return VisualizationComponent->AddVisualDescWithISMComponent(Desc, ISMComponent);
 }
 
+int16 UMassRepresentationSubsystem::AddVisualDescWithISMComponents(const FStaticMeshInstanceVisualizationDesc& Desc, TArrayView<TObjectPtr<UInstancedStaticMeshComponent>> ISMComponents)
+{
+	check(VisualizationComponent);
+	return VisualizationComponent->AddVisualDescWithISMComponents(Desc, ISMComponents);
+}
+
 void UMassRepresentationSubsystem::RemoveISMComponent(UInstancedStaticMeshComponent& ISMComponent)
 {
 	check(VisualizationComponent);
diff --git a/Engine/Plugins/Runtime/MassGameplay/Source/MassRepresentation/Private/MassVisualizationComponent.cpp b/Engine/Plugins/Runtime/MassGameplay/Source/MassRepresentation/Private/MassVisualizationComponent.cpp
--- a/Engine/Plugins/Runtime/MassGameplay/Source/MassRepresentation/Private/MassVisualizationComponent.cpp
+++ b/Engine/Plugins/Runtime/MassGameplay/Source/MassRepresentation/Private/MassVisualizationComponent.cpp
@@ -93,27 +93,46 @@ int16 UMassVisualizationComponent::FindOrAddVisualDesc(const FStaticMeshInstance
 
 int16 UMassVisualizationComponent::AddVisualDescWithISMComponent(const FStaticMeshInstanceVisualizationDesc& Desc, UInstancedStaticMeshComponent& ISMComponent)
 {
-	checkf(Desc.Meshes.Num() > 0, TEXT("%hs is expected to be used when there's exactly one mesh description contained"), __FUNCTION__);
-	ensureMsgf(Desc.Meshes.Num() == 1, TEXT("%hs is expected to be used when there's exactly one mesh description contained"), __FUNCTION__);
+	TObjectPtr<UInstancedStaticMeshComponent> AsObjectPtr = ISMComponent;
+	return AddVisualDescWithISMComponents(Desc, MakeArrayView(&AsObjectPtr, 1));
+}
 
-	const FMassStaticMeshInstanceVisualizationMeshDesc& MeshDesc = Desc.Meshes[0];
-	if (MeshDesc.Mesh == nullptr)
-	{
-		// invalid description, bail out
-		return (int16)INDEX_NONE;
-	}
+int16 UMassVisualizationComponent::AddVisualDescWithISMComponents(const FStaticMeshInstanceVisualizationDesc& Desc, TArrayView<TObjectPtr<UInstancedStaticMeshComponent>> ISMComponents)
+{
+	check(Desc.Meshes.Num() == ISMComponents.Num());
 	
 	UE_MT_SCOPED_WRITE_ACCESS(InstancedStaticMeshInfosDetector);
 
-	const int32 VisualIndex = AddInstancedStaticMeshInfo(Desc);
+	int32 VisualIndex = INDEX_NONE;
+	TArray<uint32> ISMComponentPathHashes;
 
-	const uint32 ISMComponentPathHash = GetTypeHash(ISMComponent.GetPathName()); 
-	FMassISMCSharedData& NewData = ISMCSharedData.FindOrAdd(ISMComponentPathHash, FMassISMCSharedData(&ISMComponent, /*bInRequiresExternalInstanceIDTracking=*/true));
-	InstancedStaticMeshInfos[VisualIndex].AddISMComponent(NewData);
+	for (int32 EntryIndex = 0; EntryIndex < Desc.Meshes.Num(); ++EntryIndex)
+	{
+		const FMassStaticMeshInstanceVisualizationMeshDesc& MeshDesc = Desc.Meshes[EntryIndex];
+		if (MeshDesc.Mesh == nullptr || ISMComponents[EntryIndex] == nullptr)
+		{
+			// invalid description, log an continue.
+			UE_VLOG_UELOG(this, LogMassRepresentation, Error, TEXT("Empty mesh at index %d while registering FStaticMeshInstanceVisualizationDesc instance"), EntryIndex);
+			continue;
+		}
+	
+		if (VisualIndex == INDEX_NONE)
+		{
+			VisualIndex = AddInstancedStaticMeshInfo(Desc);
+		}
 
-	BuildLODSignificanceForInfo(InstancedStaticMeshInfos[VisualIndex], ISMComponentPathHash);
+		const uint32 ISMComponentPathHash = GetTypeHash(ISMComponents[EntryIndex]->GetPathName());
+		FMassISMCSharedData& NewData = ISMCSharedData.FindOrAdd(ISMComponentPathHash, FMassISMCSharedData(ISMComponents[EntryIndex], /*bInRequiresExternalInstanceIDTracking=*/true));
+		InstancedStaticMeshInfos[VisualIndex].AddISMComponent(NewData);
+		ISMComponentPathHashes.Add(ISMComponentPathHash);
 	
-	ISMComponentMap.Add(ISMComponentPathHash, VisualIndex);
+		ISMComponentMap.Add(ISMComponentPathHash, VisualIndex);
+	}
+
+	if (VisualIndex != INDEX_NONE)
+	{
+		BuildLODSignificanceForInfo(InstancedStaticMeshInfos[VisualIndex], ISMComponentPathHashes);
+	}
 
 	checkf(VisualIndex < INT16_MAX, TEXT("%hs resulting VisualIndex is out of expected bounds"), __FUNCTION__);
 	return (int16)VisualIndex;
@@ -215,7 +234,7 @@ void UMassVisualizationComponent::ConstructStaticMeshComponents()
 	}
 }
 
-void UMassVisualizationComponent::BuildLODSignificanceForInfo(FMassInstancedStaticMeshInfo& Info, const uint32 ForcedStaticMeshRefKey)
+void UMassVisualizationComponent::BuildLODSignificanceForInfo(FMassInstancedStaticMeshInfo& Info, TConstArrayView<uint32> ForcedStaticMeshRefKeys)
 {
 	TArray<float> AllLODSignificances;
 	auto UniqueInsertOrdered = [&AllLODSignificances](const float Significance)
@@ -244,20 +263,21 @@ void UMassVisualizationComponent::BuildLODSignificanceForInfo(FMassInstancedStat
 	if (AllLODSignificances.Num() > 1)
 	{
 		Info.LODSignificanceRanges.SetNum(AllLODSignificances.Num() - 1);
-		for (int i = 0; i < Info.LODSignificanceRanges.Num(); ++i)
+		for (int RangeIndex = 0; RangeIndex < Info.LODSignificanceRanges.Num(); ++RangeIndex)
 		{
-			FMassLODSignificanceRange& Range = Info.LODSignificanceRanges[i];
-			Range.MinSignificance = AllLODSignificances[i];
-			Range.MaxSignificance = AllLODSignificances[i+1];
+			FMassLODSignificanceRange& Range = Info.LODSignificanceRanges[RangeIndex];
+			Range.MinSignificance = AllLODSignificances[RangeIndex];
+			Range.MaxSignificance = AllLODSignificances[RangeIndex + 1];
 			Range.ISMCSharedDataPtr = &ISMCSharedData;
 
-			for (int j = 0; j < Info.Desc.Meshes.Num(); ++j)
+			for (int MeshIndex = 0; MeshIndex < Info.Desc.Meshes.Num(); ++MeshIndex)
 			{
-				const FMassStaticMeshInstanceVisualizationMeshDesc& MeshDesc = Info.Desc.Meshes[j];
+				const FMassStaticMeshInstanceVisualizationMeshDesc& MeshDesc = Info.Desc.Meshes[MeshIndex];
 				const bool bAddMeshInRange = (Range.MinSignificance >= MeshDesc.MinLODSignificance && Range.MinSignificance < MeshDesc.MaxLODSignificance);
 				if (bAddMeshInRange)
 				{
-					Range.StaticMeshRefs.Add(ForcedStaticMeshRefKey ? ForcedStaticMeshRefKey : GetTypeHash(MeshDesc));
+					Range.StaticMeshRefs.Add(ForcedStaticMeshRefKeys.IsValidIndex(MeshIndex) && ForcedStaticMeshRefKeys[MeshIndex]
+						? ForcedStaticMeshRefKeys[MeshIndex] : GetTypeHash(MeshDesc));
 				}
 			}
 		}
@@ -752,3 +772,11 @@ void FMassLODSignificanceRange::WriteCustomDataFloatsAtStartIndex(int32 StaticMe
 		SharedData.WriteIterator++;
 	}
 }
+
+//-----------------------------------------------------------------------------
+// DEPRECATED
+//-----------------------------------------------------------------------------
+void UMassVisualizationComponent::BuildLODSignificanceForInfo(FMassInstancedStaticMeshInfo& Info, const uint32 ForcedStaticMeshRefKeys)
+{
+	BuildLODSignificanceForInfo(Info, MakeArrayView(&ForcedStaticMeshRefKeys, 1));
+}
\ No newline at end of file
diff --git a/Engine/Plugins/Runtime/MassGameplay/Source/MassRepresentation/Public/MassRepresentationSubsystem.h b/Engine/Plugins/Runtime/MassGameplay/Source/MassRepresentation/Public/MassRepresentationSubsystem.h
--- a/Engine/Plugins/Runtime/MassGameplay/Source/MassRepresentation/Public/MassRepresentationSubsystem.h
+++ b/Engine/Plugins/Runtime/MassGameplay/Source/MassRepresentation/Public/MassRepresentationSubsystem.h
@@ -39,10 +39,17 @@ public:
 
 	/**
 	 * Creates a dedicated visual type described by host Desc and ties ISMComponent to it.
+	 * @note this is a helper function for a common "single ISMComponent" case. Calls AddVisualDescWithISMComponents under the hood.
 	 * @return The index of the visual type
 	 */
 	int16 AddVisualDescWithISMComponent(const FStaticMeshInstanceVisualizationDesc& Desc, UInstancedStaticMeshComponent& ISMComponent);
 
+	/**
+	 * Creates a dedicated visual type described by host Desc and ties given ISMComponents to it.
+	 * @return The index of the visual type
+	 */
+	int16 AddVisualDescWithISMComponents(const FStaticMeshInstanceVisualizationDesc& Desc, TArrayView<TObjectPtr<UInstancedStaticMeshComponent>> ISMComponents);
+
 	/**
 	 * Removes the visualization data associated with the given ISM component. Note that this is safe to do only when
 	 * there are no entities relying on this data. No entity data patching will take place.
diff --git a/Engine/Plugins/Runtime/MassGameplay/Source/MassRepresentation/Public/MassVisualizationComponent.h b/Engine/Plugins/Runtime/MassGameplay/Source/MassRepresentation/Public/MassVisualizationComponent.h
--- a/Engine/Plugins/Runtime/MassGameplay/Source/MassRepresentation/Public/MassVisualizationComponent.h
+++ b/Engine/Plugins/Runtime/MassGameplay/Source/MassRepresentation/Public/MassVisualizationComponent.h
@@ -31,10 +31,17 @@ public:
 
 	/** 
 	 * Creates a dedicated visual type described by host Desc and ties ISMComponent to it.
+	 * @note this is a helper function for a common "single ISMComponent" case. Calls AddVisualDescWithISMComponents under the hood.
 	 * @return The index of the visual type 
 	 */
 	int16 AddVisualDescWithISMComponent(const FStaticMeshInstanceVisualizationDesc& Desc, UInstancedStaticMeshComponent& ISMComponent);
 
+	/**
+	 * Creates a dedicated visual type described by host Desc and ties given ISMComponents to it.
+	 * @return The index of the visual type
+	 */
+	int16 AddVisualDescWithISMComponents(const FStaticMeshInstanceVisualizationDesc& Desc, TArrayView<TObjectPtr<UInstancedStaticMeshComponent>> ISMComponents);
+
 	/** 
 	 * Removes the visualization data associated with the given ISM component. Note that this is safe to do only when
 	 * there are no entities relying on this data. No entity data patching will take place. 
@@ -75,10 +82,10 @@ protected:
 
 	/**
 	 * Creates LODSignificance ranges for all the meshes indicated by Info
-	 * @param ForcedStaticMeshRefKey if set to anything other than 0 will be used when adding individual FMassStaticMeshInstanceVisualizationMeshDesc
+	 * @param ForcedStaticMeshRefKeys if not empty will be used when adding individual FMassStaticMeshInstanceVisualizationMeshDesc
 	 *	instances to LOD significance ranges.
-	 */
-	void BuildLODSignificanceForInfo(FMassInstancedStaticMeshInfo& Info, const uint32 ForcedStaticMeshRefKey = 0);
+	 */	
+	void BuildLODSignificanceForInfo(FMassInstancedStaticMeshInfo& Info, TConstArrayView<uint32> ForcedStaticMeshRefKeys = TConstArrayView<uint32>());
 
 	/** Either adds an element to InstancedStaticMeshInfos or reuses an existing entry based on InstancedStaticMeshInfosFreeIndices*/
 	int32 AddInstancedStaticMeshInfo(const FStaticMeshInstanceVisualizationDesc& Desc);
@@ -98,4 +105,7 @@ protected:
 
 	/** Indicies to InstancedStaticMeshInfos that need their SMComponent constructed */
 	TArray<int32> InstancedSMComponentsRequiringConstructing;
+
+	UE_DEPRECATED(5.4, "This flavor of BuildLODSignificanceForInfo is no longer supported. Use the other one instead.")
+	void BuildLODSignificanceForInfo(FMassInstancedStaticMeshInfo& Info, const uint32 ForcedStaticMeshRefKey);
 };
