Only do navigation updates for the required instances in UMassVisualizationComponent::HandleChangesWithExternalIDTracking (instead of dirtying the whole component size)

diff --git a/Engine/Plugins/Runtime/MassGameplay/Source/MassRepresentation/Private/MassVisualizationComponent.cpp b/Engine/Plugins/Runtime/MassGameplay/Source/MassRepresentation/Private/MassVisualizationComponent.cpp
--- a/Engine/Plugins/Runtime/MassGameplay/Source/MassRepresentation/Private/MassVisualizationComponent.cpp
+++ b/Engine/Plugins/Runtime/MassGameplay/Source/MassRepresentation/Private/MassVisualizationComponent.cpp
@@ -417,6 +417,9 @@ void UMassVisualizationComponent::HandleChangesWithExternalIDTracking(UInstanced
 			// remove instance
 			if (ISMComponent.PerInstanceSMData.IsValidIndex(InstanceIndex))
 			{
+				// Request navigation update
+				ISMComponent.PartialNavigationUpdate(InstanceIndex);
+				
 				ISMComponent.PerInstanceSMData.RemoveAtSwap(InstanceIndex, 1, false);
 				ISMComponent.PerInstanceSMCustomData.RemoveAt(InstanceIndex * ISMComponent.NumCustomDataFloats, ISMComponent.NumCustomDataFloats, false);
 			}
@@ -557,10 +560,6 @@ void UMassVisualizationComponent::HandleChangesWithExternalIDTracking(UInstanced
 	{
 		FNavigationSystem::UnregisterComponent(ISMComponent);
 	}
-	else
-	{
-		FNavigationSystem::UpdateComponentData(ISMComponent);
-	}
 
 	if (HISMComp)
 	{
