Improved UMassActorSpawnerSubsystem's actor destruction to not attempt to destroy actors that have already been removed from the world.

This can happen when an actor is spawned with an OverrideLevel that then gets streamed out. At that point the actor will get destroyed as part of that process and UMassActorSpawnerSubsystem trying to destroy the actor again might cause a bunch of additional issues (details depend on specific build type and project).

diff --git a/Engine/Plugins/Runtime/MassGameplay/Source/MassActors/Private/MassActorSpawnerSubsystem.cpp b/Engine/Plugins/Runtime/MassGameplay/Source/MassActors/Private/MassActorSpawnerSubsystem.cpp
--- a/Engine/Plugins/Runtime/MassGameplay/Source/MassActors/Private/MassActorSpawnerSubsystem.cpp
+++ b/Engine/Plugins/Runtime/MassGameplay/Source/MassActors/Private/MassActorSpawnerSubsystem.cpp
@@ -54,8 +54,27 @@ bool UMassActorSpawnerSubsystem::RemoveActorSpawnRequest(FMassActorSpawnRequestH
 	return true;
 }
 
+void UMassActorSpawnerSubsystem::ConditionalDestroyActor(UWorld& World, AActor& ActorToDestroy)
+{
+	UWorld* ActorsWorld = ActorToDestroy.GetWorld();
+
+	// we directly call DestroyActors only if they're tied to the current world.
+	// Otherwise we rely on engine mechanics to get rid of them.
+	if (ActorsWorld == &World)
+	{
+		World.DestroyActor(&ActorToDestroy);
+	}
+	else
+	{
+		ensureMsgf(ActorToDestroy.HasActorBegunPlay() == false, TEXT("Failed to destroy %s due to world mismatch, while the actor is still being valid (as indicated by HasActorBegunPlay() == true)")
+			, *ActorToDestroy.GetName());
+	}
+}
+
 void UMassActorSpawnerSubsystem::DestroyActor(AActor* Actor, bool bImmediate /*= false*/)
 {
+	check(Actor);
+
 	// We need to unregister immediately MassAgentComponent as it will become out of sync with mass
 	if (UMassAgentComponent* AgentComp = Actor->FindComponentByClass<UMassAgentComponent>())
 	{
@@ -67,10 +86,11 @@ void UMassActorSpawnerSubsystem::DestroyActor(AActor* Actor, bool bImmediate /*=
 	{
 		if (!ReleaseActorToPool(Actor))
 		{
+			// Couldn't release actor back to pool, so destroy it
 			UWorld* World = GetWorld();
 			check(World);
+			ConditionalDestroyActor(*World, *Actor);
 
-			World->DestroyActor(Actor);
 			--NumActorSpawned;
 		}
 	}
@@ -357,10 +377,10 @@ void UMassActorSpawnerSubsystem::ProcessPendingDestruction(const double MaxTimeS
 			   (HasToDestroyAllActorsOnServerSide || FPlatformTime::Seconds() <= TimeSliceEnd))
 		{
 			AActor* ActorToDestroy = DeactivatedActorsToDestroy.Num() ? DeactivatedActorsToDestroy.Pop(/*bAllowShrinking*/false) : ActorsToDestroy.Pop(/*bAllowShrinking*/false);
-			if (!ReleaseActorToPool(ActorToDestroy))
+			if (ActorToDestroy && !ReleaseActorToPool(ActorToDestroy))
 			{
 				// Couldn't release actor back to pool, so destroy it
-				World->DestroyActor(ActorToDestroy);
+				ConditionalDestroyActor(*World, *ActorToDestroy);
 				--NumActorSpawned;
 			}
 		}
@@ -473,7 +493,10 @@ void UMassActorSpawnerSubsystem::ReleaseAllResources()
 			auto& ActorArray = It.Value();
 			for (int i = 0; i < ActorArray.Num(); i++)
 			{
-				World->DestroyActor(ActorArray[i]);
+				if (ActorArray[i])
+				{
+					ConditionalDestroyActor(*World, *ActorArray[i]);
+				}
 			}
 			NumActorSpawned -= ActorArray.Num();
 		}
diff --git a/Engine/Plugins/Runtime/MassGameplay/Source/MassActors/Public/MassActorSpawnerSubsystem.h b/Engine/Plugins/Runtime/MassGameplay/Source/MassActors/Public/MassActorSpawnerSubsystem.h
--- a/Engine/Plugins/Runtime/MassGameplay/Source/MassActors/Public/MassActorSpawnerSubsystem.h
+++ b/Engine/Plugins/Runtime/MassGameplay/Source/MassActors/Public/MassActorSpawnerSubsystem.h
@@ -108,7 +108,7 @@ public:
 /**
  * A subsystem managing spawning of actors for all mass subsystems
  */
-UCLASS()
+UCLASS(transient)
 class MASSACTORS_API UMassActorSpawnerSubsystem : public UWorldSubsystem
 {
 	GENERATED_BODY()
@@ -190,6 +190,11 @@ public:
 	void ReleaseAllResources();
 
 protected:
+	/** 
+	 * Provides consistent way of conditional destroying Actor within World. The actual destruction depends on Actor's state
+	 * and whether it belongs to World
+	 */
+	static void ConditionalDestroyActor(UWorld& World, AActor& ActorToDestroy);
 
 	/** Called at the start of the PrePhysics mass processing phase and calls ProcessPendingSpawningRequest */ 
 	void OnPrePhysicsPhaseStarted(const float DeltaSeconds);
