A minor MassActorSpawnerSubsystem refactor opening doors to instant actor spawning for Mass

diff --git a/Engine/Plugins/Runtime/MassGameplay/Source/MassActors/Private/MassActorSpawnerSubsystem.cpp b/Engine/Plugins/Runtime/MassGameplay/Source/MassActors/Private/MassActorSpawnerSubsystem.cpp
--- a/Engine/Plugins/Runtime/MassGameplay/Source/MassActors/Private/MassActorSpawnerSubsystem.cpp
+++ b/Engine/Plugins/Runtime/MassGameplay/Source/MassActors/Private/MassActorSpawnerSubsystem.cpp
@@ -316,47 +316,75 @@ void UMassActorSpawnerSubsystem::ProcessPendingSpawningRequest(const double MaxT
 			return;
 		}
 
-		// Do the spawning
-		SpawnRequest.SpawnStatus = ESpawnRequestStatus::Processing;
+		ESpawnRequestStatus Result = ProcessSpawnRequest(SpawnRequestHandle, SpawnRequestView, SpawnRequest);
+		ensureMsgf(Result != ESpawnRequestStatus::None, TEXT("Getting ESpawnRequestStatus::None as a result in this context is unexpected. Needs to be investigated."));
+	}
+}
 
-		// Call the pre spawn delegate on the spawn request
-		if (SpawnRequest.ActorPreSpawnDelegate.IsBound())
-		{
-			SpawnRequest.ActorPreSpawnDelegate.Execute(SpawnRequestHandle, SpawnRequestView);
-		}
+ESpawnRequestStatus UMassActorSpawnerSubsystem::ProcessSpawnRequest(const FMassActorSpawnRequestHandle SpawnRequestHandle)
+{
+	if (!SpawnRequestHandle.IsValid()
+		|| !SpawnRequestHandleManager.IsValidHandle(SpawnRequestHandle))
+	{
+		return ESpawnRequestStatus::None;
+	}
 
-		SpawnRequest.SpawnStatus = SpawnOrRetrieveFromPool(SpawnRequestView, SpawnRequest.SpawnedActor);
+	FStructView SpawnRequestView = SpawnRequests[SpawnRequestHandle.GetIndex()];
+	FMassActorSpawnRequest& SpawnRequest = SpawnRequestView.Get<FMassActorSpawnRequest>();
 
-		if (SpawnRequest.IsFinished())
-		{
-			if (SpawnRequest.SpawnStatus == ESpawnRequestStatus::Succeeded && IsValid(SpawnRequest.SpawnedActor))
-			{
-				if (UMassAgentComponent* AgentComp = SpawnRequest.SpawnedActor->FindComponentByClass<UMassAgentComponent>())
-				{
-					AgentComp->SetPuppetHandle(SpawnRequest.MassAgent);
-				}
-			}
+	return ProcessSpawnRequest(SpawnRequestHandle, SpawnRequestView, SpawnRequest);
+}
 
-			EMassActorSpawnRequestAction PostAction = EMassActorSpawnRequestAction::Remove;
+ESpawnRequestStatus UMassActorSpawnerSubsystem::ProcessSpawnRequest(const FMassActorSpawnRequestHandle SpawnRequestHandle, FStructView SpawnRequestView, FMassActorSpawnRequest& SpawnRequest)
+{
+	if (!ensureMsgf(SpawnRequest.IsFinished() == false, TEXT("Finished spawn requests are not expected to be processed again. Bailing out.")))
+	{
+		// returning None rather than the actual SpawnRequest.SpawnStatus to indicate the issue has occurred. 
+		return ESpawnRequestStatus::None;
+	}
 
-			// Call the post spawn delegate on the spawn request
-			if (SpawnRequest.ActorPostSpawnDelegate.IsBound())
-			{
-				PostAction = SpawnRequest.ActorPostSpawnDelegate.Execute(SpawnRequestHandle, SpawnRequestView);
-			}
+	// Do the spawning
+	SpawnRequest.SpawnStatus = ESpawnRequestStatus::Processing;
+
+	// Call the pre spawn delegate on the spawn request
+	if (SpawnRequest.ActorPreSpawnDelegate.IsBound())
+	{
+		SpawnRequest.ActorPreSpawnDelegate.Execute(SpawnRequestHandle, SpawnRequestView);
+	}
+
+	SpawnRequest.SpawnStatus = SpawnOrRetrieveFromPool(SpawnRequestView, SpawnRequest.SpawnedActor);
 
-			if (PostAction == EMassActorSpawnRequestAction::Remove)
+	if (SpawnRequest.IsFinished())
+	{
+		if (SpawnRequest.SpawnStatus == ESpawnRequestStatus::Succeeded && IsValid(SpawnRequest.SpawnedActor))
+		{
+			if (UMassAgentComponent* AgentComp = SpawnRequest.SpawnedActor->FindComponentByClass<UMassAgentComponent>())
 			{
-				// If notified, remove the spawning request
-				ensureMsgf(SpawnRequestHandleManager.RemoveHandle(SpawnRequestHandle), TEXT("When providing a delegate, the spawn request gets automatically removed, no need to remove it on your side"));
+				AgentComp->SetPuppetHandle(SpawnRequest.MassAgent);
 			}
 		}
-		else
+
+		EMassActorSpawnRequestAction PostAction = EMassActorSpawnRequestAction::Remove;
+
+		// Call the post spawn delegate on the spawn request
+		if (SpawnRequest.ActorPostSpawnDelegate.IsBound())
 		{
-			// lower priority
-			SpawnRequest.SpawnStatus = ESpawnRequestStatus::RetryPending;
+			PostAction = SpawnRequest.ActorPostSpawnDelegate.Execute(SpawnRequestHandle, SpawnRequestView);
 		}
+
+		if (PostAction == EMassActorSpawnRequestAction::Remove)
+		{
+			// If notified, remove the spawning request
+			ensureMsgf(SpawnRequestHandleManager.RemoveHandle(SpawnRequestHandle), TEXT("When providing a delegate, the spawn request gets automatically removed, no need to remove it on your side"));
+		}
+	}
+	else
+	{
+		// lower priority
+		SpawnRequest.SpawnStatus = ESpawnRequestStatus::RetryPending;
 	}
+
+	return SpawnRequest.SpawnStatus;
 }
 
 void UMassActorSpawnerSubsystem::ProcessPendingDestruction(const double MaxTimeSlicePerTick)
diff --git a/Engine/Plugins/Runtime/MassGameplay/Source/MassActors/Public/MassActorSpawnerSubsystem.h b/Engine/Plugins/Runtime/MassGameplay/Source/MassActors/Public/MassActorSpawnerSubsystem.h
--- a/Engine/Plugins/Runtime/MassGameplay/Source/MassActors/Public/MassActorSpawnerSubsystem.h
+++ b/Engine/Plugins/Runtime/MassGameplay/Source/MassActors/Public/MassActorSpawnerSubsystem.h
@@ -131,6 +131,19 @@ public:
 		return RequestActorSpawnInternal(FConstStructView::Make(InSpawnRequest));
 	}
 
+	/**
+	 * Process a valid spawn request indicated by given handle. Can be used to force instant-spawn of an actor provided 
+	 * a valid handle is obtained by calling RequestActorSpawn first. 
+	 * @return indicates the status of processed spawn request, with ESpawnRequestStatus::None indicating that "something 
+	 *	went wrong" and spawning request has not been processed. 
+	 */
+	[[nodiscard]] ESpawnRequestStatus ProcessSpawnRequest(const FMassActorSpawnRequestHandle SpawnRequestHandle);
+
+	/** 
+	 * Similar to the other ProcessSpawnRequest flavor, but with SpawnRequestView and SpawnRequest already provided. 
+	 */
+	[[nodiscard]] ESpawnRequestStatus ProcessSpawnRequest(const FMassActorSpawnRequestHandle SpawnRequestHandle, FStructView SpawnRequestView, FMassActorSpawnRequest& SpawnRequest);
+
 	/** Retries a failed spawn request
 	 * @param SpawnRequestHandle the spawn request handle to retry
 	 */
