Refactored UMassRepresentationActorManagement::GetOrSpawnActor to remove the defunct parameter. The old version has been deprecated and all the call sites updated.

Also renamed some function parameters to make their function clearer.

[RN] This change might cause compilation errors since the deprecated UMassRepresentationActorManagement::GetOrSpawnActor is virtual and has been marked final.

diff --git a/Engine/Plugins/AI/MassCrowd/Source/MassCrowd/Private/MassCrowdRepresentationActorManagement.cpp b/Engine/Plugins/AI/MassCrowd/Source/MassCrowd/Private/MassCrowdRepresentationActorManagement.cpp
--- a/Engine/Plugins/AI/MassCrowd/Source/MassCrowd/Private/MassCrowdRepresentationActorManagement.cpp
+++ b/Engine/Plugins/AI/MassCrowd/Source/MassCrowd/Private/MassCrowdRepresentationActorManagement.cpp
@@ -65,7 +65,7 @@ void UMassCrowdRepresentationActorManagement::SetActorEnabled(const EMassActorEn
 }
 
 AActor* UMassCrowdRepresentationActorManagement::GetOrSpawnActor(UMassRepresentationSubsystem& RepresentationSubsystem
-	, FMassEntityManager& EntityManager, const FMassEntityHandle MassAgent, FMassActorFragment& ActorInfo, const FTransform& Transform
+	, FMassEntityManager& EntityManager, const FMassEntityHandle MassAgent, const FTransform& Transform
 	, const int16 TemplateActorIndex, FMassActorSpawnRequestHandle& SpawnRequestHandle, const float Priority) const
 {
 	FTransform RootTransform = Transform;
@@ -78,7 +78,7 @@ AActor* UMassCrowdRepresentationActorManagement::GetOrSpawnActor(UMassRepresenta
 		}
 	}
 
-	return Super::GetOrSpawnActor(RepresentationSubsystem, EntityManager, MassAgent, ActorInfo, RootTransform, TemplateActorIndex, SpawnRequestHandle, Priority);
+	return Super::GetOrSpawnActor(RepresentationSubsystem, EntityManager, MassAgent, RootTransform, TemplateActorIndex, SpawnRequestHandle, Priority);
 }
 
 void UMassCrowdRepresentationActorManagement::TeleportActor(const FTransform& Transform, AActor& Actor, FMassCommandBuffer& CommandBuffer) const
diff --git a/Engine/Plugins/AI/MassCrowd/Source/MassCrowd/Public/MassCrowdRepresentationActorManagement.h b/Engine/Plugins/AI/MassCrowd/Source/MassCrowd/Public/MassCrowdRepresentationActorManagement.h
--- a/Engine/Plugins/AI/MassCrowd/Source/MassCrowd/Public/MassCrowdRepresentationActorManagement.h
+++ b/Engine/Plugins/AI/MassCrowd/Source/MassCrowd/Public/MassCrowdRepresentationActorManagement.h
@@ -39,7 +39,7 @@ protected:
 	 * @return the actor spawned
 	 */
 	virtual AActor* GetOrSpawnActor(UMassRepresentationSubsystem& RepresentationSubsystem, FMassEntityManager& EntitySubsystem
-		, const FMassEntityHandle MassAgent, FMassActorFragment& ActorInfo, const FTransform& Transform, const int16 TemplateActorIndex
+		, const FMassEntityHandle MassAgent, const FTransform& Transform, const int16 TemplateActorIndex
 		, FMassActorSpawnRequestHandle& SpawnRequestHandle, const float Priority) const override;
 
 	/**
diff --git a/Engine/Plugins/Runtime/MassGameplay/Source/MassRepresentation/Private/MassRepresentationActorManagement.cpp b/Engine/Plugins/Runtime/MassGameplay/Source/MassRepresentation/Private/MassRepresentationActorManagement.cpp
--- a/Engine/Plugins/Runtime/MassGameplay/Source/MassRepresentation/Private/MassRepresentationActorManagement.cpp
+++ b/Engine/Plugins/Runtime/MassGameplay/Source/MassRepresentation/Private/MassRepresentationActorManagement.cpp
@@ -14,10 +14,10 @@ float UMassRepresentationActorManagement::GetSpawnPriority(const FMassRepresenta
 }
 
 AActor* UMassRepresentationActorManagement::GetOrSpawnActor(UMassRepresentationSubsystem& RepresentationSubsystem, FMassEntityManager& EntityManager
-	, const FMassEntityHandle MassAgent, FMassActorFragment& ActorInfo, const FTransform& Transform, const int16 TemplateActorIndex
-	, FMassActorSpawnRequestHandle& SpawnRequestHandle, const float Priority) const
+	, const FMassEntityHandle MassAgent, const FTransform& Transform, const int16 TemplateActorIndex
+	, FMassActorSpawnRequestHandle& InOutSpawnRequestHandle, const float Priority) const
 {
-	return RepresentationSubsystem.GetOrSpawnActorFromTemplate(MassAgent, Transform, TemplateActorIndex, SpawnRequestHandle, Priority,
+	return RepresentationSubsystem.GetOrSpawnActorFromTemplate(MassAgent, Transform, TemplateActorIndex, InOutSpawnRequestHandle, Priority,
 		FMassActorPreSpawnDelegate::CreateUObject(this, &UMassRepresentationActorManagement::OnPreActorSpawn, &EntityManager),
 		FMassActorPostSpawnDelegate::CreateUObject(this, &UMassRepresentationActorManagement::OnPostActorSpawn, &EntityManager));
 }
@@ -137,3 +137,12 @@ void UMassRepresentationActorManagement::ReleaseAnyActorOrCancelAnySpawning(UMas
 	}
 	check(!Representation.ActorSpawnRequestHandle.IsValid());
 }
+
+AActor* UMassRepresentationActorManagement::GetOrSpawnActor(UMassRepresentationSubsystem& RepresentationSubsystem, FMassEntityManager& EntityManager
+	, const FMassEntityHandle MassAgent, FMassActorFragment&/* OutActorInfo*/, const FTransform& Transform, const int16 TemplateActorIndex
+	, FMassActorSpawnRequestHandle& InOutSpawnRequestHandle, const float Priority) const
+{
+	return GetOrSpawnActor(RepresentationSubsystem, EntityManager
+		, MassAgent, Transform, TemplateActorIndex
+		, InOutSpawnRequestHandle, Priority);
+}
\ No newline at end of file
diff --git a/Engine/Plugins/Runtime/MassGameplay/Source/MassRepresentation/Private/MassRepresentationProcessor.cpp b/Engine/Plugins/Runtime/MassGameplay/Source/MassRepresentation/Private/MassRepresentationProcessor.cpp
--- a/Engine/Plugins/Runtime/MassGameplay/Source/MassRepresentation/Private/MassRepresentationProcessor.cpp
+++ b/Engine/Plugins/Runtime/MassGameplay/Source/MassRepresentation/Private/MassRepresentationProcessor.cpp
@@ -110,7 +110,7 @@ void UMassRepresentationProcessor::UpdateRepresentation(FMassExecutionContext& C
 				// If we already queued spawn request but have changed our mind, continue with it but once we get the actor back, disable it immediately
 				if (Representation.ActorSpawnRequestHandle.IsValid())
 				{
-					Actor = RepresentationActorManagement->GetOrSpawnActor(*RepresentationSubsystem, CachedEntityManager, MassAgentEntityHandle, ActorInfo, TransformFragment.GetTransform(), Representation.LowResTemplateActorIndex, Representation.ActorSpawnRequestHandle, RepresentationActorManagement->GetSpawnPriority(RepresentationLOD));
+					Actor = RepresentationActorManagement->GetOrSpawnActor(*RepresentationSubsystem, CachedEntityManager, MassAgentEntityHandle, TransformFragment.GetTransform(), Representation.LowResTemplateActorIndex, Representation.ActorSpawnRequestHandle, RepresentationActorManagement->GetSpawnPriority(RepresentationLOD));
 				}
 			}
 			if (Actor != nullptr)
@@ -157,7 +157,7 @@ void UMassRepresentationProcessor::UpdateRepresentation(FMassExecutionContext& C
 							!RepresentationSubsystem->DoesActorMatchTemplate(*Actor, WantedTemplateActorIndex) ||
 							Representation.ActorSpawnRequestHandle.IsValid())
 						{
-							NewActor = RepresentationActorManagement->GetOrSpawnActor(*RepresentationSubsystem, CachedEntityManager, MassAgentEntityHandle, ActorInfo, TransformFragment.GetTransform(), WantedTemplateActorIndex, Representation.ActorSpawnRequestHandle, RepresentationActorManagement->GetSpawnPriority(RepresentationLOD));
+							NewActor = RepresentationActorManagement->GetOrSpawnActor(*RepresentationSubsystem, CachedEntityManager, MassAgentEntityHandle, TransformFragment.GetTransform(), WantedTemplateActorIndex, Representation.ActorSpawnRequestHandle, RepresentationActorManagement->GetSpawnPriority(RepresentationLOD));
 						}
 						else
 						{
diff --git a/Engine/Plugins/Runtime/MassGameplay/Source/MassRepresentation/Private/MassRepresentationSubsystem.cpp b/Engine/Plugins/Runtime/MassGameplay/Source/MassRepresentation/Private/MassRepresentationSubsystem.cpp
--- a/Engine/Plugins/Runtime/MassGameplay/Source/MassRepresentation/Private/MassRepresentationSubsystem.cpp
+++ b/Engine/Plugins/Runtime/MassGameplay/Source/MassRepresentation/Private/MassRepresentationSubsystem.cpp
@@ -70,7 +70,7 @@ int16 UMassRepresentationSubsystem::FindOrAddTemplateActor(const TSubclassOf<AAc
 }
 
 AActor* UMassRepresentationSubsystem::GetOrSpawnActorFromTemplate(const FMassEntityHandle MassAgent, const FTransform& Transform
-	, const int16 TemplateActorIndex, FMassActorSpawnRequestHandle& SpawnRequestHandle, float Priority
+	, const int16 TemplateActorIndex, FMassActorSpawnRequestHandle& InOutSpawnRequestHandle, float Priority
 	, FMassActorPreSpawnDelegate ActorPreSpawnDelegate, FMassActorPostSpawnDelegate ActorPostSpawnDelegate)
 {
 	UE_MT_SCOPED_READ_ACCESS(TemplateActorsMTAccessDetector);
@@ -84,9 +84,9 @@ AActor* UMassRepresentationSubsystem::GetOrSpawnActorFromTemplate(const FMassEnt
 	check(ActorSpawnerSubsystem);
 	const TSubclassOf<AActor> TemplateToSpawn = TemplateActors[TemplateActorIndex];
 
-	if (SpawnRequestHandle.IsValid())
+	if (InOutSpawnRequestHandle.IsValid())
 	{
-		FMassActorSpawnRequest& SpawnRequest = ActorSpawnerSubsystem->GetMutableSpawnRequest<FMassActorSpawnRequest>(SpawnRequestHandle);
+		FMassActorSpawnRequest& SpawnRequest = ActorSpawnerSubsystem->GetMutableSpawnRequest<FMassActorSpawnRequest>(InOutSpawnRequestHandle);
 		// Check if this existing spawn request is matching the template actor
 		if (SpawnRequest.Template != TemplateToSpawn)
 		{
@@ -114,7 +114,7 @@ AActor* UMassRepresentationSubsystem::GetOrSpawnActorFromTemplate(const FMassEnt
 					// Update spawn request with latest information and retry
 					SpawnRequest.Transform = Transform;
 					SpawnRequest.Priority = Priority;
-					ActorSpawnerSubsystem->RetryActorSpawnRequest(SpawnRequestHandle);
+					ActorSpawnerSubsystem->RetryActorSpawnRequest(InOutSpawnRequestHandle);
 				}
 				return nullptr;
 			}
@@ -122,12 +122,12 @@ AActor* UMassRepresentationSubsystem::GetOrSpawnActorFromTemplate(const FMassEnt
 			{
 				AActor* SpawnedActor = SpawnRequest.SpawnedActor;
 				check(SpawnedActor);
-				ensureMsgf(ActorSpawnerSubsystem->RemoveActorSpawnRequest(SpawnRequestHandle), TEXT("Unable to remove a valid spawn request"));
+				ensureMsgf(ActorSpawnerSubsystem->RemoveActorSpawnRequest(InOutSpawnRequestHandle), TEXT("Unable to remove a valid spawn request"));
 				return SpawnedActor;
 			}
 			default:
 				checkf(false, TEXT("Unexpected spawn request status!"));
-				SpawnRequestHandle.Invalidate();
+				InOutSpawnRequestHandle.Invalidate();
 				return nullptr;
 		}
 	}
@@ -140,7 +140,7 @@ AActor* UMassRepresentationSubsystem::GetOrSpawnActorFromTemplate(const FMassEnt
 	SpawnRequest.Priority = Priority;
 	SpawnRequest.ActorPreSpawnDelegate = ActorPreSpawnDelegate;
 	SpawnRequest.ActorPostSpawnDelegate = ActorPostSpawnDelegate;
-	SpawnRequestHandle = ActorSpawnerSubsystem->RequestActorSpawn(SpawnRequest);
+	InOutSpawnRequestHandle = ActorSpawnerSubsystem->RequestActorSpawn(SpawnRequest);
 
 	++(HandledMassAgents.FindOrAdd(MassAgent));
 
diff --git a/Engine/Plugins/Runtime/MassGameplay/Source/MassRepresentation/Public/MassRepresentationActorManagement.h b/Engine/Plugins/Runtime/MassGameplay/Source/MassRepresentation/Public/MassRepresentationActorManagement.h
--- a/Engine/Plugins/Runtime/MassGameplay/Source/MassRepresentation/Public/MassRepresentationActorManagement.h
+++ b/Engine/Plugins/Runtime/MassGameplay/Source/MassRepresentation/Public/MassRepresentationActorManagement.h
@@ -42,8 +42,8 @@ public:
 	 * @return the actor spawned
 	 */
 	virtual AActor* GetOrSpawnActor(UMassRepresentationSubsystem& RepresentationSubsystem, FMassEntityManager& EntitySubsystem
-		, const FMassEntityHandle MassAgent, FMassActorFragment& ActorInfo, const FTransform& Transform, const int16 TemplateActorIndex
-		, FMassActorSpawnRequestHandle& SpawnRequestHandle, const float Priority) const;
+		, const FMassEntityHandle MassAgent, const FTransform& Transform, const int16 TemplateActorIndex
+		, FMassActorSpawnRequestHandle& InOutSpawnRequestHandle, const float Priority) const;
 
 	/**
 	 * Enable/disable a spawned actor for a mass entity
@@ -101,4 +101,9 @@ public:
 	 * @param Representation fragment containing the current and previous visual state
 	 */
 	static void ReleaseAnyActorOrCancelAnySpawning(UMassRepresentationSubsystem& RepresentationSubsystem, const FMassEntityHandle MassAgent, FMassActorFragment& ActorInfo, FMassRepresentationFragment& Representation);
+
+	UE_DEPRECATED(5.4, "This flavor of GetOrSpawnActor has been deprecated due to a defunct parameter, OutActorInfo, that was never being used.")
+	virtual AActor* GetOrSpawnActor(UMassRepresentationSubsystem& RepresentationSubsystem, FMassEntityManager& EntitySubsystem
+		, const FMassEntityHandle MassAgent, FMassActorFragment&/* OutActorInfo*/, const FTransform& Transform, const int16 TemplateActorIndex
+		, FMassActorSpawnRequestHandle& InOutSpawnRequestHandle, const float Priority) const final;
 };
\ No newline at end of file
diff --git a/Engine/Plugins/Runtime/MassGameplay/Source/MassRepresentation/Public/MassRepresentationSubsystem.h b/Engine/Plugins/Runtime/MassGameplay/Source/MassRepresentation/Public/MassRepresentationSubsystem.h
--- a/Engine/Plugins/Runtime/MassGameplay/Source/MassRepresentation/Public/MassRepresentationSubsystem.h
+++ b/Engine/Plugins/Runtime/MassGameplay/Source/MassRepresentation/Public/MassRepresentationSubsystem.h
@@ -82,7 +82,7 @@ public:
 	 * @param ActorPostSpawnDelegate is an optional delegate called once the actor is spawned
 	 * @return The spawned actor from the template actor type if ready
 	 */
-	AActor* GetOrSpawnActorFromTemplate(const FMassEntityHandle MassAgent, const FTransform& Transform, const int16 TemplateActorIndex, FMassActorSpawnRequestHandle& SpawnRequestHandle, float Priority = MAX_FLT,
+	AActor* GetOrSpawnActorFromTemplate(const FMassEntityHandle MassAgent, const FTransform& Transform, const int16 TemplateActorIndex, FMassActorSpawnRequestHandle& InOutSpawnRequestHandle, float Priority = MAX_FLT,
 		FMassActorPreSpawnDelegate ActorPreSpawnDelegate = FMassActorPreSpawnDelegate(), FMassActorPostSpawnDelegate ActorPostSpawnDelegate = FMassActorPostSpawnDelegate());
 
 	/**
