reworded comments on FMassVisualizationLODParameters's properties to provide more detail.

diff --git a/Engine/Plugins/Runtime/MassGameplay/Source/MassRepresentation/Public/MassRepresentationFragments.h b/Engine/Plugins/Runtime/MassGameplay/Source/MassRepresentation/Public/MassRepresentationFragments.h
--- a/Engine/Plugins/Runtime/MassGameplay/Source/MassRepresentation/Public/MassRepresentationFragments.h
+++ b/Engine/Plugins/Runtime/MassGameplay/Source/MassRepresentation/Public/MassRepresentationFragments.h
@@ -197,10 +197,15 @@ struct FMassVisualizationLODParameters : public FMassSharedFragment
 	UPROPERTY(EditAnywhere, Category = "Mass|LOD", config)
 	int32 LODMaxCount[EMassLOD::Max] = {50, 100, 500, MAX_int32};
 
-	/** How far away from frustum does this entities are considered visible */
+	/** Entities within this distance from frustum will be considered visible. Expressed in Unreal Units. */
 	UPROPERTY(EditAnywhere, Category = "Mass|LOD", meta = (ClampMin = "0.0", UIMin = "0.0"), config)
 	float DistanceToFrustum = 0.0f;
+
 	/** Once visible how much further than DistanceToFrustum does the entities need to be before being cull again */
+	/** 
+	 * Once an entity is visible how far away from frustum does it need to get to lose "visible" state. 
+	 * Expressed in Unreal Units and is added to DistanceToFrustum to arrive at the final value to be used for testing.
+	 */
 	UPROPERTY(EditAnywhere, Category = "Mass|LOD", meta = (ClampMin = "0.0", UIMin = "0.0"), config)
 	float DistanceToFrustumHysteresis = 0.0f;
 
