[Mass] MassVisualizer are no longer visible in the outliner so user can't delete them.

diff --git a/Engine/Plugins/Runtime/MassGameplay/Source/MassRepresentation/Private/MassVisualizer.cpp b/Engine/Plugins/Runtime/MassGameplay/Source/MassRepresentation/Private/MassVisualizer.cpp
--- a/Engine/Plugins/Runtime/MassGameplay/Source/MassRepresentation/Private/MassVisualizer.cpp
+++ b/Engine/Plugins/Runtime/MassGameplay/Source/MassRepresentation/Private/MassVisualizer.cpp
@@ -8,5 +8,10 @@ AMassVisualizer::AMassVisualizer()
 	RootComponent = CreateDefaultSubobject<USceneComponent>(TEXT("SceneComp"));
 	RootComponent->SetMobility(EComponentMobility::Static);
 	VisComponent = CreateDefaultSubobject<UMassVisualizationComponent>(TEXT("VisualizerComponent"));
+
+	// MassVisualizers are created by the subsystems and should not be exposed to users or allowed to be deleted.
+#if WITH_EDITORONLY_DATA
+	bListedInSceneOutliner = false;
+#endif
 }
 
