Removed MassLWI and all of project specific code that was using it.

diff --git a/Engine/Plugins/Runtime/MassGameplay/MassGameplay.uplugin b/Engine/Plugins/Runtime/MassGameplay/MassGameplay.uplugin
--- a/Engine/Plugins/Runtime/MassGameplay/MassGameplay.uplugin
+++ b/Engine/Plugins/Runtime/MassGameplay/MassGameplay.uplugin
@@ -73,10 +73,6 @@
 			"Name": "MassGameplayExternalTraits",
 			"Type": "Runtime"
 		},
-		{
-			"Name": "MassLWI",
-			"Type": "Runtime"
-		},
 		{
 			"Name": "MassGameplayTestSuite",
 			"Type": "UncookedOnly",
diff --git a/Engine/Plugins/Runtime/MassGameplay/Source/MassLWI/MassLWI.Build.cs b/Engine/Plugins/Runtime/MassGameplay/Source/MassLWI/MassLWI.Build.cs
deleted file mode 100644
diff --git a/Engine/Plugins/Runtime/MassGameplay/Source/MassLWI/Private/MassLWIClientActorSpawnerSubsystem.cpp b/Engine/Plugins/Runtime/MassGameplay/Source/MassLWI/Private/MassLWIClientActorSpawnerSubsystem.cpp
deleted file mode 100644
diff --git a/Engine/Plugins/Runtime/MassGameplay/Source/MassLWI/Private/MassLWIConfigActor.cpp b/Engine/Plugins/Runtime/MassGameplay/Source/MassLWI/Private/MassLWIConfigActor.cpp
deleted file mode 100644
diff --git a/Engine/Plugins/Runtime/MassGameplay/Source/MassLWI/Private/MassLWIModule.cpp b/Engine/Plugins/Runtime/MassGameplay/Source/MassLWI/Private/MassLWIModule.cpp
deleted file mode 100644
diff --git a/Engine/Plugins/Runtime/MassGameplay/Source/MassLWI/Private/MassLWIRepresentationSubsystem.cpp b/Engine/Plugins/Runtime/MassGameplay/Source/MassLWI/Private/MassLWIRepresentationSubsystem.cpp
deleted file mode 100644
diff --git a/Engine/Plugins/Runtime/MassGameplay/Source/MassLWI/Private/MassLWISettings.cpp b/Engine/Plugins/Runtime/MassGameplay/Source/MassLWI/Private/MassLWISettings.cpp
deleted file mode 100644
diff --git a/Engine/Plugins/Runtime/MassGameplay/Source/MassLWI/Private/MassLWIStaticMeshManager.cpp b/Engine/Plugins/Runtime/MassGameplay/Source/MassLWI/Private/MassLWIStaticMeshManager.cpp
deleted file mode 100644
diff --git a/Engine/Plugins/Runtime/MassGameplay/Source/MassLWI/Private/MassLWISubsystem.cpp b/Engine/Plugins/Runtime/MassGameplay/Source/MassLWI/Private/MassLWISubsystem.cpp
deleted file mode 100644
diff --git a/Engine/Plugins/Runtime/MassGameplay/Source/MassLWI/Private/MassLWITypes.cpp b/Engine/Plugins/Runtime/MassGameplay/Source/MassLWI/Private/MassLWITypes.cpp
deleted file mode 100644
diff --git a/Engine/Plugins/Runtime/MassGameplay/Source/MassLWI/Private/MassLWIVisualizationTrait.cpp b/Engine/Plugins/Runtime/MassGameplay/Source/MassLWI/Private/MassLWIVisualizationTrait.cpp
deleted file mode 100644
diff --git a/Engine/Plugins/Runtime/MassGameplay/Source/MassLWI/Public/IMassLWIModule.h b/Engine/Plugins/Runtime/MassGameplay/Source/MassLWI/Public/IMassLWIModule.h
deleted file mode 100644
diff --git a/Engine/Plugins/Runtime/MassGameplay/Source/MassLWI/Public/MassLWIClientActorSpawnerSubsystem.h b/Engine/Plugins/Runtime/MassGameplay/Source/MassLWI/Public/MassLWIClientActorSpawnerSubsystem.h
deleted file mode 100644
diff --git a/Engine/Plugins/Runtime/MassGameplay/Source/MassLWI/Public/MassLWIConfigActor.h b/Engine/Plugins/Runtime/MassGameplay/Source/MassLWI/Public/MassLWIConfigActor.h
deleted file mode 100644
diff --git a/Engine/Plugins/Runtime/MassGameplay/Source/MassLWI/Public/MassLWIRepresentationSubsystem.h b/Engine/Plugins/Runtime/MassGameplay/Source/MassLWI/Public/MassLWIRepresentationSubsystem.h
deleted file mode 100644
diff --git a/Engine/Plugins/Runtime/MassGameplay/Source/MassLWI/Public/MassLWISettings.h b/Engine/Plugins/Runtime/MassGameplay/Source/MassLWI/Public/MassLWISettings.h
deleted file mode 100644
diff --git a/Engine/Plugins/Runtime/MassGameplay/Source/MassLWI/Public/MassLWIStaticMeshManager.h b/Engine/Plugins/Runtime/MassGameplay/Source/MassLWI/Public/MassLWIStaticMeshManager.h
deleted file mode 100644
diff --git a/Engine/Plugins/Runtime/MassGameplay/Source/MassLWI/Public/MassLWISubsystem.h b/Engine/Plugins/Runtime/MassGameplay/Source/MassLWI/Public/MassLWISubsystem.h
deleted file mode 100644
diff --git a/Engine/Plugins/Runtime/MassGameplay/Source/MassLWI/Public/MassLWITypes.h b/Engine/Plugins/Runtime/MassGameplay/Source/MassLWI/Public/MassLWITypes.h
deleted file mode 100644
diff --git a/Engine/Plugins/Runtime/MassGameplay/Source/MassLWI/Public/MassLWIVisualizationTrait.h b/Engine/Plugins/Runtime/MassGameplay/Source/MassLWI/Public/MassLWIVisualizationTrait.h
deleted file mode 100644
