[Backout] - CL27952862
Original CL Desc
-----------------------------------------------------------------
Removed MassLWI and all of project specific code that was using it.

diff --git a/Engine/Plugins/Runtime/MassGameplay/MassGameplay.uplugin b/Engine/Plugins/Runtime/MassGameplay/MassGameplay.uplugin
--- a/Engine/Plugins/Runtime/MassGameplay/MassGameplay.uplugin
+++ b/Engine/Plugins/Runtime/MassGameplay/MassGameplay.uplugin
@@ -73,6 +73,10 @@
 			"Name": "MassGameplayExternalTraits",
 			"Type": "Runtime"
 		},
+		{
+			"Name": "MassLWI",
+			"Type": "Runtime"
+		},
 		{
 			"Name": "MassGameplayTestSuite",
 			"Type": "UncookedOnly",
diff --git a/Engine/Plugins/Runtime/MassGameplay/Source/MassLWI/MassLWI.Build.cs b/Engine/Plugins/Runtime/MassGameplay/Source/MassLWI/MassLWI.Build.cs
new file mode 100644
--- /dev/null
+++ b/Engine/Plugins/Runtime/MassGameplay/Source/MassLWI/MassLWI.Build.cs
@@ -0,0 +1,35 @@
+// Copyright Epic Games, Inc. All Rights Reserved.
+
+namespace UnrealBuildTool.Rules
+{
+	public class MassLWI : ModuleRules
+	{
+		public MassLWI(ReadOnlyTargetRules Target) : base(Target)
+		{
+			PCHUsage = ModuleRules.PCHUsageMode.UseExplicitOrSharedPCHs;
+
+			PublicDependencyModuleNames.AddRange(
+				new string[] {
+					"Core",
+					"CoreUObject",
+					"Engine",
+					"MassEntity",
+					"MassCommon",
+					"MassReplication",
+					"MassRepresentation",
+					"MassSpawner",
+					"MassActors",
+					"StructUtils",
+					"AIModule",
+					"NetCore",
+				}
+			);
+
+			if (Target.bBuildEditor == true)
+			{
+				PrivateDependencyModuleNames.Add("EditorFramework");
+				PrivateDependencyModuleNames.Add("UnrealEd");
+			}
+		}
+	}
+}
diff --git a/Engine/Plugins/Runtime/MassGameplay/Source/MassLWI/Private/MassLWIClientActorSpawnerSubsystem.cpp b/Engine/Plugins/Runtime/MassGameplay/Source/MassLWI/Private/MassLWIClientActorSpawnerSubsystem.cpp
new file mode 100644
--- /dev/null
+++ b/Engine/Plugins/Runtime/MassGameplay/Source/MassLWI/Private/MassLWIClientActorSpawnerSubsystem.cpp
@@ -0,0 +1,63 @@
+// Copyright Epic Games, Inc. All Rights Reserved.
+
+#include "MassLWIClientActorSpawnerSubsystem.h"
+#include "Engine/World.h"
+
+
+//-----------------------------------------------------------------------------
+// UMassLWIClientActorSpawnerSubsystem
+//-----------------------------------------------------------------------------
+void UMassLWIClientActorSpawnerSubsystem::Initialize(FSubsystemCollectionBase& Collection)
+{
+	Super::Initialize(Collection);
+
+	UWorld* World = GetWorld();
+	check(World);
+	check(World->GetNetMode() == NM_Client);
+	
+	WorldOnActorSpawnedHandle = World->AddOnActorSpawnedHandler(FOnActorSpawned::FDelegate::CreateUObject(this, &UMassLWIClientActorSpawnerSubsystem::OnActorSpawned));
+}
+
+bool UMassLWIClientActorSpawnerSubsystem::ShouldCreateSubsystem(UObject* Outer) const
+{
+	UWorld* World = Cast<UWorld>(Outer);
+	return (World != nullptr && World->GetNetMode() == NM_Client);
+}
+
+void UMassLWIClientActorSpawnerSubsystem::OnActorSpawned(AActor* InActor)
+{
+	check(InActor);
+	// check if it's an actor we're waiting for
+	if (FMassStoredActorsContainer* StoredActors = PendingActors.Find(InActor->GetClass()))
+	{
+		// getting here means we're waiting for this class of actor
+		StoredActors->Container.Add(InActor);
+	}
+}
+
+ESpawnRequestStatus UMassLWIClientActorSpawnerSubsystem::SpawnActor(FConstStructView SpawnRequestView, TObjectPtr<AActor>& OutSpawnedActor) const
+{
+	constexpr double ComparisonPrecission = 1.0;
+	const FMassActorSpawnRequest& SpawnRequest = SpawnRequestView.Get<const FMassActorSpawnRequest>();
+	TArray<TObjectPtr<AActor>>& StoredActors = PendingActors.FindOrAdd(SpawnRequest.Template).Container;
+
+	for (auto It = StoredActors.CreateIterator(); It; ++It)
+	{
+		if (*It && (*It)->GetActorTransform().Equals(SpawnRequest.Transform, ComparisonPrecission))
+		{
+			OutSpawnedActor = *It;
+			It.RemoveCurrentSwap();
+			return ESpawnRequestStatus::Succeeded;
+		}
+	}
+
+	return ESpawnRequestStatus::Pending;
+}
+
+void UMassLWIClientActorSpawnerSubsystem::RegisterRepresentedClass(UClass* RepresentedClass)
+{
+	if (TSubclassOf<AActor> ActorClass = RepresentedClass)
+	{
+		PendingActors.Add(MoveTemp(ActorClass));
+	}
+}
diff --git a/Engine/Plugins/Runtime/MassGameplay/Source/MassLWI/Private/MassLWIConfigActor.cpp b/Engine/Plugins/Runtime/MassGameplay/Source/MassLWI/Private/MassLWIConfigActor.cpp
new file mode 100644
--- /dev/null
+++ b/Engine/Plugins/Runtime/MassGameplay/Source/MassLWI/Private/MassLWIConfigActor.cpp
@@ -0,0 +1,17 @@
+// Copyright Epic Games, Inc. All Rights Reserved.
+
+#include "MassLWIConfigActor.h"
+#include "MassLWISubsystem.h"
+
+//-----------------------------------------------------------------------------
+// ALWIxMassConfigActor
+//-----------------------------------------------------------------------------
+void AMassLWIConfigActor::PostLoad()
+{
+	if (UMassLWISubsystem* MassLWI = UWorld::GetSubsystem<UMassLWISubsystem>(GetWorld()))
+	{
+		MassLWI->RegisterConfigActor(*this);
+	}
+	Super::PostLoad();
+}
+
diff --git a/Engine/Plugins/Runtime/MassGameplay/Source/MassLWI/Private/MassLWIModule.cpp b/Engine/Plugins/Runtime/MassGameplay/Source/MassLWI/Private/MassLWIModule.cpp
new file mode 100644
--- /dev/null
+++ b/Engine/Plugins/Runtime/MassGameplay/Source/MassLWI/Private/MassLWIModule.cpp
@@ -0,0 +1,33 @@
+// Copyright Epic Games, Inc. All Rights Reserved.
+
+#include "MassLWITypes.h"
+#include "IMassLWIModule.h"
+
+DEFINE_LOG_CATEGORY(LogMassLWI)
+
+
+class FMassLWIModule : public IMassLWIModule
+{
+	/** IModuleInterface implementation */
+	virtual void StartupModule() override;
+	virtual void ShutdownModule() override;
+};
+
+IMPLEMENT_MODULE(FMassLWIModule, MassLWI)
+
+
+
+void FMassLWIModule::StartupModule()
+{
+	// This code will execute after your module is loaded into memory (but after global variables are initialized, of course.)
+}
+
+
+void FMassLWIModule::ShutdownModule()
+{
+	// This function may be called during shutdown to clean up your module.  For modules that support dynamic reloading,
+	// we call this function before unloading the module.
+}
+
+
+
diff --git a/Engine/Plugins/Runtime/MassGameplay/Source/MassLWI/Private/MassLWIRepresentationSubsystem.cpp b/Engine/Plugins/Runtime/MassGameplay/Source/MassLWI/Private/MassLWIRepresentationSubsystem.cpp
new file mode 100644
--- /dev/null
+++ b/Engine/Plugins/Runtime/MassGameplay/Source/MassLWI/Private/MassLWIRepresentationSubsystem.cpp
@@ -0,0 +1,23 @@
+// Copyright Epic Games, Inc. All Rights Reserved.
+
+#include "MassLWIRepresentationSubsystem.h"
+#include "Engine/World.h"
+#include "MassLWIClientActorSpawnerSubsystem.h"
+
+
+//-----------------------------------------------------------------------------
+// UMassLWIRepresentationSubsystem
+//-----------------------------------------------------------------------------
+void UMassLWIRepresentationSubsystem::Initialize(FSubsystemCollectionBase& Collection)
+{
+	Super::Initialize(Collection);
+	
+	UWorld* World = GetWorld();
+	check(World);
+
+	if (World->GetNetMode() == NM_Client)
+	{
+		ActorSpawnerSubsystem = World->GetSubsystem<UMassLWIClientActorSpawnerSubsystem>();
+	}
+}
+
diff --git a/Engine/Plugins/Runtime/MassGameplay/Source/MassLWI/Private/MassLWISettings.cpp b/Engine/Plugins/Runtime/MassGameplay/Source/MassLWI/Private/MassLWISettings.cpp
new file mode 100644
--- /dev/null
+++ b/Engine/Plugins/Runtime/MassGameplay/Source/MassLWI/Private/MassLWISettings.cpp
@@ -0,0 +1,17 @@
+// Copyright Epic Games, Inc. All Rights Reserved.
+
+#include "MassLWISettings.h"
+#include "HAL/IConsoleManager.h"
+
+
+namespace UE::Mass::LWI::Private
+{
+	bool bMassLWIForceEnable = false;
+
+	FAutoConsoleVariableRef CVarMassLWIEnable(TEXT("mass.lwi.ForceEnable"), bMassLWIForceEnable, TEXT("Use this to temporarily enable MassLWI without affecting the project wide UMassLWISettings::bMassLWIEnabled setting"));
+} // UE::MassLOD::Debug
+
+bool UMassLWISettings::IsMassLWIEnabled() const
+{
+	return bMassLWIEnabled || UE::Mass::LWI::Private::bMassLWIForceEnable;
+}
diff --git a/Engine/Plugins/Runtime/MassGameplay/Source/MassLWI/Private/MassLWIStaticMeshManager.cpp b/Engine/Plugins/Runtime/MassGameplay/Source/MassLWI/Private/MassLWIStaticMeshManager.cpp
new file mode 100644
--- /dev/null
+++ b/Engine/Plugins/Runtime/MassGameplay/Source/MassLWI/Private/MassLWIStaticMeshManager.cpp
@@ -0,0 +1,318 @@
+// Copyright Epic Games, Inc. All Rights Reserved.
+
+#include "MassLWIStaticMeshManager.h"
+#include "MassLWITypes.h"
+#include "MassLWISubsystem.h"
+#include "MassLWITypes.h"
+#include "Components/InstancedStaticMeshComponent.h"
+#include "MassCommonFragments.h"
+#include "MassEntitySubsystem.h"
+#include "MassEntityView.h"
+#include "MassEntityQuery.h"
+#include "MassRepresentationSubsystem.h"
+#include "MassSpawnerSubsystem.h"
+#include "MassSpawnerTypes.h"
+#include "MassSpawnLocationProcessor.h"
+#include "MassVisualizationTrait.h"
+#include "VisualLogger/VisualLogger.h"
+#include "MassDebugger.h"
+
+
+namespace UE::Mass::Tweakables
+{
+	bool bDestroyEntitiesOnEndPlay = true;
+	bool bReuseLWIISMComponents = true;
+
+	FAutoConsoleVariableRef CLWIVars[] = {
+		{TEXT("mass.LWI.DestroyEntitiesOnEndPlay"), bDestroyEntitiesOnEndPlay, TEXT("Whether we should destroy LWI-sources entities when the original LWI manager ends play")},
+		{TEXT("mass.LWI.ReuseManagersISMComponents"), bReuseLWIISMComponents, TEXT("Whether we should reuse the LWIManager-owned ISM components for Mass visualization")},
+	};
+}
+
+//-----------------------------------------------------------------------------
+// AMassLWIStaticMeshManager
+//-----------------------------------------------------------------------------
+void AMassLWIStaticMeshManager::BeginPlay()
+{
+	Super::BeginPlay();
+
+	UWorld* World = GetWorld();
+	if (World != nullptr && IsRegisteredWithMass() == false)
+	{
+		if (UMassLWISubsystem* LWIxMass = World->GetSubsystem<UMassLWISubsystem>())
+		{
+			LWIxMass->RegisterLWIManager(*this);
+		}
+	}
+}
+
+void AMassLWIStaticMeshManager::EndPlay(const EEndPlayReason::Type EndPlayReason)
+{
+	UWorld* World = GetWorld();
+	if (World && IsRegisteredWithMass())
+	{
+		if (UE::Mass::Tweakables::bDestroyEntitiesOnEndPlay && Entities.Num())
+		{
+			check(InstanceTransforms.Num() == 0);
+			InstanceTransforms.Reset(Entities.Num());
+
+			if (UMassEntitySubsystem* EntitySubsystem = UWorld::GetSubsystem<UMassEntitySubsystem>(GetWorld()))
+			{
+				FMassEntityManager& EntityManager = EntitySubsystem->GetMutableEntityManager();
+				StoreMassDataInActor(EntityManager);
+			}
+		}
+
+		if (UMassLWISubsystem* LWIxMass = World->GetSubsystem<UMassLWISubsystem>())
+		{
+			LWIxMass->UnregisterLWIManager(*this);
+		}
+	}
+
+	Super::EndPlay(EndPlayReason);
+}
+
+void AMassLWIStaticMeshManager::TransferDataToMass(FMassEntityManager& EntityManager)
+{
+	if (InstanceTransforms.Num() == 0)
+	{
+		return;
+	}
+
+	if (MassTemplateID.IsValid() == false && RepresentedClass)
+	{
+		CreateMassTemplate(EntityManager);
+	}
+
+	if (MassTemplateID.IsValid() && ISMComponent && FinalizedTemplate)
+	{
+		const FVector& ManagerLocation = GetActorLocation();
+
+		UMassSpawnerSubsystem* SpawnerSystem = UWorld::GetSubsystem<UMassSpawnerSubsystem>(GetWorld());
+		check(SpawnerSystem);
+		const int32 NumEntities = InstanceTransforms.Num();
+		
+		FMassTransformsSpawnData SpawnLocationData;
+		SpawnLocationData.bRandomize = false;
+		SpawnLocationData.Transforms = InstanceTransforms;
+		if (ManagerLocation.IsNearlyZero() == false)
+		{
+			for (FTransform& EntityTransform : SpawnLocationData.Transforms)
+			{
+				EntityTransform.AddToTranslation(ManagerLocation);
+			}
+		}
+
+		FConstStructView SpawnLocationDataView = FConstStructView::Make(SpawnLocationData);
+		SpawnerSystem->SpawnEntities(FinalizedTemplate->GetTemplateID(), NumEntities, SpawnLocationDataView, UMassSpawnLocationProcessor::StaticClass(), Entities);
+
+		// destroy actors that have already been created
+		for (const TPair<int32, TObjectPtr<AActor>>& ActorPair : Actors)
+		{
+			if (ActorPair.Value)
+			{
+				ActorPair.Value->OnDestroyed.RemoveAll(this);
+				ActorPair.Value->Destroy();
+			}
+		}
+		Actors.Reset();
+
+		InstanceTransforms.Reset();
+		ValidIndices.Reset();
+		FreeIndices.Reset();
+		RenderingIndicesToDataIndices.SetNum(Entities.Num());
+		DataIndicesToRenderingIndices.SetNum(Entities.Num());
+		for (int32 Index = 0; Index < Entities.Num(); ++Index)
+		{
+			RenderingIndicesToDataIndices[Index] = Index;
+			DataIndicesToRenderingIndices[Index] = Index;
+		}
+		
+		ISMComponent->ClearInstances();
+	}
+}
+
+void AMassLWIStaticMeshManager::StoreMassDataInActor(FMassEntityManager& EntityManager)
+{
+	TArray<FMassArchetypeEntityCollection> EntityCollectionsToDestroy;
+	UE::Mass::Utils::CreateEntityCollections(EntityManager, Entities, FMassArchetypeEntityCollection::FoldDuplicates, EntityCollectionsToDestroy);
+
+	// we create a query to "recreate" contents of InstanceTransforms just in case BeginPlay will be get
+	// called again for this LWIManager, which can happen with actor streaming
+	FMassEntityQuery LocationQuery;
+	LocationQuery.AddRequirement<FTransformFragment>(EMassFragmentAccess::ReadOnly);
+
+	const FTransform WorldTransform = ISMComponent ? ISMComponent->GetComponentTransform() : GetActorTransform();
+	
+	FMassExecutionContext ExecutionContext(EntityManager);
+	for (FMassArchetypeEntityCollection& Collection : EntityCollectionsToDestroy)
+	{
+		LocationQuery.ForEachEntityChunk(Collection, EntityManager, ExecutionContext, [this, WorldTransform](FMassExecutionContext& Context)
+		{
+			TConstArrayView<FTransformFragment> TransformsList = Context.GetFragmentView<FTransformFragment>();
+			for (const FTransformFragment& Fragment : TransformsList)
+			{
+				const FTransform LocalTransform = Fragment.GetTransform().GetRelativeTransform(WorldTransform);
+				InstanceTransforms.Add(LocalTransform);
+			}
+		});
+
+		EntityManager.BatchDestroyEntityChunks(Collection);
+	}
+
+	Entities.Reset();
+}
+
+int32 AMassLWIStaticMeshManager::FindIndexForEntity(const FMassEntityHandle Entity) const
+{
+	return Entities.IndexOfByKey(Entity);
+}
+
+AMassLWIStaticMeshManager* AMassLWIStaticMeshManager::GetMassLWIManagerForEntity(const FMassEntityManager& EntityManager, const FMassEntityHandle Entity)
+{
+	FMassEntityView EntityView(EntityManager, Entity);
+	if (EntityView.IsSet())
+	{	
+		const FMassLWIManagerSharedFragment* LWIManagerSharedFragmentPtr = EntityView.GetSharedFragmentDataPtr<FMassLWIManagerSharedFragment>();
+		if (ensure(LWIManagerSharedFragmentPtr))
+		{
+			AMassLWIStaticMeshManager* LWIManager = LWIManagerSharedFragmentPtr->LWIManager.Get();
+			return LWIManager;
+		}
+	}
+
+	return nullptr;
+}
+
+void AMassLWIStaticMeshManager::CreateMassTemplate(FMassEntityManager& EntityManager)
+{
+	UWorld* World = GetWorld();
+	if (World == nullptr)
+	{
+		return;
+	}
+
+	UMassSpawnerSubsystem* SpawnerSystem = World->GetSubsystem<UMassSpawnerSubsystem>(World);
+	if (!ensure(SpawnerSystem))
+	{
+		//+complain
+		return;
+	}
+
+	UMassLWISubsystem* LWIxMass = World->GetSubsystem<UMassLWISubsystem>();
+	if (!ensure(LWIxMass))
+	{
+		return;
+	}
+
+	const FMassEntityConfig* ClassConfig = LWIxMass->GetConfigForClass(RepresentedClass);
+	if (!ClassConfig || ClassConfig->IsEmpty())
+	{
+		UE_VLOG_UELOG(this, LogMass, Log, TEXT("%s failed to find a calid entity config for %s class. This LWI manager won't transfer data to Mass.")
+			, *UObjectBaseUtility::GetName(), *GetNameSafe(RepresentedClass));
+		return;
+	}
+
+	const UMassVisualizationTrait* VisTrait = Cast<const UMassVisualizationTrait>(ClassConfig->FindTrait(UMassVisualizationTrait::StaticClass()));
+	if (!ensureMsgf(VisTrait, TEXT("The config used doesn't contain a VisualizationTrait, which is required for LWIxMass to function")))
+	{
+		return;
+	}
+
+	UMassRepresentationSubsystem* RepresentationSubsystem = nullptr;
+	
+	const FMassEntityTemplate& SourceTemplate = ClassConfig->GetOrCreateEntityTemplate(*World);
+	// what we want to do now is to modify the config to point at the specific static mesh and actor class
+	FMassEntityTemplateData NewTemplate(SourceTemplate);
+
+		
+	const FMassArchetypeSharedFragmentValues& SharedFragmentValues = SourceTemplate.GetSharedFragmentValues();
+	const TArray<FSharedStruct>& SharedFragments = SharedFragmentValues.GetSharedFragments();
+	for (FSharedStruct SharedFragment : SharedFragments)
+	{
+		if (FMassRepresentationSubsystemSharedFragment* AsRepresentationSubsystemSharedFragment = SharedFragment.GetPtr<FMassRepresentationSubsystemSharedFragment>())
+		{
+			RepresentationSubsystem = AsRepresentationSubsystemSharedFragment->RepresentationSubsystem;
+			break;
+		}
+	}
+	
+	FStaticMeshInstanceVisualizationDesc StaticMeshInstanceDesc = VisTrait->StaticMeshInstanceDesc;
+	
+	// we don't care about the mesh that has been set in the source template, we're overriding it anyway.
+	StaticMeshInstanceDesc.Meshes.SetNumZeroed(1);
+	FMassStaticMeshInstanceVisualizationMeshDesc& MeshDesc = StaticMeshInstanceDesc.Meshes[0];
+
+	if (!MeshDesc.ISMComponentClass)
+	{
+		MeshDesc.ISMComponentClass = UInstancedStaticMeshComponent::StaticClass();
+	}
+
+	// forcing the "full range" since we only ever expect there to be one mesh for the relevant ISM component
+	MeshDesc.SetSignificanceRange(EMassLOD::High, EMassLOD::Max);
+	
+	MeshDesc.Mesh = StaticMesh.Get();
+
+	if (!ensure(MeshDesc.Mesh))
+	{
+		MeshDesc.Mesh = StaticMesh.LoadSynchronous();
+		if (!ensure(MeshDesc.Mesh))
+		{
+			UE_VLOG_UELOG(this, LogMassLWI, Error, TEXT("%s: Unable to load mesh asset %s"), *UObject::GetName(), *StaticMesh.GetLongPackageName());
+			return;
+		}
+	}
+
+	if (ISMComponent)
+	{
+		MeshDesc.MaterialOverrides = ISMComponent->OverrideMaterials;
+		MeshDesc.bCastShadows = (ISMComponent->CastShadow != 0);
+	}
+
+	check(MeshDesc.Mesh);
+	NewTemplate.SetTemplateName(MeshDesc.Mesh->GetName());
+
+	if (RepresentationSubsystem)
+	{
+		FMassRepresentationFragment& RepresentationFragment = NewTemplate.AddFragment_GetRef<FMassRepresentationFragment>();
+		if (UE::Mass::Tweakables::bReuseLWIISMComponents)
+		{
+			RepresentationFragment.StaticMeshDescIndex = RepresentationSubsystem->AddVisualDescWithISMComponent(StaticMeshInstanceDesc, *ISMComponent);
+		}
+		else
+		{
+			RepresentationFragment.StaticMeshDescIndex = RepresentationSubsystem->FindOrAddStaticMeshDesc(StaticMeshInstanceDesc);
+		}
+
+		const int32 TemplateActorIndex = RepresentationSubsystem->FindOrAddTemplateActor(RepresentedClass);
+		RepresentationFragment.HighResTemplateActorIndex = TemplateActorIndex;
+		RepresentationFragment.LowResTemplateActorIndex = TemplateActorIndex;
+	}
+
+	FMassLWIManagerSharedFragment LWIManagerSharedFragment;
+	LWIManagerSharedFragment.LWIManager = this;
+	const uint32 LWIManagerHash = UE::StructUtils::GetStructCrc32(FConstStructView::Make(LWIManagerSharedFragment));
+	FSharedStruct LWIManagerFragment = EntityManager.GetOrCreateSharedFragmentByHash<FMassLWIManagerSharedFragment>(LWIManagerHash, LWIManagerSharedFragment);
+	NewTemplate.AddSharedFragment(LWIManagerFragment);
+
+	FMassEntityTemplateRegistry& TemplateRegistry = SpawnerSystem->GetMutableTemplateRegistryInstance();
+	const FString ObjectPath = GetPathName();
+	const uint32 FlavorHash = HashCombine(GetTypeHash(GetNameSafe(RepresentedClass)), GetTypeHash(ObjectPath));
+	MassTemplateID = FMassEntityTemplateIDFactory::MakeFlavor(SourceTemplate.GetTemplateID(), FlavorHash);
+
+	const TSharedRef<FMassEntityTemplate>& NewFinalizedTemplate = TemplateRegistry.FindOrAddTemplate(MassTemplateID, MoveTemp(NewTemplate));
+	MassTemplateID = NewFinalizedTemplate->GetTemplateID();
+	FinalizedTemplate = NewFinalizedTemplate;
+
+}
+
+void AMassLWIStaticMeshManager::MarkRegisteredWithMass(const FMassLWIManagerRegistrationHandle RegistrationIndex)
+{
+	check(MassRegistrationHandle.IsValid() == false && RegistrationIndex.IsValid() == true);
+	new (&MassRegistrationHandle) FMassLWIManagerRegistrationHandle(RegistrationIndex);
+}
+
+void AMassLWIStaticMeshManager::MarkUnregisteredWithMass()
+{
+	new (&MassRegistrationHandle) FMassLWIManagerRegistrationHandle();
+}
diff --git a/Engine/Plugins/Runtime/MassGameplay/Source/MassLWI/Private/MassLWISubsystem.cpp b/Engine/Plugins/Runtime/MassGameplay/Source/MassLWI/Private/MassLWISubsystem.cpp
new file mode 100644
--- /dev/null
+++ b/Engine/Plugins/Runtime/MassGameplay/Source/MassLWI/Private/MassLWISubsystem.cpp
@@ -0,0 +1,143 @@
+// Copyright Epic Games, Inc. All Rights Reserved.
+
+#include "MassLWISubsystem.h"
+#include "EngineUtils.h"
+#include "MassLWITypes.h"
+#include "MassLWIStaticMeshManager.h"
+#include "MassLWIConfigActor.h"
+#include "MassEntitySubsystem.h"
+#include "MassLWIClientActorSpawnerSubsystem.h"
+#include "VisualLogger/VisualLogger.h"
+#include "MassLWISettings.h"
+
+
+#include UE_INLINE_GENERATED_CPP_BY_NAME(MassLWISubsystem)
+
+//-----------------------------------------------------------------------------
+// UMassLWISubsystem
+//-----------------------------------------------------------------------------
+UMassLWISubsystem::UMassLWISubsystem()
+{
+
+}
+
+void UMassLWISubsystem::Initialize(FSubsystemCollectionBase& Collection)
+{
+	Super::Initialize(Collection);
+
+	UMassEntitySubsystem* EntitySubsystem = Collection.InitializeDependency<UMassEntitySubsystem>();
+	check(EntitySubsystem);
+	EntityManager = EntitySubsystem->GetMutableEntityManager().AsShared();
+
+	// note that it's ok for LWISpawnerSubsystem to be null since in principle it's only available on clients,
+	// while UMassLWISubsystem can be created on both.
+	LWISpawnerSubsystem = Collection.InitializeDependency<UMassLWIClientActorSpawnerSubsystem>();
+}
+
+void UMassLWISubsystem::PostInitialize()
+{
+	Super::PostInitialize();
+
+	RegisterExistingConfigActors();
+}
+
+void UMassLWISubsystem::Deinitialize()
+{
+	EntityManager = nullptr;
+	Super::Deinitialize();
+}
+
+bool UMassLWISubsystem::ShouldCreateSubsystem(UObject* Outer) const
+{
+	UWorld* World = Outer ? Outer->GetWorld() : nullptr;
+	return World && World->IsGameWorld() && Super::ShouldCreateSubsystem(Outer) && GET_MASSLWI_CONFIG_VALUE(IsMassLWIEnabled());
+}
+
+void UMassLWISubsystem::RegisterLWIManager(AMassLWIStaticMeshManager& Manager)
+{
+	check(EntityManager);
+	if (Manager.IsRegisteredWithMass())
+	{
+		return;
+	}
+	checkSlow(RegisteredManagers.Find(&Manager) == INDEX_NONE);
+
+	REDIRECT_OBJECT_TO_VLOG(&Manager, this);
+	
+	int32 RegistrationIndex = INDEX_NONE;
+	if (FreeIndices.Num() > 0)
+	{
+		RegistrationIndex = FreeIndices.Pop(/*bAllowShrinking = */true);
+		RegisteredManagers[RegistrationIndex] = &Manager;
+	}
+	else
+	{
+		RegistrationIndex = RegisteredManagers.Add(&Manager);
+	}
+
+	if (LWISpawnerSubsystem)
+	{
+		LWISpawnerSubsystem->RegisterRepresentedClass(Manager.GetRepresentedClass());
+	}
+
+	Manager.TransferDataToMass(*EntityManager);
+	Manager.MarkRegisteredWithMass(FMassLWIManagerRegistrationHandle(RegistrationIndex));
+}
+
+void UMassLWISubsystem::UnregisterLWIManager(AMassLWIStaticMeshManager& Manager)
+{
+	if (Manager.IsRegisteredWithMass() == false)
+	{
+		return;
+	}
+
+	FMassLWIManagerRegistrationHandle IndexHandle = Manager.GetMassRegistrationHandle();
+	if (RegisteredManagers.IsValidIndex(IndexHandle))
+	{
+		FreeIndices.Add(IndexHandle);
+		RegisteredManagers[IndexHandle] = nullptr;
+	}
+
+	Manager.MarkUnregisteredWithMass();
+}
+
+void UMassLWISubsystem::RegisterConfigActor(AMassLWIConfigActor& ConfigActor)
+{
+	if (ConfigActor.IsRegisteredWithSubsystem())
+	{
+		return;
+	}
+
+	if (DefaultConfig.IsEmpty() && !ConfigActor.DefaultConfig.IsEmpty())
+	{
+		DefaultConfig = ConfigActor.DefaultConfig;
+	}
+
+}
+
+void UMassLWISubsystem::RegisterExistingConfigActors()
+{
+	UWorld* World = GetWorld();
+	if (World == nullptr)
+	{
+		return;
+	}
+
+	for (TActorIterator<AMassLWIConfigActor> It(World); It; ++It)
+	{
+		RegisterConfigActor(**It);
+	}
+}
+
+const FMassEntityConfig* UMassLWISubsystem::GetConfigForClass(TSubclassOf<AActor> ActorClass) const
+{
+	if (const FMassEntityConfig* ConfigFound = ClassToConfigMap.Find(ActorClass))
+	{
+		return ConfigFound;
+	}
+	else if (DefaultConfig.IsEmpty() == false)
+	{
+		return &DefaultConfig;
+	}
+	return nullptr;
+}
diff --git a/Engine/Plugins/Runtime/MassGameplay/Source/MassLWI/Private/MassLWITypes.cpp b/Engine/Plugins/Runtime/MassGameplay/Source/MassLWI/Private/MassLWITypes.cpp
new file mode 100644
--- /dev/null
+++ b/Engine/Plugins/Runtime/MassGameplay/Source/MassLWI/Private/MassLWITypes.cpp
@@ -0,0 +1,3 @@
+// Copyright Epic Games, Inc. All Rights Reserved.
+
+#include "MassLWITypes.h"
diff --git a/Engine/Plugins/Runtime/MassGameplay/Source/MassLWI/Private/MassLWIVisualizationTrait.cpp b/Engine/Plugins/Runtime/MassGameplay/Source/MassLWI/Private/MassLWIVisualizationTrait.cpp
new file mode 100644
--- /dev/null
+++ b/Engine/Plugins/Runtime/MassGameplay/Source/MassLWI/Private/MassLWIVisualizationTrait.cpp
@@ -0,0 +1,15 @@
+// Copyright Epic Games, Inc. All Rights Reserved.
+
+#include "MassLWIVisualizationTrait.h"
+#include "MassEntityTemplateRegistry.h"
+#include "MassLWIRepresentationSubsystem.h"
+
+
+UMassLWIVisualizationTrait::UMassLWIVisualizationTrait(const FObjectInitializer& ObjectInitializer)
+{
+	bAllowServerSideVisualization = true;
+#if WITH_EDITORONLY_DATA
+	bCanModifyRepresentationSubsystemClass = false;
+	RepresentationSubsystemClass = UMassLWIRepresentationSubsystem::StaticClass();
+#endif // WITH_EDITORONLY_DATA
+}
diff --git a/Engine/Plugins/Runtime/MassGameplay/Source/MassLWI/Public/IMassLWIModule.h b/Engine/Plugins/Runtime/MassGameplay/Source/MassLWI/Public/IMassLWIModule.h
new file mode 100644
--- /dev/null
+++ b/Engine/Plugins/Runtime/MassGameplay/Source/MassLWI/Public/IMassLWIModule.h
@@ -0,0 +1,39 @@
+// Copyright Epic Games, Inc. All Rights Reserved.
+
+#pragma once
+
+#include "Modules/ModuleInterface.h"
+#include "Modules/ModuleManager.h"
+
+
+/**
+ * The public interface to this module.  In most cases, this interface is only public to sibling modules 
+ * within this plugin.
+ */
+class IMassLWIModule : public IModuleInterface
+{
+
+public:
+
+	/**
+	 * Singleton-like access to this module's interface.  This is just for convenience!
+	 * Beware of calling this during the shutdown phase, though.  Your module might have been unloaded already.
+	 *
+	 * @return Returns singleton instance, loading the module on demand if needed
+	 */
+	static inline IMassLWIModule& Get()
+	{
+		return FModuleManager::LoadModuleChecked<IMassLWIModule>("MassLWI");
+	}
+
+	/**
+	 * Checks to see if this module is loaded and ready.  It is only valid to call Get() if IsAvailable() returns true.
+	 *
+	 * @return True if the module is loaded and ready to use
+	 */
+	static inline bool IsAvailable()
+	{
+		return FModuleManager::Get().IsModuleLoaded("MassLWI");
+	}
+};
+
diff --git a/Engine/Plugins/Runtime/MassGameplay/Source/MassLWI/Public/MassLWIClientActorSpawnerSubsystem.h b/Engine/Plugins/Runtime/MassGameplay/Source/MassLWI/Public/MassLWIClientActorSpawnerSubsystem.h
new file mode 100644
--- /dev/null
+++ b/Engine/Plugins/Runtime/MassGameplay/Source/MassLWI/Public/MassLWIClientActorSpawnerSubsystem.h
@@ -0,0 +1,53 @@
+// Copyright Epic Games, Inc. All Rights Reserved.
+
+#pragma once
+
+#include "MassExternalSubsystemTraits.h"
+#include "MassActorSpawnerSubsystem.h"
+#include "MassLWIClientActorSpawnerSubsystem.generated.h"
+
+
+USTRUCT()
+struct MASSLWI_API FMassStoredActorsContainer
+{
+	GENERATED_BODY()
+
+	UPROPERTY()
+	TArray<TObjectPtr<AActor>> Container;
+};
+
+UCLASS()
+class MASSLWI_API UMassLWIClientActorSpawnerSubsystem : public UMassActorSpawnerSubsystem
+{
+	GENERATED_BODY()
+
+public:
+	/** checks if RepresentedClass represents an actor class and if so creates an entry in PendingActors for the type */
+	void RegisterRepresentedClass(UClass* RepresentedClass);
+
+protected:
+	// USubsystem API begin
+	virtual bool ShouldCreateSubsystem(UObject* Outer) const override;
+	virtual void Initialize(FSubsystemCollectionBase& Collection) override;
+	// USubsystem API end
+
+	// UMassActorSpawnerSubsystem API begin
+	virtual ESpawnRequestStatus SpawnActor(FConstStructView SpawnRequestView, TObjectPtr<AActor>& OutSpawnedActor) const override;
+	// UMassActorSpawnerSubsystem API end
+
+	void OnActorSpawned(AActor* InActor);
+
+	FDelegateHandle WorldOnActorSpawnedHandle;
+
+	UPROPERTY()
+	mutable TMap<TSubclassOf<AActor>, FMassStoredActorsContainer> PendingActors;
+};
+
+template<>
+struct TMassExternalSubsystemTraits<UMassLWIClientActorSpawnerSubsystem> final
+{
+	enum
+	{
+		GameThreadOnly = true
+	};
+};
\ No newline at end of file
diff --git a/Engine/Plugins/Runtime/MassGameplay/Source/MassLWI/Public/MassLWIConfigActor.h b/Engine/Plugins/Runtime/MassGameplay/Source/MassLWI/Public/MassLWIConfigActor.h
new file mode 100644
--- /dev/null
+++ b/Engine/Plugins/Runtime/MassGameplay/Source/MassLWI/Public/MassLWIConfigActor.h
@@ -0,0 +1,31 @@
+// Copyright Epic Games, Inc. All Rights Reserved.
+
+#pragma once
+
+#include "GameFramework/LightWeightInstanceStaticMeshManager.h"
+#include "MassEntityConfigAsset.h"
+#include "MassLWIConfigActor.generated.h"
+
+
+class UMassLWISubsystem;
+
+UCLASS(hidecategories = (Actor, Input, Collision, Rendering, Replication, Partition, HLOD, Cooking))
+class AMassLWIConfigActor : public AActor
+{
+	GENERATED_BODY()
+public:
+	virtual void PostLoad() override;
+
+	bool IsRegisteredWithSubsystem() const { return bRegisteredWithSubsystem; }
+
+protected:
+	UPROPERTY(Category = "Derived Traits", EditAnywhere)
+	FMassEntityConfig DefaultConfig;
+
+	UPROPERTY(Category = "Derived Traits", EditAnywhere)
+	TMap<TObjectPtr<UClass>, FMassEntityConfig> PerClassConfigs;
+
+	uint8 bRegisteredWithSubsystem : 1;
+
+	friend UMassLWISubsystem;
+};
diff --git a/Engine/Plugins/Runtime/MassGameplay/Source/MassLWI/Public/MassLWIRepresentationSubsystem.h b/Engine/Plugins/Runtime/MassGameplay/Source/MassLWI/Public/MassLWIRepresentationSubsystem.h
new file mode 100644
--- /dev/null
+++ b/Engine/Plugins/Runtime/MassGameplay/Source/MassLWI/Public/MassLWIRepresentationSubsystem.h
@@ -0,0 +1,28 @@
+// Copyright Epic Games, Inc. All Rights Reserved.
+
+#pragma once
+
+#include "MassRepresentationSubsystem.h"
+#include "MassExternalSubsystemTraits.h"
+#include "MassLWIRepresentationSubsystem.generated.h"
+
+
+UCLASS()
+class MASSLWI_API UMassLWIRepresentationSubsystem : public UMassRepresentationSubsystem
+{
+	GENERATED_BODY()
+
+protected:
+	// USubsystem API begin
+	virtual void Initialize(FSubsystemCollectionBase& Collection) override;
+	// USubsystem API end
+};
+
+template<>
+struct TMassExternalSubsystemTraits<UMassLWIRepresentationSubsystem> final
+{
+	enum
+	{
+		GameThreadOnly = true
+	};
+};
diff --git a/Engine/Plugins/Runtime/MassGameplay/Source/MassLWI/Public/MassLWISettings.h b/Engine/Plugins/Runtime/MassGameplay/Source/MassLWI/Public/MassLWISettings.h
new file mode 100644
--- /dev/null
+++ b/Engine/Plugins/Runtime/MassGameplay/Source/MassLWI/Public/MassLWISettings.h
@@ -0,0 +1,25 @@
+// Copyright Epic Games, Inc. All Rights Reserved.
+
+#pragma once
+
+#include "MassSettings.h"
+#include "MassLWISettings.generated.h"
+
+
+#define GET_MASSLWI_CONFIG_VALUE(a) (GetMutableDefault<UMassLWISettings>()->a)
+
+/**
+ * Implements the settings for MassLWI plugin
+ */
+UCLASS(config = Mass, defaultconfig, DisplayName = "Mass LWI")
+class MASSLWI_API UMassLWISettings : public UMassModuleSettings
+{
+	GENERATED_BODY()
+
+public:
+	bool IsMassLWIEnabled() const;
+
+protected:
+	UPROPERTY(EditDefaultsOnly, Category = Mass, config)
+	bool bMassLWIEnabled = true;
+};
diff --git a/Engine/Plugins/Runtime/MassGameplay/Source/MassLWI/Public/MassLWIStaticMeshManager.h b/Engine/Plugins/Runtime/MassGameplay/Source/MassLWI/Public/MassLWIStaticMeshManager.h
new file mode 100644
--- /dev/null
+++ b/Engine/Plugins/Runtime/MassGameplay/Source/MassLWI/Public/MassLWIStaticMeshManager.h
@@ -0,0 +1,57 @@
+// Copyright Epic Games, Inc. All Rights Reserved.
+
+#pragma once
+
+#include "GameFramework/LightWeightInstanceStaticMeshManager.h"
+#include "MassEntityConfigAsset.h"
+#include "MassEntityTemplate.h"
+#include "MassLWIStaticMeshManager.generated.h"
+
+
+struct MASSLWI_API FMassLWIManagerRegistrationHandle
+{
+	explicit FMassLWIManagerRegistrationHandle(const int32 InRegisteredIndex = INDEX_NONE)
+		: RegisteredIndex(InRegisteredIndex)
+	{}
+
+	operator int32() const { return RegisteredIndex; }
+	bool IsValid() const { return RegisteredIndex != INDEX_NONE; }
+
+private:
+	const int32 RegisteredIndex = INDEX_NONE;
+};
+
+UCLASS()
+class MASSLWI_API AMassLWIStaticMeshManager : public ALightWeightInstanceStaticMeshManager
+{
+	GENERATED_BODY()
+
+public:
+	bool IsRegisteredWithMass() const { return MassRegistrationHandle.IsValid(); }
+
+	virtual void TransferDataToMass(FMassEntityManager& EntityManager);
+	virtual void StoreMassDataInActor(FMassEntityManager& EntityManager);
+	FMassLWIManagerRegistrationHandle GetMassRegistrationHandle() const { return MassRegistrationHandle; }
+	void MarkRegisteredWithMass(const FMassLWIManagerRegistrationHandle RegistrationIndex);
+	void MarkUnregisteredWithMass();
+
+	// Returns the index of the light weight instance associated with InEntity if one exists; otherwise we return INDEX_NONE
+	int32 FindIndexForEntity(const FMassEntityHandle Entity) const;
+
+	static AMassLWIStaticMeshManager* GetMassLWIManagerForEntity(const FMassEntityManager& EntityManager, const FMassEntityHandle Entity);
+
+protected:
+	// AActor API
+	virtual void BeginPlay() override;
+	virtual void EndPlay(const EEndPlayReason::Type EndPlayReason) override;
+	// End AActor API
+
+	virtual void CreateMassTemplate(FMassEntityManager& EntityManager);
+
+
+	FMassLWIManagerRegistrationHandle MassRegistrationHandle;
+	FMassEntityTemplateID MassTemplateID;
+	TSharedPtr<FMassEntityTemplate> FinalizedTemplate;
+	FMassArchetypeHandle TargetArchetype;
+	TArray<FMassEntityHandle> Entities;
+};
diff --git a/Engine/Plugins/Runtime/MassGameplay/Source/MassLWI/Public/MassLWISubsystem.h b/Engine/Plugins/Runtime/MassGameplay/Source/MassLWI/Public/MassLWISubsystem.h
new file mode 100644
--- /dev/null
+++ b/Engine/Plugins/Runtime/MassGameplay/Source/MassLWI/Public/MassLWISubsystem.h
@@ -0,0 +1,68 @@
+// Copyright Epic Games, Inc. All Rights Reserved.
+
+#pragma once
+
+#include "MassEntityManager.h"
+#include "Subsystems/WorldSubsystem.h"
+#include "MassEntityConfigAsset.h"
+#include "MassRepresentationSubsystem.h"
+#include "MassLWISubsystem.generated.h"
+
+
+struct FActorInstanceHandle;
+class AMassLWIStaticMeshManager;
+class AMassLWIConfigActor;
+class UMassLWIClientActorSpawnerSubsystem;
+struct FMassEntityManager;
+
+
+UCLASS()
+class MASSLWI_API UMassLWISubsystem : public UWorldSubsystem
+{
+	GENERATED_BODY()
+
+	//~USubsystem interface
+	virtual void Initialize(FSubsystemCollectionBase& Collection) override;
+	virtual void PostInitialize() override;
+	virtual void Deinitialize() override;
+	virtual bool ShouldCreateSubsystem(UObject* Outer) const override;
+	//~End of USubsystem interface
+
+public:
+	UMassLWISubsystem();
+
+	void RegisterLWIManager(AMassLWIStaticMeshManager& Manager);
+	void UnregisterLWIManager(AMassLWIStaticMeshManager& Manager);
+
+	const FMassEntityConfig* GetConfigForClass(TSubclassOf<AActor> ActorClass) const;
+
+	void RegisterConfigActor(AMassLWIConfigActor& ConfigActor);
+
+protected:
+	void RegisterExistingConfigActors();
+
+	UPROPERTY()
+	TArray<TObjectPtr<AMassLWIStaticMeshManager>> RegisteredManagers;
+
+	TArray<int32> FreeIndices;
+
+	UPROPERTY()
+	FMassEntityConfig DefaultConfig;
+
+	TMap<UClass*, FMassEntityConfig> ClassToConfigMap;
+
+	UPROPERTY()
+	TObjectPtr<UMassLWIClientActorSpawnerSubsystem> LWISpawnerSubsystem;
+
+private:
+	TSharedPtr<FMassEntityManager> EntityManager;
+};
+
+template<>
+struct TMassExternalSubsystemTraits<UMassLWISubsystem> final
+{
+	enum
+	{
+		GameThreadOnly = true
+	};
+};
diff --git a/Engine/Plugins/Runtime/MassGameplay/Source/MassLWI/Public/MassLWITypes.h b/Engine/Plugins/Runtime/MassGameplay/Source/MassLWI/Public/MassLWITypes.h
new file mode 100644
--- /dev/null
+++ b/Engine/Plugins/Runtime/MassGameplay/Source/MassLWI/Public/MassLWITypes.h
@@ -0,0 +1,21 @@
+// Copyright Epic Games, Inc. All Rights Reserved.
+
+#pragma once
+
+#include "Logging/LogMacros.h"
+#include "MassEntityTypes.h"
+
+#include "MassLWITypes.generated.h"
+
+DECLARE_LOG_CATEGORY_EXTERN(LogMassLWI, Log, All)
+
+class AMassLWIStaticMeshManager;
+
+USTRUCT()
+struct FMassLWIManagerSharedFragment : public FMassSharedFragment
+{
+	GENERATED_BODY()
+
+	UPROPERTY(Transient)
+	TSoftObjectPtr<AMassLWIStaticMeshManager> LWIManager = nullptr;
+};
diff --git a/Engine/Plugins/Runtime/MassGameplay/Source/MassLWI/Public/MassLWIVisualizationTrait.h b/Engine/Plugins/Runtime/MassGameplay/Source/MassLWI/Public/MassLWIVisualizationTrait.h
new file mode 100644
--- /dev/null
+++ b/Engine/Plugins/Runtime/MassGameplay/Source/MassLWI/Public/MassLWIVisualizationTrait.h
@@ -0,0 +1,15 @@
+// Copyright Epic Games, Inc. All Rights Reserved.
+
+#pragma once 
+
+#include "MassStationaryVisualizationTrait.h"
+#include "MassLWIVisualizationTrait.generated.h"
+
+
+UCLASS()
+class MASSLWI_API UMassLWIVisualizationTrait : public UMassStationaryVisualizationTrait
+{
+	GENERATED_BODY()
+public:
+	UMassLWIVisualizationTrait(const FObjectInitializer& ObjectInitializer = FObjectInitializer::Get());
+};
