Deprecated UMassVisualizationComponent::RemoveISMComponent in favor of  RemoveVisualizationIndex that removes the whole VisualizationIndex at once making removal of multi-ISMComponent visualizations safe.

diff --git a/Engine/Plugins/Runtime/MassGameplay/Source/MassRepresentation/Private/MassRepresentationSubsystem.cpp b/Engine/Plugins/Runtime/MassGameplay/Source/MassRepresentation/Private/MassRepresentationSubsystem.cpp
--- a/Engine/Plugins/Runtime/MassGameplay/Source/MassRepresentation/Private/MassRepresentationSubsystem.cpp
+++ b/Engine/Plugins/Runtime/MassGameplay/Source/MassRepresentation/Private/MassRepresentationSubsystem.cpp
@@ -40,8 +40,16 @@ int16 UMassRepresentationSubsystem::AddVisualDescWithISMComponents(const FStatic
 
 void UMassRepresentationSubsystem::RemoveISMComponent(UInstancedStaticMeshComponent& ISMComponent)
 {
+PRAGMA_DISABLE_DEPRECATION_WARNINGS
 	check(VisualizationComponent);
 	return VisualizationComponent->RemoveISMComponent(ISMComponent);
+PRAGMA_ENABLE_DEPRECATION_WARNINGS
+}
+
+void UMassRepresentationSubsystem::RemoveVisualDescByIndex(const int32 VisualizationIndex)
+{
+	check(VisualizationComponent);
+	return VisualizationComponent->RemoveVisualDescByIndex(VisualizationIndex);
 }
 
 FMassInstancedStaticMeshInfoArrayView UMassRepresentationSubsystem::GetMutableInstancedStaticMeshInfos()
diff --git a/Engine/Plugins/Runtime/MassGameplay/Source/MassRepresentation/Private/MassVisualizationComponent.cpp b/Engine/Plugins/Runtime/MassGameplay/Source/MassRepresentation/Private/MassVisualizationComponent.cpp
--- a/Engine/Plugins/Runtime/MassGameplay/Source/MassRepresentation/Private/MassVisualizationComponent.cpp
+++ b/Engine/Plugins/Runtime/MassGameplay/Source/MassRepresentation/Private/MassVisualizationComponent.cpp
@@ -138,16 +138,24 @@ int16 UMassVisualizationComponent::AddVisualDescWithISMComponents(const FStaticM
 	return (int16)VisualIndex;
 }
 
-void UMassVisualizationComponent::RemoveISMComponent(UInstancedStaticMeshComponent& ISMComponent)
+void UMassVisualizationComponent::RemoveVisualDescByIndex(const int32 VisualizationIndex)
 {
-	const uint32 ISMComponentPathHash = GetTypeHash(ISMComponent.GetPathName());
-	const int32* VisualIndexPtr = ISMComponentMap.Find(ISMComponentPathHash);
-	if (ensure(VisualIndexPtr))
-	{
-		ISMCSharedData.Remove(ISMComponentPathHash);
+	UE_MT_SCOPED_WRITE_ACCESS(InstancedStaticMeshInfosDetector);
 
-		InstancedStaticMeshInfos[*VisualIndexPtr].Reset();
-		InstancedStaticMeshInfosFreeIndices.Add(*VisualIndexPtr);
+	if (ensure(InstancedStaticMeshInfos.IsValidIndex(VisualizationIndex))
+		&& ensureMsgf(InstancedStaticMeshInfos[VisualizationIndex].IsValid(), TEXT("Trying to remove visualization data that has already been cleaned")))
+	{
+		for (TObjectPtr<UInstancedStaticMeshComponent>& ISMComponent : InstancedStaticMeshInfos[VisualizationIndex].InstancedStaticMeshComponents)
+		{
+			const uint32 ISMComponentPathHash = GetTypeHash(ISMComponent.GetPathName());
+			const int32 StoredVisualizationIndex = ISMComponentMap.FindAndRemoveChecked(ISMComponentPathHash);
+			ensure(StoredVisualizationIndex == VisualizationIndex);
+		
+			ISMCSharedData.Remove(ISMComponentPathHash);
+		}
+		
+		InstancedStaticMeshInfos[VisualizationIndex].Reset();
+		InstancedStaticMeshInfosFreeIndices.Add(VisualizationIndex);
 	}
 }
 
@@ -778,4 +786,17 @@ void FMassLODSignificanceRange::WriteCustomDataFloatsAtStartIndex(int32 StaticMe
 void UMassVisualizationComponent::BuildLODSignificanceForInfo(FMassInstancedStaticMeshInfo& Info, const uint32 ForcedStaticMeshRefKeys)
 {
 	BuildLODSignificanceForInfo(Info, MakeArrayView(&ForcedStaticMeshRefKeys, 1));
-}
\ No newline at end of file
+}
+
+void UMassVisualizationComponent::RemoveISMComponent(UInstancedStaticMeshComponent& ISMComponent)
+{
+	UE_MT_SCOPED_WRITE_ACCESS(InstancedStaticMeshInfosDetector);
+
+	const uint32 ISMComponentPathHash = GetTypeHash(ISMComponent.GetPathName());
+	const int32* VisualIndexPtr = ISMComponentMap.Find(ISMComponentPathHash);
+	if (VisualIndexPtr)
+	{
+		RemoveVisualDescByIndex(*VisualIndexPtr);
+	}
+}
+
diff --git a/Engine/Plugins/Runtime/MassGameplay/Source/MassRepresentation/Public/MassRepresentationSubsystem.h b/Engine/Plugins/Runtime/MassGameplay/Source/MassRepresentation/Public/MassRepresentationSubsystem.h
--- a/Engine/Plugins/Runtime/MassGameplay/Source/MassRepresentation/Public/MassRepresentationSubsystem.h
+++ b/Engine/Plugins/Runtime/MassGameplay/Source/MassRepresentation/Public/MassRepresentationSubsystem.h
@@ -53,9 +53,18 @@ public:
 	/**
 	 * Removes the visualization data associated with the given ISM component. Note that this is safe to do only when
 	 * there are no entities relying on this data. No entity data patching will take place.
+	 * Note that the function will assert if there's more ISM components associated with given visualization. Also, in 
+	 * that case RemoveVisualDescByIndex will be called under the hood. 
 	 */
+	UE_DEPRECATED(5.4, "RemoveISMComponent has been deprecated in favor of RemoveVisualDescByIndex. Please use that instead.")
 	void RemoveISMComponent(UInstancedStaticMeshComponent& ISMComponent);
 
+	/** 
+	 * Removes all data associated with a given VisualizationIndex. Note that this is safe to do only if there are no
+	 * entities relying on this index. No entity data patching will take place.
+	 */
+	void RemoveVisualDescByIndex(const int32 VisualizationIndex);
+
 	/** 
 	 * @return the array of all the static mesh instance component information
 	 */
diff --git a/Engine/Plugins/Runtime/MassGameplay/Source/MassRepresentation/Public/MassRepresentationTypes.h b/Engine/Plugins/Runtime/MassGameplay/Source/MassRepresentation/Public/MassRepresentationTypes.h
--- a/Engine/Plugins/Runtime/MassGameplay/Source/MassRepresentation/Public/MassRepresentationTypes.h
+++ b/Engine/Plugins/Runtime/MassGameplay/Source/MassRepresentation/Public/MassRepresentationTypes.h
@@ -565,6 +565,11 @@ public:
 
 	int32 GetLODSignificanceRangesNum() const { return LODSignificanceRanges.Num(); }
 
+	bool IsValid() const
+	{
+		return Desc.Meshes.Num() && InstancedStaticMeshComponents.Num() && LODSignificanceRanges.Num();
+	}
+
 protected:
 
 	/** Destroy the visual instance */
diff --git a/Engine/Plugins/Runtime/MassGameplay/Source/MassRepresentation/Public/MassVisualizationComponent.h b/Engine/Plugins/Runtime/MassGameplay/Source/MassRepresentation/Public/MassVisualizationComponent.h
--- a/Engine/Plugins/Runtime/MassGameplay/Source/MassRepresentation/Public/MassVisualizationComponent.h
+++ b/Engine/Plugins/Runtime/MassGameplay/Source/MassRepresentation/Public/MassVisualizationComponent.h
@@ -43,11 +43,20 @@ public:
 	int16 AddVisualDescWithISMComponents(const FStaticMeshInstanceVisualizationDesc& Desc, TArrayView<TObjectPtr<UInstancedStaticMeshComponent>> ISMComponents);
 
 	/** 
-	 * Removes the visualization data associated with the given ISM component. Note that this is safe to do only when
-	 * there are no entities relying on this data. No entity data patching will take place. 
+	 * Removes all the visualization data associated with the given ISM component. The function resolves the VisualizationIndex
+	 * associated with the given ISMComponent and calls RemoveVisualDescByIndex which will remove data on all ISMComponents
+	 * associated with the index. Note that this is safe to do only when there are no entities relying on this data. 
+	 * No entity data patching will take place. 
 	 */
+	UE_DEPRECATED(5.4, "RemoveISMComponent has been deprecated in favor of RemoveVisualDescByIndex. Please use that instead.")
 	void RemoveISMComponent(UInstancedStaticMeshComponent& ISMComponent);
 
+	/**
+	 * Removes all data associated with a given VisualizationIndex. Note that this is safe to do only if there are no
+	 * entities relying on this index. No entity data patching will take place.
+	 */
+	void RemoveVisualDescByIndex(const int32 VisualizationIndex);
+
 	/** Get the array of all visual instance informations */
 	FMassInstancedStaticMeshInfoArrayView GetMutableVisualInfos()
 	{
