PR #9326: Fix typos in comment in MassAgentComponent.h

diff --git a/Engine/Plugins/Runtime/MassGameplay/Source/MassActors/Public/MassAgentComponent.h b/Engine/Plugins/Runtime/MassGameplay/Source/MassActors/Public/MassAgentComponent.h
--- a/Engine/Plugins/Runtime/MassGameplay/Source/MassActors/Public/MassAgentComponent.h
+++ b/Engine/Plugins/Runtime/MassGameplay/Source/MassActors/Public/MassAgentComponent.h
@@ -29,11 +29,11 @@ enum class EAgentComponentState : uint8
 
 /** 
  *  There are two primary use cases for this component:
- *  1. If placed on an AActor blueprint can it lets the used specify additional fragments that will be created for 
+ *  1. If placed on an AActor blueprint it lets the user specify additional fragments that will be created for 
  *     entities spawned based on this given blueprint. 
- *  2. If present present on an actor in the world it makes it communicate with the MassSimulation which will create an 
- *     entity representing given actor. Use case 1) will also applicable in this case. The component is unregistered by 
- *     default and requires manual enabling via a 'Enable' call
+ *  2. If present on an actor in the world it makes it communicate with the MassSimulation which will create an 
+ *     entity representing given actor. Use case 1) will also be applicable in this case. The component is unregistered by 
+ *     default and requires manual enabling via a 'Enable' call.
  * 
  *  @todo use case 2) is currently sitting in a shelved CL of mine. Will be worked on next.
  */
