PR #9620: Fix typo in MassEntityZoneGraphSpawnPointsGenerator.h

diff --git a/Engine/Plugins/Runtime/MassGameplay/Source/MassSpawner/Public/MassEntityZoneGraphSpawnPointsGenerator.h b/Engine/Plugins/Runtime/MassGameplay/Source/MassSpawner/Public/MassEntityZoneGraphSpawnPointsGenerator.h
--- a/Engine/Plugins/Runtime/MassGameplay/Source/MassSpawner/Public/MassEntityZoneGraphSpawnPointsGenerator.h
+++ b/Engine/Plugins/Runtime/MassGameplay/Source/MassSpawner/Public/MassEntityZoneGraphSpawnPointsGenerator.h
@@ -32,7 +32,7 @@ protected:
 	UPROPERTY(EditAnywhere, Category = "ZoneGraph Generator Config")
 	float MinGap = 100;
 
-	/** Minimum gap for spawning entities on a given lanes */
+	/** Maximum gap for spawning entities on a given lanes */
 	UPROPERTY(EditAnywhere, Category = "ZoneGraph Generator Config")
 	float MaxGap = 300;
-};
\ No newline at end of file
+};
