PR #10936: [Mass] fixed: ConfigGuid being duplicated which result in wrong entit?

diff --git a/Engine/Plugins/Runtime/MassGameplay/Source/MassSpawner/Private/MassEntityConfigAsset.cpp b/Engine/Plugins/Runtime/MassGameplay/Source/MassSpawner/Private/MassEntityConfigAsset.cpp
--- a/Engine/Plugins/Runtime/MassGameplay/Source/MassSpawner/Private/MassEntityConfigAsset.cpp
+++ b/Engine/Plugins/Runtime/MassGameplay/Source/MassSpawner/Private/MassEntityConfigAsset.cpp
@@ -214,10 +214,27 @@ void FMassEntityConfig::GetCombinedTraits(TArray<UMassEntityTraitBase*>& OutTrai
 	return GetCombinedTraitsInternal(OutTraits, Visited);
 }
 
+#if WITH_EDITOR
+void FMassEntityConfig::PostDuplicate(const bool bDuplicateForPIE)
+{
+	if (bDuplicateForPIE == false)
+	{
+		ConfigGuid = FGuid::NewGuid();
+	}
+}
+#endif // WITH_EDITOR
+
 //-----------------------------------------------------------------------------
 // UMassEntityConfigAsset
 //-----------------------------------------------------------------------------
 #if WITH_EDITOR
+void UMassEntityConfigAsset::PostDuplicate(const bool bDuplicateForPIE)
+{
+	Super::PostDuplicate(bDuplicateForPIE);
+	
+	Config.PostDuplicate(bDuplicateForPIE);
+}
+
 void UMassEntityConfigAsset::ValidateEntityConfig()
 {
 	if (UWorld* EditorWorld = GEditor->GetEditorWorldContext().World())
diff --git a/Engine/Plugins/Runtime/MassGameplay/Source/MassSpawner/Public/MassEntityConfigAsset.h b/Engine/Plugins/Runtime/MassGameplay/Source/MassSpawner/Public/MassEntityConfigAsset.h
--- a/Engine/Plugins/Runtime/MassGameplay/Source/MassSpawner/Public/MassEntityConfigAsset.h
+++ b/Engine/Plugins/Runtime/MassGameplay/Source/MassSpawner/Public/MassEntityConfigAsset.h
@@ -73,6 +73,11 @@ struct MASSSPAWNER_API FMassEntityConfig
 
 	const FGuid& GetGuid() const { return ConfigGuid; }
 	
+#if WITH_EDITOR
+	/** Needs to be called when the given config is being duplicated - ensured the ConfigGuid remains unique */
+	void PostDuplicate(bool bDuplicateForPIE);
+#endif // WITH_EDITOR
+	
 protected:
 	/** Combines traits based on the config hierarchy and returns list of unique traits */
 	void GetCombinedTraits(TArray<UMassEntityTraitBase*>& OutTraits) const;
@@ -138,6 +143,8 @@ public:
 	}
 
 #if WITH_EDITOR
+	/** Called upon asset's dupllication. Ensured the underlying config's ConfigGuid remains unique */
+	virtual void PostDuplicate(bool bDuplicateForPIE) override;
 
 	UFUNCTION(CallInEditor, Category = "Entity Config")
 	void ValidateEntityConfig();
