Made pawn-and-camera-less player not register as viewers for Mass LOD purposes

Also, reverted the temp fix that was addressing the symptoms of the issue being fixed.

diff --git a/Engine/Plugins/Runtime/MassGameplay/Source/MassLOD/Private/MassLODSubsystem.cpp b/Engine/Plugins/Runtime/MassGameplay/Source/MassLOD/Private/MassLODSubsystem.cpp
--- a/Engine/Plugins/Runtime/MassGameplay/Source/MassLOD/Private/MassLODSubsystem.cpp
+++ b/Engine/Plugins/Runtime/MassGameplay/Source/MassLOD/Private/MassLODSubsystem.cpp
@@ -6,6 +6,8 @@
 #include "WorldPartition/WorldPartitionStreamingSource.h"
 #include "Engine/World.h"
 #include "MassSimulationSubsystem.h"
+#include "Camera/PlayerCameraManager.h"
+#include "GameFramework/Pawn.h"
 #if WITH_EDITOR
 #include "CoreGlobals.h" // GIsEditor
 #include "Editor.h" // GEditor
@@ -319,6 +321,13 @@ void UMassLODSubsystem::AddPlayerViewer(APlayerController& PlayerController)
 		return;
 	}
 #endif // WITH_EDITOR
+	// ignore players that don't have a pawn nor a camera
+	if (PlayerController.GetPawn() == nullptr
+		|| bool(PlayerController.PlayerCameraManager) == false
+		|| PlayerController.PlayerCameraManager->GetLastFrameCameraCacheTime() == 0.f)
+	{
+		return;
+	}
 
 	const int32 HashValue = GetTypeHash(PlayerController.GetFName());
 
