Fixed MassLOD player viewer filtering condition, that caused InstancedActors' ISMs not show up in replays

diff --git a/Engine/Plugins/Runtime/MassGameplay/Source/MassLOD/Private/MassLODSubsystem.cpp b/Engine/Plugins/Runtime/MassGameplay/Source/MassLOD/Private/MassLODSubsystem.cpp
--- a/Engine/Plugins/Runtime/MassGameplay/Source/MassLOD/Private/MassLODSubsystem.cpp
+++ b/Engine/Plugins/Runtime/MassGameplay/Source/MassLOD/Private/MassLODSubsystem.cpp
@@ -323,8 +323,9 @@ void UMassLODSubsystem::AddPlayerViewer(APlayerController& PlayerController)
 #endif // WITH_EDITOR
 	// ignore players that don't have a pawn nor a camera
 	if (PlayerController.GetPawn() == nullptr
-		|| bool(PlayerController.PlayerCameraManager) == false
-		|| PlayerController.PlayerCameraManager->GetLastFrameCameraCacheTime() == 0.f)
+		&& (bool(PlayerController.PlayerCameraManager) == false
+			|| PlayerController.PlayerCameraManager->GetLastFrameCameraCacheTime() == 0.f)
+		)	
 	{
 		return;
 	}
