PR #10931: Fix MassRepresentationAccessMode

Made an additional fix to another query detected by checkSlow

diff --git a/Engine/Plugins/Runtime/MassGameplay/Source/MassLOD/Private/MassSimulationLOD.cpp b/Engine/Plugins/Runtime/MassGameplay/Source/MassLOD/Private/MassSimulationLOD.cpp
--- a/Engine/Plugins/Runtime/MassGameplay/Source/MassLOD/Private/MassSimulationLOD.cpp
+++ b/Engine/Plugins/Runtime/MassGameplay/Source/MassLOD/Private/MassSimulationLOD.cpp
@@ -165,7 +165,7 @@ void UMassSimulationLODProcessor::Execute(FMassEntityManager& EntityManager, FMa
 		EntityQueryVariableTick.ForEachEntityChunk(EntityManager, Context, [Time](FMassExecutionContext& Context)
 		{
 			FMassSimulationVariableTickSharedFragment& TickRateSharedFragment = Context.GetMutableSharedFragment<FMassSimulationVariableTickSharedFragment>();
-			TConstArrayView<FMassSimulationLODFragment> SimulationLODFragments = Context.GetMutableFragmentView<FMassSimulationLODFragment>();
+			TConstArrayView<FMassSimulationLODFragment> SimulationLODFragments = Context.GetFragmentView<FMassSimulationLODFragment>();
 			TArrayView<FMassSimulationVariableTickFragment> SimulationVariableTickFragments = Context.GetMutableFragmentView<FMassSimulationVariableTickFragment>();
 
 			TickRateSharedFragment.LODTickRateController.UpdateTickRateFromLOD(Context, SimulationLODFragments, SimulationVariableTickFragments, Time);
diff --git a/Engine/Plugins/Runtime/MassGameplay/Source/MassRepresentation/Private/MassRepresentationProcessor.cpp b/Engine/Plugins/Runtime/MassGameplay/Source/MassRepresentation/Private/MassRepresentationProcessor.cpp
--- a/Engine/Plugins/Runtime/MassGameplay/Source/MassRepresentation/Private/MassRepresentationProcessor.cpp
+++ b/Engine/Plugins/Runtime/MassGameplay/Source/MassRepresentation/Private/MassRepresentationProcessor.cpp
@@ -405,7 +405,7 @@ UMassRepresentationFragmentDestructor::UMassRepresentationFragmentDestructor()
 
 void UMassRepresentationFragmentDestructor::ConfigureQueries()
 {
-	EntityQuery.AddRequirement<FMassRepresentationFragment>(EMassFragmentAccess::ReadOnly);
+	EntityQuery.AddRequirement<FMassRepresentationFragment>(EMassFragmentAccess::ReadWrite);
 	EntityQuery.AddRequirement<FMassActorFragment>(EMassFragmentAccess::ReadWrite);
 	EntityQuery.AddConstSharedRequirement<FMassRepresentationParameters>();
 	EntityQuery.AddSharedRequirement<FMassRepresentationSubsystemSharedFragment>(EMassFragmentAccess::ReadWrite);
