Added the FMassGuidFragment and Made mass spawning utilize it. The change also includes an extension of MassActorSpawner.SpawnActor so that custom actor spawn parameters can be passed in.

diff --git a/Engine/Plugins/Runtime/MassGameplay/Source/MassActors/Private/MassActorSpawnerSubsystem.cpp b/Engine/Plugins/Runtime/MassGameplay/Source/MassActors/Private/MassActorSpawnerSubsystem.cpp
--- a/Engine/Plugins/Runtime/MassGameplay/Source/MassActors/Private/MassActorSpawnerSubsystem.cpp
+++ b/Engine/Plugins/Runtime/MassGameplay/Source/MassActors/Private/MassActorSpawnerSubsystem.cpp
@@ -10,6 +10,8 @@
 #include "MassSimulationSubsystem.h"
 #include "Logging/LogScopedVerbosityOverride.h"
 #include "ProfilingDebugging/CsvProfiler.h"
+#include "Engine/Level.h"
+
 
 CSV_DEFINE_CATEGORY(MassActors, true);
 
@@ -246,7 +248,8 @@ ESpawnRequestStatus UMassActorSpawnerSubsystem::SpawnOrRetrieveFromPool(FConstSt
 		}
 	}
 
-	ESpawnRequestStatus SpawnStatus = SpawnActor(SpawnRequestView, OutSpawnedActor);
+	FActorSpawnParameters ActorSpawnParameters;
+	ESpawnRequestStatus SpawnStatus = SpawnActor(SpawnRequestView, OutSpawnedActor, ActorSpawnParameters);
 
 	if (SpawnStatus == ESpawnRequestStatus::Succeeded)
 	{
@@ -269,7 +272,27 @@ ESpawnRequestStatus UMassActorSpawnerSubsystem::SpawnOrRetrieveFromPool(FConstSt
 	return SpawnStatus;
 }
 
- ESpawnRequestStatus UMassActorSpawnerSubsystem::SpawnActor(FConstStructView SpawnRequestView, TObjectPtr<AActor>& OutSpawnedActor) const
+TObjectPtr<AActor> UMassActorSpawnerSubsystem::FindActorByName(const FName ActorName, ULevel* OverrideLevel) const
+{
+	TObjectPtr<AActor> FoundActor;
+	check(GetWorld());
+	OverrideLevel = OverrideLevel ? OverrideLevel : GetWorld()->GetCurrentLevel();
+	
+	if (UObject* FoundObject = StaticFindObjectFast(nullptr, OverrideLevel, ActorName))
+	{
+		FoundActor = Cast<AActor>(FoundObject);
+		if (FoundObject)
+		{
+			if (IsValid(FoundActor) == false)
+			{
+				FoundActor->ClearGarbage();
+			}
+		}
+	}
+	return FoundActor;
+}
+
+ ESpawnRequestStatus UMassActorSpawnerSubsystem::SpawnActor(FConstStructView SpawnRequestView, TObjectPtr<AActor>& OutSpawnedActor, FActorSpawnParameters& InOutSpawnParameters) const
 {
 	TRACE_CPUPROFILER_EVENT_SCOPE(UMassActorSpawnerSubsystem::SpawnActor);
 
@@ -278,9 +301,23 @@ ESpawnRequestStatus UMassActorSpawnerSubsystem::SpawnOrRetrieveFromPool(FConstSt
 
 	const FMassActorSpawnRequest& SpawnRequest = SpawnRequestView.Get<const FMassActorSpawnRequest>();
 
-	FActorSpawnParameters SpawnParameters;
-	SpawnParameters.SpawnCollisionHandlingOverride = ESpawnActorCollisionHandlingMethod::AlwaysSpawn;
-	OutSpawnedActor = World->SpawnActor<AActor>(SpawnRequest.Template, SpawnRequest.Transform, SpawnParameters);
+	if (SpawnRequest.Guid.IsValid())
+	{
+		// offsetting `D` by 1 since `0` has special meaning for FNames
+		InOutSpawnParameters.Name = FName(FString::Printf(TEXT("%s_%ud_%ud_%ud"), *SpawnRequest.Template->GetName(), SpawnRequest.Guid.A, SpawnRequest.Guid.B, SpawnRequest.Guid.C), SpawnRequest.Guid.D + 1);
+		//InOutSpawnParameters.OverrideLevel = InOutSpawnParameters.OverrideLevel ? OverrideLevel : World->GetCurrentLevel();
+
+		OutSpawnedActor = FindActorByName(InOutSpawnParameters.Name, InOutSpawnParameters.OverrideLevel ? InOutSpawnParameters.OverrideLevel : World->GetCurrentLevel());
+		if (OutSpawnedActor)
+		{
+			OutSpawnedActor->SetActorEnableCollision(true);
+			OutSpawnedActor->SetActorHiddenInGame(false);
+			return ESpawnRequestStatus::Succeeded;
+		}
+	}
+	
+	InOutSpawnParameters.SpawnCollisionHandlingOverride = ESpawnActorCollisionHandlingMethod::AlwaysSpawn;
+	OutSpawnedActor = World->SpawnActor<AActor>(SpawnRequest.Template, SpawnRequest.Transform, InOutSpawnParameters);
 
 	return IsValid(OutSpawnedActor) ? ESpawnRequestStatus::Succeeded : ESpawnRequestStatus::Failed;
 }
diff --git a/Engine/Plugins/Runtime/MassGameplay/Source/MassActors/Public/MassActorSpawnerSubsystem.h b/Engine/Plugins/Runtime/MassGameplay/Source/MassActors/Public/MassActorSpawnerSubsystem.h
--- a/Engine/Plugins/Runtime/MassGameplay/Source/MassActors/Public/MassActorSpawnerSubsystem.h
+++ b/Engine/Plugins/Runtime/MassGameplay/Source/MassActors/Public/MassActorSpawnerSubsystem.h
@@ -12,6 +12,10 @@
 
 #include "MassActorSpawnerSubsystem.generated.h"
 
+
+class ULevel; 
+struct FActorSpawnParameters;
+
 // Handle for an actor spawning request
 USTRUCT()
 struct MASSACTORS_API FMassActorSpawnRequestHandle : public FIndexedHandleBase
@@ -89,6 +93,9 @@ public:
 	/** Requested world time seconds */
 	double RequestedTime = 0.;
 
+	/** If set, will be used to name the spawned character */
+	FGuid Guid;
+
 	void Reset()
 	{
 		MassAgent = FMassEntityHandle();
@@ -230,7 +237,9 @@ protected:
 
 	/** Actual code that will spawn the actor, overridable by subclass if need to be.
 	 *  @return spawned actor if succeeded. */
-	virtual ESpawnRequestStatus SpawnActor(FConstStructView SpawnRequestView, TObjectPtr<AActor>& OutSpawnedActor) const;
+	virtual ESpawnRequestStatus SpawnActor(FConstStructView SpawnRequestView, TObjectPtr<AActor>& OutSpawnedActor, FActorSpawnParameters& InOutSpawnParameters) const;
+
+	TObjectPtr<AActor> FindActorByName(const FName ActorName, ULevel* OverrideLevel) const;
 
 	/** Go through the spawning request and spawn them until we reach the budget 
 	 * @param MaxTimeSlicePerTick is the budget in seconds allowed to do spawning */
diff --git a/Engine/Plugins/Runtime/MassGameplay/Source/MassActors/Public/MassActorSubsystem.h b/Engine/Plugins/Runtime/MassGameplay/Source/MassActors/Public/MassActorSubsystem.h
--- a/Engine/Plugins/Runtime/MassGameplay/Source/MassActors/Public/MassActorSubsystem.h
+++ b/Engine/Plugins/Runtime/MassGameplay/Source/MassActors/Public/MassActorSubsystem.h
@@ -14,6 +14,14 @@ struct FMassEntityHandle;
 class AActor;
 struct FMassEntityManager;
 
+USTRUCT()
+struct MASSACTORS_API FMassGuidFragment : public FObjectWrapperFragment
+{
+	GENERATED_BODY()
+
+	FGuid Guid;
+};
+
 /**
  * Fragment to save the actor pointer of a mass entity if it exist
  */
diff --git a/Engine/Plugins/Runtime/MassGameplay/Source/MassRepresentation/Private/MassRepresentationSubsystem.cpp b/Engine/Plugins/Runtime/MassGameplay/Source/MassRepresentation/Private/MassRepresentationSubsystem.cpp
--- a/Engine/Plugins/Runtime/MassGameplay/Source/MassRepresentation/Private/MassRepresentationSubsystem.cpp
+++ b/Engine/Plugins/Runtime/MassGameplay/Source/MassRepresentation/Private/MassRepresentationSubsystem.cpp
@@ -143,6 +143,10 @@ AActor* UMassRepresentationSubsystem::GetOrSpawnActorFromTemplate(const FMassEnt
 	// If we reach here, means we need to create a spawn request
 	FMassActorSpawnRequest SpawnRequest;
 	SpawnRequest.MassAgent = MassAgent;
+	if (FMassGuidFragment* GuidFragment = EntityManager->GetFragmentDataPtr<FMassGuidFragment>(MassAgent))
+	{
+		SpawnRequest.Guid = GuidFragment->Guid;
+	}
 	SpawnRequest.Template = TemplateToSpawn;
 	SpawnRequest.Transform = Transform;
 	SpawnRequest.Priority = Priority;
