Mass Representation - Added bForceActorRepresentationForExternalActors option to FMassRepresentationParameters which forces UMassRepresentationProcessor to override the WantedRepresentationType to actor representation whenever an 'external' Actor (not owned by Mass) is set on an entities actor fragment & reverting to the natural WantedRepresentationType when the actor fragment is reset.

Also added a related but general check for Actor removal occurring outside of Mass, to forcibly set CurrentRepresentation to None, to force a new Actor to spawn if WantedRepresentation = Actor, or instance an ISMC etc.

Useful for server-authoritative actor spawning to force actor representation on clients.

#tests PIE

diff --git a/Engine/Plugins/Runtime/MassGameplay/Source/MassRepresentation/Private/MassRepresentationProcessor.cpp b/Engine/Plugins/Runtime/MassGameplay/Source/MassRepresentation/Private/MassRepresentationProcessor.cpp
--- a/Engine/Plugins/Runtime/MassGameplay/Source/MassRepresentation/Private/MassRepresentationProcessor.cpp
+++ b/Engine/Plugins/Runtime/MassGameplay/Source/MassRepresentation/Private/MassRepresentationProcessor.cpp
@@ -73,6 +73,7 @@ void UMassRepresentationProcessor::UpdateRepresentation(FMassExecutionContext& C
 		const FMassRepresentationLODFragment& RepresentationLOD = RepresentationLODList[EntityIdx];
 		FMassRepresentationFragment& Representation = RepresentationList[EntityIdx];
 		FMassActorFragment& ActorInfo = ActorList[EntityIdx];
+		AActor* Actor = ActorInfo.GetMutable();
 
 		// Keeping a copy of the that last calculated previous representation
 		const EMassRepresentationType PrevRepresentationCopy = Representation.PrevRepresentation;
@@ -87,13 +88,33 @@ void UMassRepresentationProcessor::UpdateRepresentation(FMassExecutionContext& C
 		{
 			WantedRepresentationType = RepresentationParams.CachedDefaultRepresentationType;
 		}
+
+		// If bForceActorRepresentationForExternalActors is enabled and we have an Actor reference for this entity, forcibly use it
+		// by enforcing an actor representation as the WantedRepresentation. If we're coming from ISMC, we'll remove the instance
+		// and switch to this actor, commiting either LowResSpawnedActor or HighResSpawnedActor as the new CurrentRepresentation.
+		// Once the Actor is destroyed however, this override stops, allowing the natural WantedRepresentationType to return.
+		//
+		// Useful for server-authoritative Actor spawning, with replicated Actors inserting themselves into Mass whilst they're
+		// replicated, enforcing actor representation on clients whilst they're present.
+		if (RepresentationParams.bForceActorRepresentationForExternalActors && IsValid(Actor) && !ActorInfo.IsOwnedByMass())
+		{
+			WantedRepresentationType = Representation.CurrentRepresentation == EMassRepresentationType::LowResSpawnedActor ? EMassRepresentationType::LowResSpawnedActor : EMassRepresentationType::HighResSpawnedActor;
+		}
+
+		// Has Actor unexpectedly been unset / destroyed since we last ran? 
+		if (!IsValid(Actor) && (Representation.CurrentRepresentation == EMassRepresentationType::LowResSpawnedActor || Representation.CurrentRepresentation == EMassRepresentationType::HighResSpawnedActor))
+		{
+			// Set CurrentRepresentation = None so we get a chance to see CurrentRepresentation != WantedRepresentationType and spawn 
+			// another actor.
+			Representation.CurrentRepresentation = EMassRepresentationType::None;
+		}
 		
 		auto DisableActorForISM = [&](AActor*& Actor)
 		{
 			if (!Actor || ActorInfo.IsOwnedByMass())
 			{
 				// Execute only if the high res is different than the low res Actor 
-				// Or if we do not wish to keep the low res actor while in ISM
+				// Or if we do not wish to keep the low res actor while in TransformList
 				if (Representation.HighResTemplateActorIndex != Representation.LowResTemplateActorIndex || !RepresentationParams.bKeepLowResActors)
 				{
 					// Try releasing the high actor or any high res spawning request
@@ -128,7 +149,6 @@ void UMassRepresentationProcessor::UpdateRepresentation(FMassExecutionContext& C
 				Representation.PrevLODSignificance = RepresentationLOD.LODSignificance;
 			}
 
-			AActor* Actor = ActorInfo.GetMutable();
 			switch (WantedRepresentationType)
 			{
 				case EMassRepresentationType::HighResSpawnedActor:
@@ -218,7 +238,6 @@ void UMassRepresentationProcessor::UpdateRepresentation(FMassExecutionContext& C
 				 Representation.PrevRepresentation == EMassRepresentationType::StaticMeshInstance &&
 			    (PrevRepresentationCopy == EMassRepresentationType::HighResSpawnedActor || PrevRepresentationCopy == EMassRepresentationType::LowResSpawnedActor))
 		{
-			AActor* Actor = ActorInfo.GetMutable();
 			DisableActorForISM(Actor);
 		}
 	}
diff --git a/Engine/Plugins/Runtime/MassGameplay/Source/MassRepresentation/Public/MassRepresentationFragments.h b/Engine/Plugins/Runtime/MassGameplay/Source/MassRepresentation/Public/MassRepresentationFragments.h
--- a/Engine/Plugins/Runtime/MassGameplay/Source/MassRepresentation/Public/MassRepresentationFragments.h
+++ b/Engine/Plugins/Runtime/MassGameplay/Source/MassRepresentation/Public/MassRepresentationFragments.h
@@ -110,17 +110,33 @@ struct FMassRepresentationParameters : public FMassSharedFragment
 	UPROPERTY(EditAnywhere, Category = "Mass|Representation", config)
 	EMassRepresentationType LODRepresentation[EMassLOD::Max] = { EMassRepresentationType::HighResSpawnedActor, EMassRepresentationType::LowResSpawnedActor, EMassRepresentationType::StaticMeshInstance, EMassRepresentationType::None };
 
+	/** 
+	 * If true, forces UMassRepresentationProcessor to override the WantedRepresentationType to actor representation whenever an external (non Mass owned)
+	 * actor is set on an entitie's FMassActorFragment fragment. If / when the actor fragment is reset, WantedRepresentationType resumes selecting the 
+	 * appropriate representation for the current representation LOD.
+	 *
+	 * Useful for server-authoritative actor spawning to force actor representation on clients for replicated actors. 
+	 */ 
+	UPROPERTY(EditAnywhere, Category = "Mass|Representation", config)
+	uint8 bForceActorRepresentationForExternalActors : 1 = false;
+
 	/** If true, LowRes actors will be kept around, disabled, whilst StaticMeshInstance representation is active */
 	UPROPERTY(EditAnywhere, Category = "Mass|Representation", config)
-	bool bKeepLowResActors = true;
+	uint8 bKeepLowResActors : 1  = true;
 
 	/** When switching to ISM keep the actor an extra frame, helps cover rendering glitches (i.e. occlusion query being one frame late) */
 	UPROPERTY(EditAnywhere, Category = "Mass|Representation", config)
-	bool bKeepActorExtraFrame = false;
+	uint8 bKeepActorExtraFrame : 1  = false;
 
 	/** If true, will spread the first visualization update over the period specified in NotVisibleUpdateRate member */
 	UPROPERTY(EditAnywhere, Category = "Mass|Representation", config)
-	bool bSpreadFirstVisualizationUpdate = false;
+	uint8 bSpreadFirstVisualizationUpdate : 1  = false;
+
+#if WITH_EDITORONLY_DATA
+	/** the property is marked like this to ensure it won't show up in UI */
+	UPROPERTY(EditDefaultsOnly, Category = "Mass|Visual")
+	uint8 bCanModifyRepresentationActorManagementClass : 1 = true;
+#endif // WITH_EDITORONLY_DATA
 
 	/** World Partition grid name to test collision against, default None will be the main grid */
 	UPROPERTY(EditAnywhere, Category = "Mass|Representation", config)
@@ -138,12 +154,6 @@ struct FMassRepresentationParameters : public FMassSharedFragment
 
 	UPROPERTY(Transient)
 	mutable TObjectPtr<UMassRepresentationActorManagement> CachedRepresentationActorManagement = nullptr;
-
-#if WITH_EDITORONLY_DATA
-	/** the property is marked like this to ensure it won't show up in UI */
-	UPROPERTY(EditDefaultsOnly, Category = "Mass|Visual")
-	bool bCanModifyRepresentationActorManagementClass = true;
-#endif // WITH_EDITORONLY_DATA
 };
 
 template<>
