Minor MassLODSubsystem refactor to support any Actor as a viewer

diff --git a/Engine/Plugins/Runtime/MassGameplay/Source/MassLOD/Private/MassLODSubsystem.cpp b/Engine/Plugins/Runtime/MassGameplay/Source/MassLOD/Private/MassLODSubsystem.cpp
--- a/Engine/Plugins/Runtime/MassGameplay/Source/MassLOD/Private/MassLODSubsystem.cpp
+++ b/Engine/Plugins/Runtime/MassGameplay/Source/MassLOD/Private/MassLODSubsystem.cpp
@@ -29,6 +29,9 @@ namespace MassLOD
 }
 }
 
+//-----------------------------------------------------------------------------
+// UMassLODSubsystem
+//-----------------------------------------------------------------------------
 void UMassLODSubsystem::Initialize(FSubsystemCollectionBase& Collection)
 {
 	Collection.InitializeDependency(UMassSimulationSubsystem::StaticClass());
@@ -105,9 +108,9 @@ const TArray<FViewerInfo>& UMassLODSubsystem::GetSynchronizedViewers()
 	return Viewers;
 }
 
-FMassViewerHandle UMassLODSubsystem::GetViewerHandleFromPlayerController(const APlayerController* PlayerController) const
+FMassViewerHandle UMassLODSubsystem::GetViewerHandleFromActor(const AActor& Actor) const
 {
-	const FMassViewerHandle* Handle = ViewerMap.Find(GetTypeHash(PlayerController->GetFName()));
+	const FMassViewerHandle* Handle = ViewerMap.Find(GetTypeHash(Actor.GetFName()));
 	return Handle ? *Handle : FMassViewerHandle();
 }
 
@@ -121,7 +124,7 @@ FMassViewerHandle UMassLODSubsystem::GetViewerHandleFromStreamingSource(const FN
 APlayerController* UMassLODSubsystem::GetPlayerControllerFromViewerHandle(const FMassViewerHandle& ViewerHandle) const
 {
 	const int32 ViewerIdx = GetValidViewerIdx(ViewerHandle);
-	return ViewerIdx != INDEX_NONE ? Viewers[ViewerIdx].PlayerController : nullptr;
+	return ViewerIdx != INDEX_NONE ? Viewers[ViewerIdx].GetPlayerController() : nullptr;
 }
 
 void UMassLODSubsystem::SynchronizeViewers()
@@ -149,13 +152,14 @@ void UMassLODSubsystem::SynchronizeViewers()
 			continue;
 		}
 
-		if (ViewerInfo.PlayerController != nullptr
+		APlayerController* ViewerAsPlayerController = ViewerInfo.GetPlayerController();
+		if (ViewerAsPlayerController != nullptr
 #if WITH_EDITOR
 			&& bIgnorePlayerControllersDueToSimulation == false
 #endif // WITH_EDITOR
 			)
 		{
-			LocalViewerMap.Add(GetTypeHash(ViewerInfo.PlayerController->GetFName()), ViewerInfo.Handle);
+			LocalViewerMap.Add(GetTypeHash(ViewerAsPlayerController->GetFName()), ViewerInfo.Handle);
 		}
 		else if (!ViewerInfo.StreamingSourceName.IsNone() && StreamingSources.FindByPredicate([&ViewerInfo](const FWorldPartitionStreamingSource& Source){ return Source.Name == ViewerInfo.StreamingSourceName; }) != nullptr)
 		{
@@ -209,6 +213,21 @@ void UMassLODSubsystem::SynchronizeViewers()
 				}
 			}
 		}
+
+		if (bAllowNonPlayerViwerActors)
+		{
+			for (int32 ActorViewerIndex = RegisteredActorViewers.Num() - 1; ActorViewerIndex >= 0; --ActorViewerIndex)
+			{
+				if (RegisteredActorViewers[ActorViewerIndex])
+				{
+					AddActorViewer(*RegisteredActorViewers[ActorViewerIndex]);
+				}
+				else
+				{
+					RegisteredActorViewers.RemoveAtSwap(ActorViewerIndex, 1, /*bAllowShrinking=*/false);
+				}
+			}
+		}
 	}
 #if WITH_EDITOR
 	if (bUseEditorLevelViewports)
@@ -255,25 +274,30 @@ void UMassLODSubsystem::SynchronizeViewers()
 			continue;
 		}
 
-		if (ViewerInfo.PlayerController)
+		if (APlayerController* ViewerAsPlayerController = ViewerInfo.GetPlayerController())
 		{
-			ViewerInfo.bEnabled = !WorldPartition || ViewerInfo.PlayerController->bEnableStreamingSource;
+			ViewerInfo.bEnabled = !WorldPartition || ViewerAsPlayerController->bEnableStreamingSource;
 
 			FVector PlayerCameraLocation(ForceInitToZero);
 			FRotator PlayerCameraRotation(FRotator::ZeroRotator);
-			ViewerInfo.PlayerController->GetPlayerViewPoint(PlayerCameraLocation, PlayerCameraRotation);
+			ViewerAsPlayerController->GetPlayerViewPoint(PlayerCameraLocation, PlayerCameraRotation);
 			ViewerInfo.Location = PlayerCameraLocation;
 			ViewerInfo.Rotation = PlayerCameraRotation;
 
 			// Try to fetch a more precise FOV
-			if(ViewerInfo.PlayerController->PlayerCameraManager)
+			if(ViewerAsPlayerController->PlayerCameraManager)
 			{
-				ViewerInfo.FOV = ViewerInfo.PlayerController->PlayerCameraManager->GetFOVAngle();
+				ViewerInfo.FOV = ViewerAsPlayerController->PlayerCameraManager->GetFOVAngle();
 
 				// @todo need to find a way to retrieve aspect ratio, this does not seems to work
 				//ViewerInfo.AspectRatio = MinViewInfo.AspectRatio;
 			}
 		}
+		else if (AActor* Actor = ViewerInfo.ActorViewer.Get())
+		{
+			ViewerInfo.Location = Actor->GetActorLocation();
+			ViewerInfo.Rotation = Actor->GetActorRotation();
+		}
 #if WITH_EDITOR
 		else if (bUseEditorLevelViewports && ViewerInfo.EditorViewportClientIndex != INDEX_NONE)
 		{
@@ -340,8 +364,8 @@ void UMassLODSubsystem::AddPlayerViewer(APlayerController& PlayerController)
 		check(ViewerHandleIdx != INDEX_NONE);
 
 		FViewerInfo& ViewerInfo = Viewers[ViewerHandleIdx];
-		check(ViewerInfo.PlayerController == nullptr);
-		ViewerInfo.PlayerController = &PlayerController;
+		check(ViewerInfo.ActorViewer == nullptr);
+		ViewerInfo.ActorViewer = &PlayerController;
 	}
 	else
 	{
@@ -355,7 +379,7 @@ void UMassLODSubsystem::AddPlayerViewer(APlayerController& PlayerController)
 			const bool bAddNew = ViewerFreeIndices.Num() == 0;
 			const int NewIdx = bAddNew ? Viewers.Num() : ViewerFreeIndices.Pop();
 			FViewerInfo& NewViewerInfo = bAddNew ? Viewers.AddDefaulted_GetRef() : Viewers[NewIdx];
-			NewViewerInfo.PlayerController = &PlayerController;
+			NewViewerInfo.ActorViewer = &PlayerController;
 			NewViewerInfo.Handle.Index = NewIdx;
 			NewViewerInfo.Handle.SerialNumber = GetNextViewerSerialNumber();
 			NewViewerInfo.HashValue = HashValue;
@@ -391,6 +415,38 @@ void UMassLODSubsystem::AddStreamingSourceViewer(const FName StreamingSourceName
 	}
 }
 
+void UMassLODSubsystem::AddActorViewer(AActor& ActorViewer)
+{
+	// @todo we might need to use PathName instead 
+	const int32 HashValue = GetTypeHash(ActorViewer.GetFName());
+
+	FMassViewerHandle& ViewerHandle = ViewerMap.FindOrAdd(HashValue, FMassViewerHandle());
+	if (ViewerHandle.IsValid())
+	{
+		// We are only interested to set the player controller if it was not already set.
+		const int32 ViewerHandleIdx = GetValidViewerIdx(ViewerHandle);
+		check(ViewerHandleIdx != INDEX_NONE);
+
+		FViewerInfo& ViewerInfo = Viewers[ViewerHandleIdx];
+		ViewerInfo.ActorViewer = &ActorViewer;
+	}
+	else
+	{
+		// Add new viewer
+		const bool bAddNew = ViewerFreeIndices.Num() == 0;
+		const int NewIdx = bAddNew ? Viewers.Num() : ViewerFreeIndices.Pop();
+
+		FViewerInfo& NewViewerInfo = bAddNew ? Viewers.AddDefaulted_GetRef() : Viewers[NewIdx];
+		NewViewerInfo.ActorViewer = &ActorViewer;
+		NewViewerInfo.Handle.Index = NewIdx;
+		NewViewerInfo.Handle.SerialNumber = GetNextViewerSerialNumber();
+		NewViewerInfo.HashValue = HashValue;
+		ViewerHandle = NewViewerInfo.Handle;
+
+		OnViewerAddedDelegate.Broadcast(NewViewerInfo);
+	}
+}
+
 #if WITH_EDITOR
 void UMassLODSubsystem::AddEditorViewer(const int32 HashValue, const int32 ClientIndex)
 {
@@ -449,9 +505,9 @@ void UMassLODSubsystem::RemoveViewerInternal(const FMassViewerHandle& ViewerHand
 
 	OnViewerRemovedDelegate.Broadcast(ViewerInfo);
 
-	if (ViewerInfo.PlayerController)
+	if (APlayerController* ViewerAsPlayerController = ViewerInfo.GetPlayerController())
 	{
-		ViewerInfo.PlayerController->OnEndPlay.RemoveDynamic(this, &UMassLODSubsystem::OnPlayerControllerEndPlay);
+		ViewerAsPlayerController->OnEndPlay.RemoveDynamic(this, &UMassLODSubsystem::OnPlayerControllerEndPlay);
 	}
 
 	ViewerMap.Remove(ViewerInfo.HashValue);
@@ -465,7 +521,7 @@ void UMassLODSubsystem::OnPlayerControllerEndPlay(AActor* Actor, EEndPlayReason:
 	APlayerController* PlayerController = Cast<APlayerController>(Actor);
 	if (ensure(PlayerController))
 	{
-		const FMassViewerHandle ViewerHandle = GetViewerHandleFromPlayerController(PlayerController);
+		const FMassViewerHandle ViewerHandle = GetViewerHandleFromActor(*PlayerController);
 		if (ensure(ViewerHandle.IsValid()))
 		{
 			RemoveViewer(ViewerHandle);
@@ -473,10 +529,37 @@ void UMassLODSubsystem::OnPlayerControllerEndPlay(AActor* Actor, EEndPlayReason:
 	}
 }
 
+void UMassLODSubsystem::RegisterActorViewer(AActor& ActorViewer)
+{
+	RegisteredActorViewers.AddUnique(&ActorViewer);
+}
+
+void UMassLODSubsystem::UnregisterActorViewer(AActor& ActorViewer)
+{
+	if (RegisteredActorViewers.RemoveSingleSwap(&ActorViewer, /*bAllowShrinking=*/false))
+	{
+		const FMassViewerHandle ViewerHandle = GetViewerHandleFromActor(ActorViewer);
+		if (ensure(ViewerHandle.IsValid()))
+		{
+			RemoveViewer(ViewerHandle);
+		}
+	}
+}
+
+FMassViewerHandle UMassLODSubsystem::GetViewerHandleFromPlayerController(const APlayerController* PlayerController) const
+{
+	PRAGMA_DISABLE_DEPRECATION_WARNINGS
+	return PlayerController ? GetViewerHandleFromActor(*PlayerController) : FMassViewerHandle();
+	PRAGMA_ENABLE_DEPRECATION_WARNINGS
+}
+
+//-----------------------------------------------------------------------------
+// FViewerInfo
+//-----------------------------------------------------------------------------
 void FViewerInfo::Reset()
 {
 	Handle.Invalidate();
-	PlayerController = nullptr;
+	ActorViewer = nullptr;
 #if WITH_EDITOR
 	EditorViewportClientIndex = INDEX_NONE;
 #endif // WITH_EDITOR
@@ -485,9 +568,15 @@ void FViewerInfo::Reset()
 
 bool FViewerInfo::IsLocal() const
 {
-	return (PlayerController && PlayerController->IsLocalController()) || !StreamingSourceName.IsNone()
+	APlayerController* ViewerAsPlayerController = GetPlayerController();
+	return (ViewerAsPlayerController && ViewerAsPlayerController->IsLocalController()) || !StreamingSourceName.IsNone()
 #if WITH_EDITOR
 		|| EditorViewportClientIndex != INDEX_NONE
 #endif // WITH_EDITOR
 		;
 }
+
+APlayerController* FViewerInfo::GetPlayerController() const
+{
+	return Cast<APlayerController>(ActorViewer.Get());
+}
diff --git a/Engine/Plugins/Runtime/MassGameplay/Source/MassLOD/Public/MassLODSubsystem.h b/Engine/Plugins/Runtime/MassGameplay/Source/MassLOD/Public/MassLODSubsystem.h
--- a/Engine/Plugins/Runtime/MassGameplay/Source/MassLOD/Public/MassLODSubsystem.h
+++ b/Engine/Plugins/Runtime/MassGameplay/Source/MassLOD/Public/MassLODSubsystem.h
@@ -11,6 +11,8 @@
 #include "MassLODSubsystem.generated.h"
 
 class UMassLODSubsystem;
+class AActor;
+class APlayerController;
 
 /*
  * Handle that lets you reference the concept of a viewer
@@ -28,8 +30,13 @@ struct FViewerInfo
 {
 	GENERATED_BODY()
 
+	FViewerInfo() = default;
+PRAGMA_DISABLE_DEPRECATION_WARNINGS
+	FViewerInfo(const FViewerInfo& Other) = default;
+PRAGMA_ENABLE_DEPRECATION_WARNINGS
+
 	UPROPERTY(transient)
-	TObjectPtr<APlayerController> PlayerController = nullptr;
+	TObjectPtr<AActor> ActorViewer = nullptr;
 	
 	FName StreamingSourceName;
 
@@ -50,6 +57,11 @@ struct FViewerInfo
 	void Reset();
 
 	bool IsLocal() const;
+
+	MASSLOD_API APlayerController* GetPlayerController() const;
+private:
+	UE_DEPRECATED_FORGAME(5.4, "PlayerController member variable has been deprecated in favor of more generic ActorViewer. Use that instead.")
+	TObjectPtr<APlayerController> PlayerController = nullptr;
 };
 
 UE_DEPRECATED(5.3, "FOnViewerAdded is deprecated. Use UMassLODSubsystem::FOnViewerAdded instead.")
@@ -82,7 +94,7 @@ public:
 	const TArray<FViewerInfo>& GetSynchronizedViewers();
 
 	/** Returns viewer handle from the PlayerController pointer */
-	FMassViewerHandle GetViewerHandleFromPlayerController(const APlayerController* PlayerController) const;
+	FMassViewerHandle GetViewerHandleFromActor(const AActor& Actor) const;
 
 	/** Returns viewer handle from the streaming source name */
 	FMassViewerHandle GetViewerHandleFromStreamingSource(const FName StreamingSourceName) const;
@@ -96,6 +108,9 @@ public:
 	/** Returns the delegate called when viewer are removed from the list */
 	FOnViewerRemoved& GetOnViewerRemovedDelegate() { return OnViewerRemovedDelegate; }
 
+	void RegisterActorViewer(AActor& ActorViewer);
+	void UnregisterActorViewer(AActor& ActorViewer);
+
 protected:
 	// USubsystem BEGIN
 	virtual void Initialize(FSubsystemCollectionBase& Collection) override;
@@ -119,6 +134,8 @@ protected:
 	/** Adds the given streaming source as a viewer to the list and sends notification about addition */
 	void AddStreamingSourceViewer(const FName StreamingSourceName);
 
+	void AddActorViewer(AActor& ActorViewer);
+
 #if WITH_EDITOR
 	/** Adds the editor viewport client (identified via an index) as a viewer to the list and sends notification about addition */
 	void AddEditorViewer(const int32 HashValue, const int32 ClientIndex);
@@ -143,6 +160,9 @@ protected:
 	UPROPERTY(EditDefaultsOnly, Category = "Mass|LOD", config)
 	uint8 bGatherStreamingSources : 1 = true;
 
+	UPROPERTY(EditDefaultsOnly, Category = "Mass|LOD", config)
+	uint8 bAllowNonPlayerViwerActors : 1 = true;
+
 private:
 	/** Removes a viewer to the list and send notification about removal */
 	void RemoveViewerInternal(const FMassViewerHandle& ViewerHandle);
@@ -155,6 +175,9 @@ private:
 	UPROPERTY(Transient)
 	TMap<uint32, FMassViewerHandle> ViewerMap;
 
+	UPROPERTY(Transient)
+	TArray<TObjectPtr<AActor>> RegisteredActorViewers;
+
 	uint64 LastSynchronizedFrame = 0;
 
 	/** Viewer serial number counter */
@@ -172,6 +195,9 @@ private:
 	FOnViewerAdded OnViewerAddedDelegate;
 	FOnViewerRemoved OnViewerRemovedDelegate;
 
+public:
+	UE_DEPRECATED(5.4, "GetViewerHandleFromPlayerController is deprecated. Use more teneric GetViewerHandleFromActor")
+	FMassViewerHandle GetViewerHandleFromPlayerController(const APlayerController* PlayerController) const;
 };
 
 template<>
diff --git a/Engine/Plugins/Runtime/MassGameplay/Source/MassReplication/Private/MassReplicationSubsystem.cpp b/Engine/Plugins/Runtime/MassGameplay/Source/MassReplication/Private/MassReplicationSubsystem.cpp
--- a/Engine/Plugins/Runtime/MassGameplay/Source/MassReplication/Private/MassReplicationSubsystem.cpp
+++ b/Engine/Plugins/Runtime/MassGameplay/Source/MassReplication/Private/MassReplicationSubsystem.cpp
@@ -197,7 +197,8 @@ bool UMassReplicationSubsystem::SynchronizeClients(const TArray<FViewerInfo>& Vi
 			{
 //no need to check netconnection if UE_ALLOW_DEBUG_REPLICATION_BUBBLES_STANDALONE
 #if !UE_ALLOW_DEBUG_REPLICATION_BUBBLES_STANDALONE
-				if (UE::Mass::Replication::HasParentNetConnection(Viewer.PlayerController))
+				APlayerController* ViewerAsPlayerController = Viewer.GetPlayerController();
+				if (UE::Mass::Replication::HasParentNetConnection(ViewerAsPlayerController))
 #endif
 				{
 					// check if the controller already exists by trying to remove it from the map which was filled up with controllers we were tracking
@@ -205,7 +206,7 @@ bool UMassReplicationSubsystem::SynchronizeClients(const TArray<FViewerInfo>& Vi
 					{
 						// If not add it to ClientsToAdd. Its important AddClient isn't called until necessary Clients are removed, as we may reuse 
 						// array indices that are already in use.
-						ClientsToAdd.Emplace(Viewer.Handle, Viewer.PlayerController);
+						ClientsToAdd.Emplace(Viewer.Handle, ViewerAsPlayerController);
 					}
 				}
 			}
@@ -304,12 +305,13 @@ void UMassReplicationSubsystem::SynchronizeClientViewers(const TArray<FViewerInf
 			if (Viewer.Handle.IsValid())
 			{
 				// we are only processing child UNetConnections that have a valid APlayerController OwningActor
-				const APlayerController* ParentController = UE::Mass::Replication::GetParentControllerFromChildNetConnection(Viewer.PlayerController);
+				APlayerController* ViewerAsPlayerController = Viewer.GetPlayerController();
+				const APlayerController* ParentController = UE::Mass::Replication::GetParentControllerFromChildNetConnection(ViewerAsPlayerController);
 
 				// check if the parent controller is valid and already exists
-				if (ParentController && (ClientViewerMap.Find(Viewer.PlayerController) == nullptr))
+				if (ParentController && (ClientViewerMap.Find(ViewerAsPlayerController) == nullptr))
 				{
-					FMassViewerHandle ParentViewerHandle = MassLODSubsystem->GetViewerHandleFromPlayerController(ParentController);
+					FMassViewerHandle ParentViewerHandle = MassLODSubsystem->GetViewerHandleFromActor(*ParentController);
 
 					if (ensureMsgf(ParentViewerHandle.IsValid(), TEXT("MassLODSubsystem handles are out of sync with PlayerController NetConnections!")))
 					{
@@ -322,7 +324,7 @@ void UMassReplicationSubsystem::SynchronizeClientViewers(const TArray<FViewerInf
 						check(ClientHandleManager.IsValidHandle(ParentViewerClientPair.ClientHandle));
 
 						// remove APlayerController from the ClientViewerMap and Add the viewer to the ClientsReplicationInfo
-						ClientViewerMap.Remove(Viewer.PlayerController);
+						ClientViewerMap.Remove(ViewerAsPlayerController);
 
 						FMassClientReplicationInfo& ClientReplicationInfo = ClientsReplicationInfo[ParentViewerClientPair.ClientHandle.GetIndex()];
 
