Fixed MassLODCalculator not handling null InLODMaxCount

diff --git a/Engine/Plugins/Runtime/MassGameplay/Source/MassLOD/Public/MassLODCalculator.h b/Engine/Plugins/Runtime/MassGameplay/Source/MassLOD/Public/MassLODCalculator.h
--- a/Engine/Plugins/Runtime/MassGameplay/Source/MassLOD/Public/MassLODCalculator.h
+++ b/Engine/Plugins/Runtime/MassGameplay/Source/MassLOD/Public/MassLODCalculator.h
@@ -20,7 +20,7 @@ public:
 	 * Initializes the LOD calculator, needed to be called once at initialization time
 	 * @Param InBaseLODDistance distances used to calculate LOD
 	 * @Param InBufferHysteresisOnFOVRatio distance hysteresis used to calculate LOD
-	 * @Param InLODMaxCount the maximum count for each LOD
+	 * @Param InLODMaxCount the maximum count for each LOD - Supports nullptr being passed in now and will put INT_MAX everywhere by default
 	 * @Param InLODMaxCountPerViewer the maximum count for each LOD per viewer (Only when FLODLogic::bMaximizeCountPerViewer is enabled)
 	 * @Param InVisibleDistanceToFrustum is the distance from the frustum to start considering this entity is visible (Only when FLODLogic::bDoVisibilityLogic is enabled)
 	 * @Param InVisibleDistanceToFrustumHysteresis once visible, what extra distance the entity need to be before considered not visible anymore (Only when FLODLogic::bDoVisibilityLogic is enabled)
@@ -306,7 +306,9 @@ void TMassLODCalculator<FLODLogic>::Initialize(const float InBaseLODDistance[EMa
 	for (int x = 0; x < EMassLOD::Max; x++)
 	{
 		BaseLODDistance[x] = InBaseLODDistance[x];
-		LODMaxCount[x] = InLODMaxCount[x];
+		
+		// @todo Treat InLODMaxCount as a possible nullptr by default for this Initialize function, would need to come as an option from FLODLogic as well
+		LODMaxCount[x] = (InLODMaxCount != nullptr) ? InLODMaxCount[x] : INT_MAX;
 		if (FLODLogic::bDoVisibilityLogic && InVisibleLODDistance)
 		{
 			VisibleLODDistance[x] = InVisibleLODDistance[x];
