Engine-level Mass change to support LOD calculations based solely on distance to viewers.

New MassDistanceLODProcessor
* Added relevant Fragment and Tags to enable it
* Added new FLODLogic for distance
* Added new LOD Params

New Dedicated Visualization Traits for DistanceLOD
* MassDistanceVisualizationTrait
* MassStationaryDistanceVisualizationTrait

LOD Info Collector to account strictly for the distance, not the Frustum
* Needed a new dedicated Tag
* Needed a new dedicated Processor

diff --git a/Engine/Plugins/Runtime/MassGameplay/Source/MassLOD/Private/MassLODDistanceCollectorProcessor.cpp b/Engine/Plugins/Runtime/MassGameplay/Source/MassLOD/Private/MassLODDistanceCollectorProcessor.cpp
new file mode 100644
--- /dev/null
+++ b/Engine/Plugins/Runtime/MassGameplay/Source/MassLOD/Private/MassLODDistanceCollectorProcessor.cpp
@@ -0,0 +1,102 @@
+// Copyright Epic Games, Inc. All Rights Reserved.
+
+#include "MassLODDistanceCollectorProcessor.h"
+#include "MassLODUtils.h"
+#include "MassCommonFragments.h"
+#include "MassExecutionContext.h"
+#include "Engine/World.h"
+#include "MassSimulationLOD.h"
+
+
+//-----------------------------------------------------------------------------
+// UMassLODDistanceCollectorProcessor
+//-----------------------------------------------------------------------------
+UMassLODDistanceCollectorProcessor::UMassLODDistanceCollectorProcessor()
+{
+	bAutoRegisterWithProcessingPhases = true;
+
+	ExecutionFlags = (int32)EProcessorExecutionFlags::All;
+
+	ExecutionOrder.ExecuteInGroup = UE::Mass::ProcessorGroupNames::LODCollector;
+	ExecutionOrder.ExecuteAfter.Add(UE::Mass::ProcessorGroupNames::SyncWorldToMass);
+}
+
+void UMassLODDistanceCollectorProcessor::ConfigureQueries()
+{
+	FMassEntityQuery BaseQuery;
+	BaseQuery.AddTagRequirement<FMassCollectDistanceLODViewerInfoTag>(EMassFragmentPresence::All);
+	BaseQuery.AddRequirement<FTransformFragment>(EMassFragmentAccess::ReadOnly);
+	BaseQuery.AddRequirement<FMassViewerInfoFragment>(EMassFragmentAccess::ReadWrite);
+	BaseQuery.AddChunkRequirement<FMassSimulationVariableTickChunkFragment>(EMassFragmentAccess::ReadOnly, EMassFragmentPresence::Optional);
+	BaseQuery.AddChunkRequirement<FMassVisualizationChunkFragment>(EMassFragmentAccess::ReadOnly, EMassFragmentPresence::Optional);
+	BaseQuery.SetChunkFilter([](const FMassExecutionContext& Context)
+	{
+		return FMassVisualizationChunkFragment::IsChunkHandledThisFrame(Context)
+			|| FMassSimulationVariableTickChunkFragment::IsChunkHandledThisFrame(Context);
+	});
+
+	EntityQuery_RelevantRangeAndOnLOD = BaseQuery;
+	EntityQuery_RelevantRangeAndOnLOD.AddTagRequirement<FMassVisibilityCulledByDistanceTag>(EMassFragmentPresence::None);
+	EntityQuery_RelevantRangeAndOnLOD.AddTagRequirement<FMassOffLODTag>(EMassFragmentPresence::None);
+	EntityQuery_RelevantRangeAndOnLOD.RegisterWithProcessor(*this);
+
+	EntityQuery_RelevantRangeOnly = BaseQuery;
+	EntityQuery_RelevantRangeOnly.AddTagRequirement<FMassVisibilityCulledByDistanceTag>(EMassFragmentPresence::None);
+	EntityQuery_RelevantRangeOnly.AddTagRequirement<FMassOffLODTag>(EMassFragmentPresence::All);
+	EntityQuery_RelevantRangeOnly.RegisterWithProcessor(*this);
+
+	EntityQuery_OnLODOnly = BaseQuery;
+	EntityQuery_OnLODOnly.AddTagRequirement<FMassVisibilityCulledByDistanceTag>(EMassFragmentPresence::All);
+	EntityQuery_OnLODOnly.AddTagRequirement<FMassOffLODTag>(EMassFragmentPresence::None);
+	EntityQuery_OnLODOnly.RegisterWithProcessor(*this);
+
+	EntityQuery_NotRelevantRangeAndOffLOD = BaseQuery;
+	EntityQuery_NotRelevantRangeAndOffLOD.AddTagRequirement<FMassVisibilityCulledByDistanceTag>(EMassFragmentPresence::All);
+	EntityQuery_NotRelevantRangeAndOffLOD.AddTagRequirement<FMassOffLODTag>(EMassFragmentPresence::All);
+	EntityQuery_NotRelevantRangeAndOffLOD.RegisterWithProcessor(*this);
+
+	ProcessorRequirements.AddSubsystemRequirement<UMassLODSubsystem>(EMassFragmentAccess::ReadOnly);
+}
+
+template <bool bLocalViewersOnly>
+void UMassLODDistanceCollectorProcessor::CollectLODForChunk(FMassExecutionContext& Context)
+{
+	TConstArrayView<FTransformFragment> LocationList = Context.GetFragmentView<FTransformFragment>();
+	TArrayView<FMassViewerInfoFragment> ViewerInfoList = Context.GetMutableFragmentView<FMassViewerInfoFragment>();
+
+	Collector.CollectLODInfo<FTransformFragment, FMassViewerInfoFragment, bLocalViewersOnly, true/*bCollectDistanceToViewer*/>(Context, LocationList, ViewerInfoList);
+}
+
+template <bool bLocalViewersOnly>
+void UMassLODDistanceCollectorProcessor::ExecuteInternal(FMassEntityManager& EntityManager, FMassExecutionContext& Context)
+{
+	{
+		TRACE_CPUPROFILER_EVENT_SCOPE(Close);
+		EntityQuery_RelevantRangeAndOnLOD.ForEachEntityChunk(EntityManager, Context, [this](FMassExecutionContext& Context) { CollectLODForChunk<bLocalViewersOnly>(Context); });
+		EntityQuery_RelevantRangeOnly.ForEachEntityChunk(EntityManager, Context, [this](FMassExecutionContext& Context) { CollectLODForChunk<bLocalViewersOnly>(Context); });
+		EntityQuery_OnLODOnly.ForEachEntityChunk(EntityManager, Context, [this](FMassExecutionContext& Context) { CollectLODForChunk<bLocalViewersOnly>(Context); });
+	}
+
+	{
+		TRACE_CPUPROFILER_EVENT_SCOPE(Far);
+		EntityQuery_NotRelevantRangeAndOffLOD.ForEachEntityChunk(EntityManager, Context, [this](FMassExecutionContext& Context) { CollectLODForChunk<bLocalViewersOnly>(Context); });
+	}
+}
+
+void UMassLODDistanceCollectorProcessor::Execute(FMassEntityManager& EntityManager, FMassExecutionContext& Context)
+{
+	const UMassLODSubsystem& LODSubsystem = Context.GetSubsystemChecked<UMassLODSubsystem>();
+	const TArray<FViewerInfo>& Viewers = LODSubsystem.GetViewers();
+	Collector.PrepareExecution(Viewers);
+
+	UWorld* World = EntityManager.GetWorld();
+	check(World);
+	if (World->IsNetMode(NM_DedicatedServer))
+	{
+		ExecuteInternal<false/*bLocalViewersOnly*/>(EntityManager, Context);
+	}
+	else
+	{
+		ExecuteInternal<true/*bLocalViewersOnly*/>(EntityManager, Context);
+	}
+}
diff --git a/Engine/Plugins/Runtime/MassGameplay/Source/MassLOD/Private/MassLODTrait.cpp b/Engine/Plugins/Runtime/MassGameplay/Source/MassLOD/Private/MassLODTrait.cpp
--- a/Engine/Plugins/Runtime/MassGameplay/Source/MassLOD/Private/MassLODTrait.cpp
+++ b/Engine/Plugins/Runtime/MassGameplay/Source/MassLOD/Private/MassLODTrait.cpp
@@ -9,6 +9,9 @@
 #include "MassEntityUtils.h"
 
 
+//-----------------------------------------------------------------------------
+// UMassLODCollectorTrait
+//-----------------------------------------------------------------------------
 void UMassLODCollectorTrait::BuildTemplate(FMassEntityTemplateBuildContext& BuildContext, const UWorld& World) const
 {
 	BuildContext.AddFragment<FMassViewerInfoFragment>();
@@ -16,6 +19,19 @@ void UMassLODCollectorTrait::BuildTemplate(FMassEntityTemplateBuildContext& Buil
 	BuildContext.RequireFragment<FTransformFragment>();
 }
 
+//-----------------------------------------------------------------------------
+// UMassDistanceLODCollectorTrait
+//-----------------------------------------------------------------------------
+void UMassDistanceLODCollectorTrait::BuildTemplate(FMassEntityTemplateBuildContext& BuildContext, const UWorld& World) const
+{
+	BuildContext.AddFragment<FMassViewerInfoFragment>();
+	BuildContext.AddTag<FMassCollectDistanceLODViewerInfoTag>();
+	BuildContext.RequireFragment<FTransformFragment>();
+}
+
+//-----------------------------------------------------------------------------
+// UMassSimulationLODTrait
+//-----------------------------------------------------------------------------
 void UMassSimulationLODTrait::BuildTemplate(FMassEntityTemplateBuildContext& BuildContext, const UWorld& World) const
 {
 	BuildContext.RequireFragment<FMassViewerInfoFragment>();
diff --git a/Engine/Plugins/Runtime/MassGameplay/Source/MassLOD/Public/MassLODDistanceCollectorProcessor.h b/Engine/Plugins/Runtime/MassGameplay/Source/MassLOD/Public/MassLODDistanceCollectorProcessor.h
new file mode 100644
--- /dev/null
+++ b/Engine/Plugins/Runtime/MassGameplay/Source/MassLOD/Public/MassLODDistanceCollectorProcessor.h
@@ -0,0 +1,45 @@
+// Copyright Epic Games, Inc. All Rights Reserved.
+
+#pragma once
+
+#include "MassEntityTypes.h"
+#include "MassProcessor.h"
+#include "MassLODCollector.h"
+#include "MassLODLogic.h"
+#include "MassLODDistanceCollectorProcessor.generated.h"
+
+/*
+ * LOD Distance collector which combines collection of LOD information for both Viewer and Simulation LODing.
+ * This collector cares only about the entities' distance to LOD viewer location, nothing else. 
+ * Matches MassDistanceLODProcessor logic which uses the same Calculator LODLogic
+ */
+UCLASS(meta = (DisplayName = "LOD Distance Collector"))
+class MASSLOD_API UMassLODDistanceCollectorProcessor : public UMassProcessor
+{
+	GENERATED_BODY()
+
+public:
+	UMassLODDistanceCollectorProcessor();
+
+protected:
+	virtual void ConfigureQueries() override;
+	virtual void Execute(FMassEntityManager& EntityManager, FMassExecutionContext& Context) override;
+
+	template <bool bLocalViewersOnly>
+	void CollectLODForChunk(FMassExecutionContext& Context);
+
+	template <bool bLocalViewersOnly>
+	void ExecuteInternal(FMassEntityManager& EntityManager, FMassExecutionContext& Context);
+
+	TMassLODCollector<FMassDistanceLODLogic> Collector;
+
+	// Queries for visualization and simulation calculation
+	/** All entities that are in relevant range and are On LOD*/
+	FMassEntityQuery EntityQuery_RelevantRangeAndOnLOD;
+	/** All entities that are in relevant range but are Off LOD */
+	FMassEntityQuery EntityQuery_RelevantRangeOnly;
+	/** All entities that are NOT in relevant range but are On LOD */
+	FMassEntityQuery EntityQuery_OnLODOnly;
+	/** All entities that are Not in relevant range and are at Off LOD */
+	FMassEntityQuery EntityQuery_NotRelevantRangeAndOffLOD;
+};
diff --git a/Engine/Plugins/Runtime/MassGameplay/Source/MassLOD/Public/MassLODFragments.h b/Engine/Plugins/Runtime/MassGameplay/Source/MassLOD/Public/MassLODFragments.h
--- a/Engine/Plugins/Runtime/MassGameplay/Source/MassLOD/Public/MassLODFragments.h
+++ b/Engine/Plugins/Runtime/MassGameplay/Source/MassLOD/Public/MassLODFragments.h
@@ -105,6 +105,15 @@ struct FMassCollectLODViewerInfoTag : public FMassTag
 	GENERATED_BODY();
 };
 
+/*
+ * Tag to use to trigger the collector processor that uses the LODCollector without Visibility Logic, so strictly based of distance
+ */
+USTRUCT()
+struct FMassCollectDistanceLODViewerInfoTag : public FMassTag
+{
+	GENERATED_BODY();
+};
+
 USTRUCT()
 struct MASSLOD_API FMassVisibilityCanBeSeenTag : public FMassTag
 {
diff --git a/Engine/Plugins/Runtime/MassGameplay/Source/MassLOD/Public/MassLODLogic.h b/Engine/Plugins/Runtime/MassGameplay/Source/MassLOD/Public/MassLODLogic.h
--- a/Engine/Plugins/Runtime/MassGameplay/Source/MassLOD/Public/MassLODLogic.h
+++ b/Engine/Plugins/Runtime/MassGameplay/Source/MassLOD/Public/MassLODLogic.h
@@ -61,6 +61,20 @@ struct FMassCombinedLODLogic : public FLODDefaultLogic
 	};
 };
 
+/** Simplest version of RepresentationLODLogic strictly based on Distance to Viewer
+ *	Compared to FMassRepresentationLODLogic, we:
+ *	* Do not care about doing the Visibility Logic
+ *	* For now we will keep the Significance computation as it could allow for finer grained control later on
+*/
+struct FMassDistanceLODLogic : public FLODDefaultLogic
+{
+	enum
+	{
+		bCalculateLODSignificance = true,
+		bLocalViewersOnly = true,
+	};
+};
+
 /**
  * TMassLODCollector outputs
  *
diff --git a/Engine/Plugins/Runtime/MassGameplay/Source/MassLOD/Public/MassLODTrait.h b/Engine/Plugins/Runtime/MassGameplay/Source/MassLOD/Public/MassLODTrait.h
--- a/Engine/Plugins/Runtime/MassGameplay/Source/MassLOD/Public/MassLODTrait.h
+++ b/Engine/Plugins/Runtime/MassGameplay/Source/MassLOD/Public/MassLODTrait.h
@@ -16,6 +16,16 @@ public:
 	virtual void BuildTemplate(FMassEntityTemplateBuildContext& BuildContext, const UWorld& World) const override;
 };
 
+// Simplest version of UMassLODCollectorTrait that will ensure collection strictly based on Distance from Viewer
+UCLASS(meta = (DisplayName = "DistanceLODCollector"))
+class MASSLOD_API UMassDistanceLODCollectorTrait : public UMassEntityTraitBase
+{
+	GENERATED_BODY()
+
+public:
+	virtual void BuildTemplate(FMassEntityTemplateBuildContext& BuildContext, const UWorld& World) const override;
+};
+
 UCLASS(meta = (DisplayName = "SimulationLOD"))
 class MASSLOD_API UMassSimulationLODTrait : public UMassEntityTraitBase
 {
diff --git a/Engine/Plugins/Runtime/MassGameplay/Source/MassRepresentation/Private/MassDistanceLODProcessor.cpp b/Engine/Plugins/Runtime/MassGameplay/Source/MassRepresentation/Private/MassDistanceLODProcessor.cpp
new file mode 100644
--- /dev/null
+++ b/Engine/Plugins/Runtime/MassGameplay/Source/MassRepresentation/Private/MassDistanceLODProcessor.cpp
@@ -0,0 +1,108 @@
+// Copyright Epic Games, Inc. All Rights Reserved.
+
+#include "MassDistanceLODProcessor.h"
+#include "MassRepresentationDebug.h"
+#include "MassExecutionContext.h"
+
+UMassDistanceLODProcessor::UMassDistanceLODProcessor()
+{
+	bAutoRegisterWithProcessingPhases = true;
+	ExecutionFlags = (int32)(EProcessorExecutionFlags::All);
+	ExecutionOrder.ExecuteInGroup = UE::Mass::ProcessorGroupNames::LOD;
+	ExecutionOrder.ExecuteAfter.Add(UE::Mass::ProcessorGroupNames::LODCollector);
+}
+
+void UMassDistanceLODProcessor::ConfigureQueries()
+{
+	FMassEntityQuery BaseQuery;
+	BaseQuery.AddTagRequirement<FMassDistanceLODProcessorTag>(EMassFragmentPresence::All);
+	BaseQuery.AddRequirement<FMassViewerInfoFragment>(EMassFragmentAccess::ReadOnly);
+	BaseQuery.AddRequirement<FMassRepresentationLODFragment>(EMassFragmentAccess::ReadWrite);
+	BaseQuery.AddRequirement<FTransformFragment>(EMassFragmentAccess::ReadOnly);
+	BaseQuery.AddConstSharedRequirement<FMassDistanceLODParameters>();
+	BaseQuery.AddSharedRequirement<FMassDistanceLODSharedFragment>(EMassFragmentAccess::ReadWrite);
+
+	CloseEntityQuery = BaseQuery;
+	CloseEntityQuery.AddTagRequirement<FMassVisibilityCulledByDistanceTag>(EMassFragmentPresence::None);
+	CloseEntityQuery.RegisterWithProcessor(*this);
+
+	FarEntityQuery = BaseQuery;
+	FarEntityQuery.AddTagRequirement<FMassVisibilityCulledByDistanceTag>(EMassFragmentPresence::All);
+	FarEntityQuery.AddChunkRequirement<FMassVisualizationChunkFragment>(EMassFragmentAccess::ReadOnly);
+	FarEntityQuery.SetChunkFilter(&FMassVisualizationChunkFragment::ShouldUpdateVisualizationForChunk);
+	FarEntityQuery.RegisterWithProcessor(*this);
+
+	DebugEntityQuery = BaseQuery;
+	DebugEntityQuery.RegisterWithProcessor(*this);
+
+	ProcessorRequirements.AddSubsystemRequirement<UMassLODSubsystem>(EMassFragmentAccess::ReadOnly);
+}
+
+void UMassDistanceLODProcessor::Execute(FMassEntityManager& EntityManager, FMassExecutionContext& Context)
+{
+	if (bForceOFFLOD)
+	{
+		CloseEntityQuery.ForEachEntityChunk(EntityManager, Context, [this](FMassExecutionContext& Context)
+		{
+			FMassDistanceLODSharedFragment& LODSharedFragment = Context.GetMutableSharedFragment<FMassDistanceLODSharedFragment>();
+			TArrayView<FMassRepresentationLODFragment> RepresentationLODList = Context.GetMutableFragmentView<FMassRepresentationLODFragment>();
+			LODSharedFragment.LODCalculator.ForceOffLOD(Context, RepresentationLODList);
+		});
+		return;
+	}
+
+	{
+		TRACE_CPUPROFILER_EVENT_SCOPE(PrepareExecution)
+		const UMassLODSubsystem& LODSubsystem = Context.GetSubsystemChecked<UMassLODSubsystem>();
+		const TArray<FViewerInfo>& Viewers = LODSubsystem.GetViewers();
+		EntityManager.ForEachSharedFragment<FMassDistanceLODSharedFragment>([this, &Viewers](FMassDistanceLODSharedFragment& LODSharedFragment)
+		{
+			if (FilterTag == LODSharedFragment.FilterTag)
+			{
+				LODSharedFragment.LODCalculator.PrepareExecution(Viewers);
+			}
+		});
+	}
+
+	{
+		TRACE_CPUPROFILER_EVENT_SCOPE(CalculateLOD)
+
+		auto CalculateLOD = [this](FMassExecutionContext& Context)
+		{
+			FMassDistanceLODSharedFragment& LODSharedFragment = Context.GetMutableSharedFragment<FMassDistanceLODSharedFragment>();
+			TArrayView<FMassRepresentationLODFragment> RepresentationLODList = Context.GetMutableFragmentView<FMassRepresentationLODFragment>();
+			TConstArrayView<FMassViewerInfoFragment> ViewerInfoList = Context.GetFragmentView<FMassViewerInfoFragment>();
+			LODSharedFragment.LODCalculator.CalculateLOD(Context, ViewerInfoList, RepresentationLODList);
+		};
+		CloseEntityQuery.ForEachEntityChunk(EntityManager, Context, CalculateLOD);
+		FarEntityQuery.ForEachEntityChunk(EntityManager, Context, CalculateLOD);
+	}
+
+#if WITH_MASSGAMEPLAY_DEBUG
+	// Optional debug display
+	if (UE::Mass::Representation::Debug::DebugRepresentationLOD == 1 || UE::Mass::Representation::Debug::DebugRepresentationLOD >= 3)
+	{
+		TRACE_CPUPROFILER_EVENT_SCOPE(DebugDisplayLOD)
+		UWorld* World = EntityManager.GetWorld();
+		DebugEntityQuery.ForEachEntityChunk(EntityManager, Context, [World](FMassExecutionContext& Context)
+		{
+			FMassDistanceLODSharedFragment& LODSharedFragment = Context.GetMutableSharedFragment<FMassDistanceLODSharedFragment>();
+			TConstArrayView<FMassRepresentationLODFragment> RepresentationLODList = Context.GetFragmentView<FMassRepresentationLODFragment>();
+			TConstArrayView<FTransformFragment> TransformList = Context.GetFragmentView<FTransformFragment>();
+			LODSharedFragment.LODCalculator.DebugDisplaySignificantLOD(Context, RepresentationLODList, TransformList, World, UE::Mass::Representation::Debug::DebugRepresentationLODMaxSignificance);
+		});
+	}
+	// Optional vislog
+	if (UE::Mass::Representation::Debug::DebugRepresentationLOD >= 2)
+	{
+		TRACE_CPUPROFILER_EVENT_SCOPE(VisLogLOD)
+		DebugEntityQuery.ForEachEntityChunk(EntityManager, Context, [this](FMassExecutionContext& Context)
+		{
+			FMassDistanceLODSharedFragment& LODSharedFragment = Context.GetMutableSharedFragment<FMassDistanceLODSharedFragment>();
+			TConstArrayView<FMassRepresentationLODFragment> RepresentationLODList = Context.GetFragmentView<FMassRepresentationLODFragment>();
+			TConstArrayView<FTransformFragment> TransformList = Context.GetFragmentView<FTransformFragment>();
+			LODSharedFragment.LODCalculator.VisLogSignificantLOD(Context, RepresentationLODList, TransformList, this, UE::Mass::Representation::Debug::DebugRepresentationLODMaxSignificance);
+		});
+	}
+#endif // WITH_MASSGAMEPLAY_DEBUG
+}
diff --git a/Engine/Plugins/Runtime/MassGameplay/Source/MassRepresentation/Private/MassDistanceVisualizationTrait.cpp b/Engine/Plugins/Runtime/MassGameplay/Source/MassRepresentation/Private/MassDistanceVisualizationTrait.cpp
new file mode 100644
--- /dev/null
+++ b/Engine/Plugins/Runtime/MassGameplay/Source/MassRepresentation/Private/MassDistanceVisualizationTrait.cpp
@@ -0,0 +1,96 @@
+﻿// Copyright Epic Games, Inc. All Rights Reserved.
+
+#include "MassDistanceVisualizationTrait.h"
+#include "MassEntityTemplateRegistry.h"
+#include "MassRepresentationSubsystem.h"
+#include "MassCommonFragments.h"
+#include "MassRepresentationFragments.h"
+#include "MassRepresentationActorManagement.h"
+#include "Engine/World.h"
+#include "MassLODFragments.h"
+#include "MassActorSubsystem.h"
+#include "MassEntityUtils.h"
+#include "MassDistanceLODProcessor.h"
+#include "MassRepresentationProcessor.h"
+
+
+UMassDistanceVisualizationTrait::UMassDistanceVisualizationTrait()
+{
+	RepresentationSubsystemClass = UMassRepresentationSubsystem::StaticClass();
+
+	Params.RepresentationActorManagementClass = UMassRepresentationActorManagement::StaticClass();
+	Params.LODRepresentation[EMassLOD::High] = EMassRepresentationType::HighResSpawnedActor;
+	Params.LODRepresentation[EMassLOD::Medium] = EMassRepresentationType::LowResSpawnedActor;
+	Params.LODRepresentation[EMassLOD::Low] = EMassRepresentationType::StaticMeshInstance;
+	Params.LODRepresentation[EMassLOD::Off] = EMassRepresentationType::None;
+
+	LODParams.LODDistance[EMassLOD::High] = 0.f;
+	LODParams.LODDistance[EMassLOD::Medium] = 1000.f;
+	LODParams.LODDistance[EMassLOD::Low] = 2500.f;
+	LODParams.LODDistance[EMassLOD::Off] = 10000.f;
+
+	LODParams.BufferHysteresisOnDistancePercentage = 10.0f;
+
+	bAllowServerSideVisualization = false;
+}
+
+void UMassDistanceVisualizationTrait::BuildTemplate(FMassEntityTemplateBuildContext& BuildContext, const UWorld& World) const
+{
+	// This should not be ran on NM_Server network mode
+	if (World.IsNetMode(NM_DedicatedServer) && !bAllowServerSideVisualization)
+	{
+		return;
+	}
+
+	BuildContext.RequireFragment<FMassViewerInfoFragment>();
+	BuildContext.RequireFragment<FTransformFragment>();
+	BuildContext.RequireFragment<FMassActorFragment>();
+
+	FMassEntityManager& EntityManager = UE::Mass::Utils::GetEntityManagerChecked(World);
+
+	UMassRepresentationSubsystem* RepresentationSubsystem = Cast<UMassRepresentationSubsystem>(World.GetSubsystemBase(RepresentationSubsystemClass));
+	if (RepresentationSubsystem == nullptr)
+	{
+		UE_LOG(LogMassRepresentation, Error, TEXT("Expecting a valid class for the representation subsystem"));
+		RepresentationSubsystem = UWorld::GetSubsystem<UMassRepresentationSubsystem>(&World);
+		check(RepresentationSubsystem);
+	}
+
+	FMassRepresentationSubsystemSharedFragment SubsystemSharedFragment;
+	SubsystemSharedFragment.RepresentationSubsystem = RepresentationSubsystem;
+	uint32 SubsystemHash = UE::StructUtils::GetStructCrc32(FConstStructView::Make(SubsystemSharedFragment));
+	FSharedStruct SubsystemFragment = EntityManager.GetOrCreateSharedFragmentByHash<FMassRepresentationSubsystemSharedFragment>(SubsystemHash, SubsystemSharedFragment);
+	BuildContext.AddSharedFragment(SubsystemFragment);
+
+	if (!Params.RepresentationActorManagementClass)
+	{
+		UE_LOG(LogMassRepresentation, Error, TEXT("Expecting a valid class for the representation actor management"));
+	}
+	FConstSharedStruct ParamsFragment = EntityManager.GetOrCreateConstSharedFragment(Params);
+	ParamsFragment.Get<const FMassRepresentationParameters>().ComputeCachedValues();
+	BuildContext.AddConstSharedFragment(ParamsFragment);
+
+	FMassRepresentationFragment& RepresentationFragment = BuildContext.AddFragment_GetRef<FMassRepresentationFragment>();
+	if (bRegisterStaticMeshDesc)
+	{
+		RepresentationFragment.StaticMeshDescIndex = RepresentationSubsystem->FindOrAddStaticMeshDesc(StaticMeshInstanceDesc);
+	}
+	RepresentationFragment.HighResTemplateActorIndex = HighResTemplateActor.Get() ? RepresentationSubsystem->FindOrAddTemplateActor(HighResTemplateActor.Get()) : INDEX_NONE;
+	RepresentationFragment.LowResTemplateActorIndex = LowResTemplateActor.Get() ? RepresentationSubsystem->FindOrAddTemplateActor(LowResTemplateActor.Get()) : INDEX_NONE;
+
+	FConstSharedStruct LODParamsFragment = EntityManager.GetOrCreateConstSharedFragment(LODParams);
+	BuildContext.AddConstSharedFragment(LODParamsFragment);
+
+	uint32 LODParamsHash = UE::StructUtils::GetStructCrc32(FConstStructView::Make(LODParams));
+	FSharedStruct LODSharedFragment = EntityManager.GetOrCreateSharedFragmentByHash<FMassDistanceLODSharedFragment>(LODParamsHash, LODParams);
+	BuildContext.AddSharedFragment(LODSharedFragment);
+
+	BuildContext.AddFragment<FMassRepresentationLODFragment>();
+	BuildContext.AddTag<FMassVisibilityCulledByDistanceTag>();
+	BuildContext.AddChunkFragment<FMassVisualizationChunkFragment>();
+
+	BuildContext.AddTag<FMassDistanceLODProcessorTag>();
+	BuildContext.AddTag<FMassVisualizationProcessorTag>();
+}
+
+
diff --git a/Engine/Plugins/Runtime/MassGameplay/Source/MassRepresentation/Private/MassRepresentationFragments.cpp b/Engine/Plugins/Runtime/MassGameplay/Source/MassRepresentation/Private/MassRepresentationFragments.cpp
--- a/Engine/Plugins/Runtime/MassGameplay/Source/MassRepresentation/Private/MassRepresentationFragments.cpp
+++ b/Engine/Plugins/Runtime/MassGameplay/Source/MassRepresentation/Private/MassRepresentationFragments.cpp
@@ -4,6 +4,13 @@
 
 FMassVisualizationLODSharedFragment::FMassVisualizationLODSharedFragment(const FMassVisualizationLODParameters& LODParams)
 {
-	LODCalculator.Initialize(LODParams.BaseLODDistance, LODParams.BufferHysteresisOnDistancePercentage / 100.f, LODParams.LODMaxCount, nullptr, LODParams.DistanceToFrustum, LODParams.DistanceToFrustumHysteresis, LODParams.VisibleLODDistance);
+	LODCalculator.Initialize(LODParams.BaseLODDistance, LODParams.BufferHysteresisOnDistancePercentage / 100.f, LODParams.LODMaxCount
+		, /*InLODMaxCountPerViewer=*/nullptr, LODParams.DistanceToFrustum, LODParams.DistanceToFrustumHysteresis, LODParams.VisibleLODDistance);
 	FilterTag = LODParams.FilterTag;
-}
\ No newline at end of file
+}
+
+FMassDistanceLODSharedFragment::FMassDistanceLODSharedFragment(const FMassDistanceLODParameters& LODParams)
+{
+	LODCalculator.Initialize(LODParams.LODDistance, LODParams.BufferHysteresisOnDistancePercentage / 100.f, /*InLODMaxCount=*/nullptr);
+	FilterTag = LODParams.FilterTag;
+}
diff --git a/Engine/Plugins/Runtime/MassGameplay/Source/MassRepresentation/Private/MassStationaryDistanceVisualizationTrait.cpp b/Engine/Plugins/Runtime/MassGameplay/Source/MassRepresentation/Private/MassStationaryDistanceVisualizationTrait.cpp
new file mode 100644
--- /dev/null
+++ b/Engine/Plugins/Runtime/MassGameplay/Source/MassRepresentation/Private/MassStationaryDistanceVisualizationTrait.cpp
@@ -0,0 +1,43 @@
+// Copyright Epic Games, Inc. All Rights Reserved.
+
+#include "MassStationaryDistanceVisualizationTrait.h"
+#include "MassEntityTemplateRegistry.h"
+#include "VisualLogger/VisualLogger.h"
+
+
+UMassStationaryDistanceVisualizationTrait::UMassStationaryDistanceVisualizationTrait(const FObjectInitializer& ObjectInitializer)
+{
+	bAllowServerSideVisualization = true;
+}
+
+void UMassStationaryDistanceVisualizationTrait::BuildTemplate(FMassEntityTemplateBuildContext& BuildContext, const UWorld& World) const
+{
+	bool bIssuesFound = false;
+	for (FMassStaticMeshInstanceVisualizationMeshDesc& MeshDesc : StaticMeshInstanceDesc.Meshes)
+	{
+		bIssuesFound = bIssuesFound || (MeshDesc.Mobility != EComponentMobility::Stationary);
+		MeshDesc.Mobility = EComponentMobility::Stationary;
+	}
+	UE_CVLOG_UELOG(bIssuesFound, this, LogMass, Log, TEXT("%s some Meshes' mobility has been set to non-Stationary. Theese settings will be overridden. "), *GetPathName());
+
+	Super::BuildTemplate(BuildContext, World);
+
+	BuildContext.AddTag<FMassStaticRepresentationTag>();
+}
+
+#if WITH_EDITOR
+void UMassStationaryDistanceVisualizationTrait::PostEditChangeProperty(struct FPropertyChangedEvent& PropertyChangedEvent)
+{
+	static const FName StaticMeshInstanceDescName = GET_MEMBER_NAME_CHECKED(UMassStationaryDistanceVisualizationTrait, StaticMeshInstanceDesc);
+
+	Super::PostEditChangeProperty(PropertyChangedEvent);
+
+	if (PropertyChangedEvent.MemberProperty && PropertyChangedEvent.MemberProperty->GetFName() == StaticMeshInstanceDescName)
+	{
+		for (FMassStaticMeshInstanceVisualizationMeshDesc& MeshDesc : StaticMeshInstanceDesc.Meshes)
+		{
+			MeshDesc.Mobility = EComponentMobility::Stationary;
+		}
+	}
+}
+#endif // WITH_EDITOR
diff --git a/Engine/Plugins/Runtime/MassGameplay/Source/MassRepresentation/Public/MassDistanceLODProcessor.h b/Engine/Plugins/Runtime/MassGameplay/Source/MassRepresentation/Public/MassDistanceLODProcessor.h
new file mode 100644
--- /dev/null
+++ b/Engine/Plugins/Runtime/MassGameplay/Source/MassRepresentation/Public/MassDistanceLODProcessor.h
@@ -0,0 +1,58 @@
+// Copyright Epic Games, Inc. All Rights Reserved.
+
+#pragma once
+
+#include "MassCommonFragments.h"
+#include "MassRepresentationFragments.h"
+#include "MassProcessor.h"
+#include "MassLODCalculator.h"
+
+#include "MassDistanceLODProcessor.generated.h"
+
+
+/** 
+ * Tag required by Distance LOD Processor to update LOD information. Removing the tag allows to support temporary 
+ * disabling of processing for individual entities.
+ */
+USTRUCT()
+struct MASSREPRESENTATION_API FMassDistanceLODProcessorTag : public FMassTag
+{
+	GENERATED_BODY();
+};
+
+UCLASS()
+class MASSREPRESENTATION_API UMassDistanceLODProcessor : public UMassProcessor
+{
+	GENERATED_BODY()
+
+public:
+	UMassDistanceLODProcessor();
+
+protected:
+
+	/** Configure the owned FMassEntityQuery instances to express processor's requirements */
+	virtual void ConfigureQueries() override;
+
+	/** 
+	 * Execution method for this processor 
+	 * @param EntitySubsystem is the system to execute the lambdas on each entity chunk
+	 * @param Context is the execution context to be passed when executing the lambdas
+	 */
+	virtual void Execute(FMassEntityManager& EntityManager, FMassExecutionContext& Context) override;
+	
+	/**
+	 * Forces Off LOD on all calculation
+	 * @param bForce to whether force or not Off LOD
+	 */
+	void ForceOffLOD(bool bForce) { bForceOFFLOD = bForce; }
+
+protected:
+	FMassEntityQuery CloseEntityQuery;
+	FMassEntityQuery FarEntityQuery;
+	FMassEntityQuery DebugEntityQuery;
+
+	bool bForceOFFLOD = false;
+
+	UPROPERTY(Transient)
+	TObjectPtr<const UScriptStruct> FilterTag = nullptr;
+};
\ No newline at end of file
diff --git a/Engine/Plugins/Runtime/MassGameplay/Source/MassRepresentation/Public/MassDistanceVisualizationTrait.h b/Engine/Plugins/Runtime/MassGameplay/Source/MassRepresentation/Public/MassDistanceVisualizationTrait.h
new file mode 100644
--- /dev/null
+++ b/Engine/Plugins/Runtime/MassGameplay/Source/MassRepresentation/Public/MassDistanceVisualizationTrait.h
@@ -0,0 +1,68 @@
+﻿// Copyright Epic Games, Inc. All Rights Reserved.
+
+#pragma once 
+
+#include "MassEntityTraitBase.h"
+#include "MassRepresentationTypes.h"
+#include "MassRepresentationFragments.h"
+#include "GameFramework/Actor.h"
+
+#include "MassDistanceVisualizationTrait.generated.h"
+
+class UMassRepresentationSubsystem;
+class UMassRepresentationActorManagement;
+class UMassProcessor;
+
+/** This class has been soft-deprecated. Use MassStationaryVisualizationTrait or MassMovableVisualizationTrait */
+UCLASS(meta=(DisplayName="DEPRECATED DistanceVisualization"))
+class MASSREPRESENTATION_API UMassDistanceVisualizationTrait : public UMassEntityTraitBase
+{
+	GENERATED_BODY()
+public:
+	UMassDistanceVisualizationTrait();
+
+	virtual void BuildTemplate(FMassEntityTemplateBuildContext& BuildContext, const UWorld& World) const override;
+
+	/** Instanced static mesh information for this agent */
+	UPROPERTY(EditAnywhere, Category = "Mass|Visual")
+	mutable FStaticMeshInstanceVisualizationDesc StaticMeshInstanceDesc;
+
+	/** Actor class of this agent when spawned in high resolution*/
+	UPROPERTY(EditAnywhere, Category = "Mass|Visual")
+	TSubclassOf<AActor> HighResTemplateActor;
+
+	/** Actor class of this agent when spawned in low resolution*/
+	UPROPERTY(EditAnywhere, Category = "Mass|Visual")
+	TSubclassOf<AActor> LowResTemplateActor;
+
+	/** Allow subclasses to override the representation subsystem to use */
+	UPROPERTY(EditAnywhere, Category = "Mass|Visual", meta = (EditCondition = "bCanModifyRepresentationSubsystemClass"))
+	TSubclassOf<UMassRepresentationSubsystem> RepresentationSubsystemClass;
+
+	/** Configuration parameters for the representation processor */
+	UPROPERTY(EditAnywhere, Category = "Mass|Visual")
+	FMassRepresentationParameters Params;
+
+	/** Configuration parameters for the Distance LOD processor */
+	UPROPERTY(EditAnywhere, Category = "Mass|Visual")
+	FMassDistanceLODParameters LODParams;
+
+	/** If set to true will result in the visualization-related fragments being added to server-size entities as well.
+	 *  By default only the clients require visualization fragments */
+	UPROPERTY(EditAnywhere, Category = "Mass|Visual")
+	bool bAllowServerSideVisualization = false;
+
+#if WITH_EDITORONLY_DATA
+	/** the property is marked like this to ensure it won't show up in UI */
+	UPROPERTY(EditDefaultsOnly, Category = "Mass|Visual")
+	bool bCanModifyRepresentationSubsystemClass = true;
+#endif // WITH_EDITORONLY_DATA
+
+protected:
+	/** 
+	 * Controls whether StaticMeshInstanceDesc gets registered via FindOrAddStaticMeshDesc call. Setting it to `false` 
+	 * can be useful for subclasses to avoid needlessly creating visualization data in RepresentationSubsystem, 
+	 * data that will never be used.
+	 */
+	bool bRegisterStaticMeshDesc = true;
+};
diff --git a/Engine/Plugins/Runtime/MassGameplay/Source/MassRepresentation/Public/MassRepresentationFragments.h b/Engine/Plugins/Runtime/MassGameplay/Source/MassRepresentation/Public/MassRepresentationFragments.h
--- a/Engine/Plugins/Runtime/MassGameplay/Source/MassRepresentation/Public/MassRepresentationFragments.h
+++ b/Engine/Plugins/Runtime/MassGameplay/Source/MassRepresentation/Public/MassRepresentationFragments.h
@@ -238,3 +238,45 @@ struct FMassVisualizationLODSharedFragment : public FMassSharedFragment
 	UPROPERTY(Transient)
 	TObjectPtr<const UScriptStruct> FilterTag = nullptr;
 };
+
+/** Simplest version of LOD Calculation based strictly on Distance parameters 
+ *	Compared to FMassVisualizationLODParameters, we:
+ *	* Only include a single set of LOD Distances (radial distance from viewer)
+ *	* we do not care about distance to Frustum
+ *	* we do not care about Max Count
+ */
+USTRUCT()
+struct FMassDistanceLODParameters : public FMassSharedFragment
+{
+	GENERATED_BODY()
+
+	/** Distances where each LOD becomes relevant */
+	UPROPERTY(EditAnywhere, Category = "Mass|LOD", config)
+	float LODDistance[EMassLOD::Max] = { 0.f, 1000.f, 2500.f, 10000.f };
+
+	UPROPERTY(EditAnywhere, Category = "Mass|LOD", meta = (ClampMin = "0.0", UIMin = "0.0"), config)
+	float BufferHysteresisOnDistancePercentage = 10.0f;
+
+	/** Filter these settings with specified tag */
+	UPROPERTY(EditAnywhere, Category = "Mass|LOD", meta = (BaseStruct = "/Script/MassEntity.MassTag"))
+	TObjectPtr<UScriptStruct> FilterTag = nullptr;
+};
+
+/** Simplest version of LOD Calculation based strictly on Distance parameters 
+ *	Compared to FMassVisualizationLODSharedFragment, we:
+ *	* Cannot Adjust the Distance from count
+ *	* We care about a MassLODCalculator with a new LOD logic that excludes Visibility computation
+ */
+USTRUCT()
+struct FMassDistanceLODSharedFragment : public FMassSharedFragment
+{
+	GENERATED_BODY()
+
+	FMassDistanceLODSharedFragment() = default;
+	FMassDistanceLODSharedFragment(const FMassDistanceLODParameters& LODParams);
+
+	TMassLODCalculator<FMassDistanceLODLogic> LODCalculator;
+
+	UPROPERTY(Transient)
+	TObjectPtr<const UScriptStruct> FilterTag = nullptr;
+};
diff --git a/Engine/Plugins/Runtime/MassGameplay/Source/MassRepresentation/Public/MassStationaryDistanceVisualizationTrait.h b/Engine/Plugins/Runtime/MassGameplay/Source/MassRepresentation/Public/MassStationaryDistanceVisualizationTrait.h
new file mode 100644
--- /dev/null
+++ b/Engine/Plugins/Runtime/MassGameplay/Source/MassRepresentation/Public/MassStationaryDistanceVisualizationTrait.h
@@ -0,0 +1,22 @@
+// Copyright Epic Games, Inc. All Rights Reserved.
+
+#pragma once 
+
+#include "MassDistanceVisualizationTrait.h"
+#include "MassStationaryDistanceVisualizationTrait.generated.h"
+
+
+UCLASS()
+class MASSREPRESENTATION_API UMassStationaryDistanceVisualizationTrait : public UMassDistanceVisualizationTrait
+{
+	GENERATED_BODY()
+public:
+	UMassStationaryDistanceVisualizationTrait(const FObjectInitializer& ObjectInitializer = FObjectInitializer::Get());
+
+	virtual void BuildTemplate(FMassEntityTemplateBuildContext& BuildContext, const UWorld& World) const override;
+
+protected:
+#if WITH_EDITOR
+	virtual void PostEditChangeProperty(struct FPropertyChangedEvent& PropertyChangedEvent) override;
+#endif // WITH_EDITOR
+};
