Added a config-controllable way to configure MassLODSubsystems to use player's pawn location/rotation instead of camera's (when a pawn is available, falling back to camera otherwise).

diff --git a/Engine/Plugins/Runtime/MassGameplay/Source/MassLOD/Private/MassLODSubsystem.cpp b/Engine/Plugins/Runtime/MassGameplay/Source/MassLOD/Private/MassLODSubsystem.cpp
--- a/Engine/Plugins/Runtime/MassGameplay/Source/MassLOD/Private/MassLODSubsystem.cpp
+++ b/Engine/Plugins/Runtime/MassGameplay/Source/MassLOD/Private/MassLODSubsystem.cpp
@@ -15,9 +15,7 @@
 #include "Editor/EditorEngine.h"
 #endif // WITH_EDITOR
 
-namespace UE
-{
-namespace MassLOD
+namespace UE::MassLOD
 {
 	FColor LODColors[] =
 	{
@@ -26,7 +24,45 @@ namespace MassLOD
 		FColor::Emerald,
 		FColor::White,
 	};
-}
+	
+#if WITH_MASSGAMEPLAY_DEBUG
+	namespace Debug
+	{
+		FAutoConsoleCommandWithWorldArgsAndOutputDevice ToggleUsePlayerLocationCmd(
+			TEXT("mass.debug.LODSubsystem.UsePlayerLocation"),
+			TEXT("Sets UMassSubsystem::bUsePlayerPawnLocationInsteadOfCamera. Note that this is a command that doesn't retain state and usually needs running both for the client and the server"),
+			FConsoleCommandWithWorldArgsAndOutputDeviceDelegate::CreateLambda([](const TArray<FString>& Args, UWorld* World, FOutputDevice& Ar)
+				{
+					if (!World)
+					{
+						UE_LOG(LogConsoleResponse, Display, TEXT("Error: invalid world"));
+						return;
+					}
+
+					UMassLODSubsystem* MassCrowdSubsystem = World->GetSubsystem<UMassLODSubsystem>();
+					if (MassCrowdSubsystem == nullptr)
+					{
+						UE_LOG(LogConsoleResponse, Display, TEXT("Error: Unable to fetch MassLODSubsystem instance"));
+						return;
+					}
+					
+					if (Args.Num() < 1)
+					{
+						UE_LOG(LogConsoleResponse, Display, TEXT("Error: Expecting 1 parameter"));
+						return;
+					}
+
+					bool bNewValue = false;
+					if (!LexTryParseString<bool>(bNewValue, *Args[0]))
+					{
+						UE_LOG(LogConsoleResponse, Display, TEXT("Error: parameter must be an integer or a boolean"));
+						return;
+					}
+
+					MassCrowdSubsystem->DebugSetUsePlayerPawnLocationInsteadOfCamera(bNewValue);
+				}));
+	}
+#endif // WITH_MASSGAMEPLAY_DEBUG
 }
 
 //-----------------------------------------------------------------------------
@@ -278,14 +314,26 @@ void UMassLODSubsystem::SynchronizeViewers()
 		{
 			ViewerInfo.bEnabled = !WorldPartition || ViewerAsPlayerController->bEnableStreamingSource;
 
-			FVector PlayerCameraLocation(ForceInitToZero);
-			FRotator PlayerCameraRotation(FRotator::ZeroRotator);
-			ViewerAsPlayerController->GetPlayerViewPoint(PlayerCameraLocation, PlayerCameraRotation);
-			ViewerInfo.Location = PlayerCameraLocation;
-			ViewerInfo.Rotation = PlayerCameraRotation;
+			// Note: Using bUsePlayerPawnLocationInsteadOfCamera will not work correctly with FOV based LOD, since the
+			// camera will be at wrong location.  
+			// @todo: separate "player location" and "view location", and use the player location on distance based LOD 
+			// calculations for stability, and view location in FOV based LOD for view precision.
+			if (bUsePlayerPawnLocationInsteadOfCamera && ViewerAsPlayerController->GetPawn())
+			{
+				ViewerInfo.Location = ViewerAsPlayerController->GetPawn()->GetActorLocation();
+				ViewerInfo.Rotation = ViewerAsPlayerController->GetPawn()->GetActorRotation();
+			}
+			else
+			{
+				FVector PlayerCameraLocation(ForceInitToZero);
+				FRotator PlayerCameraRotation(FRotator::ZeroRotator);
+				ViewerAsPlayerController->GetPlayerViewPoint(PlayerCameraLocation, PlayerCameraRotation);
+				ViewerInfo.Location = PlayerCameraLocation;
+				ViewerInfo.Rotation = PlayerCameraRotation;
+			}
 
 			// Try to fetch a more precise FOV
-			if(ViewerAsPlayerController->PlayerCameraManager)
+			if (ViewerAsPlayerController->PlayerCameraManager)
 			{
 				ViewerInfo.FOV = ViewerAsPlayerController->PlayerCameraManager->GetFOVAngle();
 
diff --git a/Engine/Plugins/Runtime/MassGameplay/Source/MassLOD/Public/MassLODSubsystem.h b/Engine/Plugins/Runtime/MassGameplay/Source/MassLOD/Public/MassLODSubsystem.h
--- a/Engine/Plugins/Runtime/MassGameplay/Source/MassLOD/Public/MassLODSubsystem.h
+++ b/Engine/Plugins/Runtime/MassGameplay/Source/MassLOD/Public/MassLODSubsystem.h
@@ -111,6 +111,10 @@ public:
 	void RegisterActorViewer(AActor& ActorViewer);
 	void UnregisterActorViewer(AActor& ActorViewer);
 
+#if WITH_MASSGAMEPLAY_DEBUG
+	void DebugSetUsePlayerPawnLocationInsteadOfCamera(const bool bInValue) { bUsePlayerPawnLocationInsteadOfCamera = bInValue; }
+#endif
+
 protected:
 	// USubsystem BEGIN
 	virtual void Initialize(FSubsystemCollectionBase& Collection) override;
@@ -152,17 +156,26 @@ protected:
 	void OnPlayerControllerEndPlay(AActor* Actor, EEndPlayReason::Type EndPlayReason);
 
 protected:
-	/** If true, all PlayerControllers will be gathered as viewers for LOD calcuations. */
+	/** If true, all PlayerControllers will be gathered as viewers for LOD calculations. */
 	UPROPERTY(EditDefaultsOnly, Category = "Mass|LOD", config)
 	uint8 bGatherPlayerControllers : 1 = true;
 
-	/** If true, all streaming sources will be gathered as viewers for LOD calcuations. */
+	/** If true, all streaming sources will be gathered as viewers for LOD calculations. */
 	UPROPERTY(EditDefaultsOnly, Category = "Mass|LOD", config)
 	uint8 bGatherStreamingSources : 1 = true;
 
+	/** Whether using non-player actors as LOD Viewers is supported. */
 	UPROPERTY(EditDefaultsOnly, Category = "Mass|LOD", config)
 	uint8 bAllowNonPlayerViwerActors : 1 = true;
 
+	/** 
+	 * If set to true will prefer to use Player-owned Pawn's location and rotation over Player's camera as the viewer's 
+	 * location and rotation.
+	 * Note that this works best with distance-only LOD and can introduce subtle inaccuracies if Frustum-based LOD is being used. 
+	 */
+	UPROPERTY(EditDefaultsOnly, Category = "Mass|LOD", config)
+	uint8 bUsePlayerPawnLocationInsteadOfCamera : 1 = false;
+
 private:
 	/** Removes a viewer to the list and send notification about removal */
 	void RemoveViewerInternal(const FMassViewerHandle& ViewerHandle);
