Quick and straightforward way to avoid frustum creation for Mass LOD collectors/calculators that don't use it.

diff --git a/Engine/Plugins/Runtime/MassGameplay/Source/MassLOD/Private/MassLODCollectorProcessor.cpp b/Engine/Plugins/Runtime/MassGameplay/Source/MassLOD/Private/MassLODCollectorProcessor.cpp
--- a/Engine/Plugins/Runtime/MassGameplay/Source/MassLOD/Private/MassLODCollectorProcessor.cpp
+++ b/Engine/Plugins/Runtime/MassGameplay/Source/MassLOD/Private/MassLODCollectorProcessor.cpp
@@ -60,7 +60,7 @@ void UMassLODCollectorProcessor::CollectLODForChunk(FMassExecutionContext& Conte
 	TConstArrayView<FTransformFragment> LocationList = Context.GetFragmentView<FTransformFragment>();
 	TArrayView<FMassViewerInfoFragment> ViewerInfoList = Context.GetMutableFragmentView<FMassViewerInfoFragment>();
 
-	Collector.CollectLODInfo<FTransformFragment, FMassViewerInfoFragment, bLocalViewersOnly, true/*bCollectDistanceToViewer*/>(Context, LocationList, ViewerInfoList);
+	Collector.CollectLODInfo<FTransformFragment, FMassViewerInfoFragment, bLocalViewersOnly, true/*bCollectDistanceToFrustum*/>(Context, LocationList, ViewerInfoList);
 }
 
 template <bool bLocalViewersOnly>
diff --git a/Engine/Plugins/Runtime/MassGameplay/Source/MassLOD/Private/MassLODDistanceCollectorProcessor.cpp b/Engine/Plugins/Runtime/MassGameplay/Source/MassLOD/Private/MassLODDistanceCollectorProcessor.cpp
--- a/Engine/Plugins/Runtime/MassGameplay/Source/MassLOD/Private/MassLODDistanceCollectorProcessor.cpp
+++ b/Engine/Plugins/Runtime/MassGameplay/Source/MassLOD/Private/MassLODDistanceCollectorProcessor.cpp
@@ -64,21 +64,30 @@ void UMassLODDistanceCollectorProcessor::CollectLODForChunk(FMassExecutionContex
 	TConstArrayView<FTransformFragment> LocationList = Context.GetFragmentView<FTransformFragment>();
 	TArrayView<FMassViewerInfoFragment> ViewerInfoList = Context.GetMutableFragmentView<FMassViewerInfoFragment>();
 
-	Collector.CollectLODInfo<FTransformFragment, FMassViewerInfoFragment, bLocalViewersOnly, true/*bCollectDistanceToViewer*/>(Context, LocationList, ViewerInfoList);
+	Collector.CollectLODInfo<FTransformFragment, FMassViewerInfoFragment, bLocalViewersOnly, false/*bCollectDistanceToFrustum*/>(Context, LocationList, ViewerInfoList);
 }
 
 template <bool bLocalViewersOnly>
 void UMassLODDistanceCollectorProcessor::ExecuteInternal(FMassEntityManager& EntityManager, FMassExecutionContext& Context)
 {
 	{
-		TRACE_CPUPROFILER_EVENT_SCOPE(Close);
-		EntityQuery_RelevantRangeAndOnLOD.ForEachEntityChunk(EntityManager, Context, [this](FMassExecutionContext& Context) { CollectLODForChunk<bLocalViewersOnly>(Context); });
-		EntityQuery_RelevantRangeOnly.ForEachEntityChunk(EntityManager, Context, [this](FMassExecutionContext& Context) { CollectLODForChunk<bLocalViewersOnly>(Context); });
-		EntityQuery_OnLODOnly.ForEachEntityChunk(EntityManager, Context, [this](FMassExecutionContext& Context) { CollectLODForChunk<bLocalViewersOnly>(Context); });
+		TRACE_CPUPROFILER_EVENT_SCOPE(LODDistanceCollector_Close);
+		{
+			TRACE_CPUPROFILER_EVENT_SCOPE(LODDistanceCollector_Close_RelevantRangeAndOnLOD);
+			EntityQuery_RelevantRangeAndOnLOD.ForEachEntityChunk(EntityManager, Context, [this](FMassExecutionContext& Context) { CollectLODForChunk<bLocalViewersOnly>(Context); });
+		}
+		{
+			TRACE_CPUPROFILER_EVENT_SCOPE(LODDistanceCollector_Close_OnLODOnly);
+			EntityQuery_RelevantRangeOnly.ForEachEntityChunk(EntityManager, Context, [this](FMassExecutionContext& Context) { CollectLODForChunk<bLocalViewersOnly>(Context); });
+		}
+		{
+			TRACE_CPUPROFILER_EVENT_SCOPE(LODDistanceCollector_Close_OnLODOnly);
+			EntityQuery_OnLODOnly.ForEachEntityChunk(EntityManager, Context, [this](FMassExecutionContext& Context) { CollectLODForChunk<bLocalViewersOnly>(Context); });
+		}
 	}
 
 	{
-		TRACE_CPUPROFILER_EVENT_SCOPE(Far);
+		TRACE_CPUPROFILER_EVENT_SCOPE(LODDistanceCollector_Far);
 		EntityQuery_NotRelevantRangeAndOffLOD.ForEachEntityChunk(EntityManager, Context, [this](FMassExecutionContext& Context) { CollectLODForChunk<bLocalViewersOnly>(Context); });
 	}
 }
diff --git a/Engine/Plugins/Runtime/MassGameplay/Source/MassLOD/Private/MassLODLogic.cpp b/Engine/Plugins/Runtime/MassGameplay/Source/MassLOD/Private/MassLODLogic.cpp
--- a/Engine/Plugins/Runtime/MassGameplay/Source/MassLOD/Private/MassLODLogic.cpp
+++ b/Engine/Plugins/Runtime/MassGameplay/Source/MassLOD/Private/MassLODLogic.cpp
@@ -28,21 +28,24 @@ void FMassLODBaseLogic::CacheViewerInformation(TConstArrayView<FViewerInfo> View
 			ViewerLOD.Location = Viewer.Location;
 			ViewerLOD.Direction = Viewer.Rotation.Vector();
 
-			const float HalfHorizontalFOVAngle = Viewer.FOV * 0.5f;
-			const float HalfVerticalFOVAngle = FMath::RadiansToDegrees(FMath::Atan(FMath::Tan(FMath::DegreesToRadians(HalfHorizontalFOVAngle)) * Viewer.AspectRatio));
-
-			const FVector RightPlaneNormal = Viewer.Rotation.RotateVector(FRotator(0.0f, HalfHorizontalFOVAngle, 0.0f).RotateVector(FVector::RightVector));
-			const FVector LeftPlaneNormal = Viewer.Rotation.RotateVector(FRotator(0.0f, -HalfHorizontalFOVAngle, 0.0f).RotateVector(FVector::LeftVector));
-			const FVector TopPlaneNormal = Viewer.Rotation.RotateVector(FRotator(HalfVerticalFOVAngle, 0.0f, 0.0f).RotateVector(FVector::UpVector));
-			const FVector BottomPlaneNormal = Viewer.Rotation.RotateVector(FRotator(-HalfVerticalFOVAngle, 0.0f, 0.0f).RotateVector(FVector::DownVector));
-			
-			TArray<FPlane, TInlineAllocator<6>> Planes;
-			Planes.Emplace(ViewerLOD.Location, RightPlaneNormal);
-			Planes.Emplace(ViewerLOD.Location, LeftPlaneNormal);
-			Planes.Emplace(ViewerLOD.Location, TopPlaneNormal);
-			Planes.Emplace(ViewerLOD.Location, BottomPlaneNormal);
-
-			ViewerLOD.Frustum = FConvexVolume(Planes);
+			if (bBuildFrustumData)
+			{
+				const float HalfHorizontalFOVAngle = Viewer.FOV * 0.5f;
+				const float HalfVerticalFOVAngle = FMath::RadiansToDegrees(FMath::Atan(FMath::Tan(FMath::DegreesToRadians(HalfHorizontalFOVAngle)) * Viewer.AspectRatio));
+
+				const FVector RightPlaneNormal = Viewer.Rotation.RotateVector(FRotator(0.0f, HalfHorizontalFOVAngle, 0.0f).RotateVector(FVector::RightVector));
+				const FVector LeftPlaneNormal = Viewer.Rotation.RotateVector(FRotator(0.0f, -HalfHorizontalFOVAngle, 0.0f).RotateVector(FVector::LeftVector));
+				const FVector TopPlaneNormal = Viewer.Rotation.RotateVector(FRotator(HalfVerticalFOVAngle, 0.0f, 0.0f).RotateVector(FVector::UpVector));
+				const FVector BottomPlaneNormal = Viewer.Rotation.RotateVector(FRotator(-HalfVerticalFOVAngle, 0.0f, 0.0f).RotateVector(FVector::DownVector));
+
+				TArray<FPlane, TInlineAllocator<6>> Planes;
+				Planes.Emplace(ViewerLOD.Location, RightPlaneNormal);
+				Planes.Emplace(ViewerLOD.Location, LeftPlaneNormal);
+				Planes.Emplace(ViewerLOD.Location, TopPlaneNormal);
+				Planes.Emplace(ViewerLOD.Location, BottomPlaneNormal);
+
+				ViewerLOD.Frustum = FConvexVolume(Planes);
+			}
 		}
 	}
 }
\ No newline at end of file
diff --git a/Engine/Plugins/Runtime/MassGameplay/Source/MassLOD/Public/MassLODCalculator.h b/Engine/Plugins/Runtime/MassGameplay/Source/MassLOD/Public/MassLODCalculator.h
--- a/Engine/Plugins/Runtime/MassGameplay/Source/MassLOD/Public/MassLODCalculator.h
+++ b/Engine/Plugins/Runtime/MassGameplay/Source/MassLOD/Public/MassLODCalculator.h
@@ -14,7 +14,9 @@
 template <typename FLODLogic = FLODDefaultLogic >
 struct TMassLODCalculator : public FMassLODBaseLogic
 {
-public:
+	TMassLODCalculator()
+		: FMassLODBaseLogic(/*bShouldBuildFrustumData=*/FLODLogic::bDoVisibilityLogic)
+	{}
 
 	/**
 	 * Initializes the LOD calculator, needed to be called once at initialization time
diff --git a/Engine/Plugins/Runtime/MassGameplay/Source/MassLOD/Public/MassLODCollector.h b/Engine/Plugins/Runtime/MassGameplay/Source/MassLOD/Public/MassLODCollector.h
--- a/Engine/Plugins/Runtime/MassGameplay/Source/MassLOD/Public/MassLODCollector.h
+++ b/Engine/Plugins/Runtime/MassGameplay/Source/MassLOD/Public/MassLODCollector.h
@@ -16,6 +16,10 @@
 template <typename FLODLogic = FLODDefaultLogic >
 struct TMassLODCollector : public FMassLODBaseLogic
 {
+	TMassLODCollector()
+		: FMassLODBaseLogic(/*bShouldBuildFrustumData=*/FLODLogic::bDoVisibilityLogic)
+	{}
+
 	/**
 	 * Prepares execution for the current frame, needed to be called before every execution
 	 * @Param Viewers is the array of all the known viewers
@@ -126,7 +130,7 @@ void TMassLODCollector<FLODLogic>::CollectLODInfo(FMassExecutionContext& Context
 				}
 				SetDistanceToViewerSq<bCollectDistancePerViewer>(EntityInfoPerViewer, ViewerIdx, DistanceToViewerSq);
 
-				if (bCollectDistanceToFrustum)
+				if constexpr (bCollectDistanceToFrustum)
 				{
 					const float DistanceToFrustum = Viewer.Frustum.DistanceTo(EntityLocation);
 					SetDistanceToFrustum<bCollectDistanceToFrustumPerViewer>(EntityInfoPerViewer, ViewerIdx, DistanceToFrustum);
diff --git a/Engine/Plugins/Runtime/MassGameplay/Source/MassLOD/Public/MassLODLogic.h b/Engine/Plugins/Runtime/MassGameplay/Source/MassLOD/Public/MassLODLogic.h
--- a/Engine/Plugins/Runtime/MassGameplay/Source/MassLOD/Public/MassLODLogic.h
+++ b/Engine/Plugins/Runtime/MassGameplay/Source/MassLOD/Public/MassLODLogic.h
@@ -160,6 +160,10 @@ struct FViewerLODInfo
  */
 struct MASSLOD_API FMassLODBaseLogic
 {
+	FMassLODBaseLogic(bool bShouldBuildFrustumData)
+		: bBuildFrustumData(bShouldBuildFrustumData)
+	{}
+
 protected:
 	void CacheViewerInformation(TConstArrayView<FViewerInfo> ViewerInfos);
 
@@ -184,4 +188,11 @@ protected:
 	DECLARE_CONDITIONAL_MEMBER_ARRAY_ACCESSORS(Condition, EMassVisibility, PrevVisibilityPerViewer);
 
 	TArray<FViewerLODInfo> Viewers;
+
+private:
+	/** 
+	 * Setting to false will prevent costly FViewerLODInfo.Frustum creation. Makes sense only if that data is not required.
+	 * Note that this property is not expected to be changed at runtime.
+	 */
+	bool bBuildFrustumData = true;
 };
\ No newline at end of file
diff --git a/Engine/Plugins/Runtime/MassGameplay/Source/MassLOD/Public/MassLODTickRateController.h b/Engine/Plugins/Runtime/MassGameplay/Source/MassLOD/Public/MassLODTickRateController.h
--- a/Engine/Plugins/Runtime/MassGameplay/Source/MassLOD/Public/MassLODTickRateController.h
+++ b/Engine/Plugins/Runtime/MassGameplay/Source/MassLOD/Public/MassLODTickRateController.h
@@ -14,6 +14,9 @@ template <typename TVariableTickChunkFragment, typename FLODLogic = FLODDefaultL
 struct TMassLODTickRateController : public FMassLODBaseLogic
 {
 public:
+	TMassLODTickRateController()
+		: FMassLODBaseLogic(/*bShouldBuildFrustumData=*/false)
+	{}
 
 	/**
 	 * Initializes the LOD trick rate controller, needed to be called once at initialization time (Only when FLODLogic::bDoVariableTickRate is enabled)
