Fixed Mass visualization kicking in too late for some InstancedActor instances, due to using Zero-location as player viewer's location

diff --git a/Engine/Plugins/Runtime/MassGameplay/Source/MassLOD/Public/MassLODSubsystem.h b/Engine/Plugins/Runtime/MassGameplay/Source/MassLOD/Public/MassLODSubsystem.h
--- a/Engine/Plugins/Runtime/MassGameplay/Source/MassLOD/Public/MassLODSubsystem.h
+++ b/Engine/Plugins/Runtime/MassGameplay/Source/MassLOD/Public/MassLODSubsystem.h
@@ -111,6 +111,8 @@ public:
 	void RegisterActorViewer(AActor& ActorViewer);
 	void UnregisterActorViewer(AActor& ActorViewer);
 
+	bool IsUsingPlayerPawnLocationInsteadOfCamera() const { return bUsePlayerPawnLocationInsteadOfCamera; }
+
 #if WITH_MASSGAMEPLAY_DEBUG
 	void DebugSetUsePlayerPawnLocationInsteadOfCamera(const bool bInValue) { bUsePlayerPawnLocationInsteadOfCamera = bInValue; }
 	void DebugUnregisterActorViewer();
