Renamed FMassISMCSharedDataMap.GetAndMarkDirty to GetAndMarkDirtyChecked and introduced new GetAndMarkDirty that returns a pointer. Used the new implementation to hopefully avoid the crash in FMassLODSignificanceRange::RemoveInstance.

diff --git a/Engine/Plugins/Runtime/MassGameplay/Source/MassRepresentation/Private/MassVisualizationComponent.cpp b/Engine/Plugins/Runtime/MassGameplay/Source/MassRepresentation/Private/MassVisualizationComponent.cpp
--- a/Engine/Plugins/Runtime/MassGameplay/Source/MassRepresentation/Private/MassVisualizationComponent.cpp
+++ b/Engine/Plugins/Runtime/MassGameplay/Source/MassRepresentation/Private/MassVisualizationComponent.cpp
@@ -633,7 +633,7 @@ void FMassLODSignificanceRange::AddBatchedTransform(const int32 InstanceId, cons
 			continue;
 		}
 
-		FMassISMCSharedData& SharedData = ISMCSharedDataPtr->GetAndMarkDirty(StaticMeshRefs[i]);
+		FMassISMCSharedData& SharedData = ISMCSharedDataPtr->GetAndMarkDirtyChecked(StaticMeshRefs[i]);
 
 		SharedData.UpdateInstanceIds.Add(InstanceId);
 		SharedData.StaticMeshInstanceTransforms.Add(Transform);
@@ -651,7 +651,7 @@ void FMassLODSignificanceRange::AddBatchedCustomDataFloats(const TArray<float>&
 			continue;
 		}
 
-		FMassISMCSharedData& SharedData = ISMCSharedDataPtr->GetAndMarkDirty(StaticMeshRefs[i]);
+		FMassISMCSharedData& SharedData = ISMCSharedDataPtr->GetAndMarkDirtyChecked(StaticMeshRefs[i]);
 		SharedData.StaticMeshInstanceCustomFloats.Append(CustomFloats);
 	}
 }
@@ -661,7 +661,7 @@ void FMassLODSignificanceRange::AddInstance(const int32 InstanceId, const FTrans
 	check(ISMCSharedDataPtr);
 	for (int i = 0; i < StaticMeshRefs.Num(); i++)
 	{
-		FMassISMCSharedData& SharedData = ISMCSharedDataPtr->GetAndMarkDirty(StaticMeshRefs[i]);
+		FMassISMCSharedData& SharedData = ISMCSharedDataPtr->GetAndMarkDirtyChecked(StaticMeshRefs[i]);
 		SharedData.UpdateInstanceIds.Add(InstanceId);
 		SharedData.StaticMeshInstanceTransforms.Add(Transform);
 		SharedData.StaticMeshInstancePrevTransforms.Add(Transform);
@@ -673,8 +673,10 @@ void FMassLODSignificanceRange::RemoveInstance(const int32 InstanceId)
 	check(ISMCSharedDataPtr);
 	for (int i = 0; i < StaticMeshRefs.Num(); i++)
 	{
-		FMassISMCSharedData& SharedData = ISMCSharedDataPtr->GetAndMarkDirty(StaticMeshRefs[i]);
-		SharedData.RemoveInstanceIds.Add(InstanceId);
+		if (FMassISMCSharedData* SharedData = ISMCSharedDataPtr->GetAndMarkDirty(StaticMeshRefs[i]))
+		{
+			SharedData->RemoveInstanceIds.Add(InstanceId);
+		}
 	}
 }
 
@@ -688,7 +690,7 @@ void FMassLODSignificanceRange::WriteCustomDataFloatsAtStartIndex(int32 StaticMe
 			return;
 		}
 
-		FMassISMCSharedData& SharedData = ISMCSharedDataPtr->GetAndMarkDirty(StaticMeshRefs[StaticMeshIndex]);
+		FMassISMCSharedData& SharedData = ISMCSharedDataPtr->GetAndMarkDirtyChecked(StaticMeshRefs[StaticMeshIndex]);
 
 		int32 StartIndex = FloatsPerInstance * SharedData.WriteIterator + StartFloatIndex;
 
diff --git a/Engine/Plugins/Runtime/MassGameplay/Source/MassRepresentation/Public/MassRepresentationTypes.h b/Engine/Plugins/Runtime/MassGameplay/Source/MassRepresentation/Public/MassRepresentationTypes.h
--- a/Engine/Plugins/Runtime/MassGameplay/Source/MassRepresentation/Public/MassRepresentationTypes.h
+++ b/Engine/Plugins/Runtime/MassGameplay/Source/MassRepresentation/Public/MassRepresentationTypes.h
@@ -298,12 +298,22 @@ struct FMassISMCSharedDataMap
 		static constexpr bool bValueToCheck = true;
 	};
 
-	FMassISMCSharedData& GetAndMarkDirty(const uint32 Hash)
+	FMassISMCSharedData& GetAndMarkDirtyChecked(const uint32 Hash)
 	{
 		const int32 DataIndex = Map[Hash];
 		DirtyData[DataIndex] = true;
 		return Data[DataIndex];
 	}
+
+	FMassISMCSharedData* GetAndMarkDirty(const uint32 Hash)
+	{
+		if (const int32* DataIndex = Map.Find(Hash))
+		{
+			DirtyData[*DataIndex] = true;
+			return &Data[*DataIndex];
+		}
+		return nullptr;
+	}
 	
 	FMassISMCSharedData& FindOrAdd(const uint32 Hash, const FMassISMCSharedData& NewData)
 	{
@@ -419,7 +429,7 @@ public:
 				continue;
 			}
 
-			FMassISMCSharedData& SharedData = (*ISMCSharedDataPtr).GetAndMarkDirty(StaticMeshRefs[i]);
+			FMassISMCSharedData& SharedData = (*ISMCSharedDataPtr).GetAndMarkDirtyChecked(StaticMeshRefs[i]);
 			const int32 StartIndex = SharedData.StaticMeshInstanceCustomFloats.AddDefaulted(StructSizeInFloats + NumFloatsToPad);
 			InCustomDataType* CustomData = reinterpret_cast<InCustomDataType*>(&SharedData.StaticMeshInstanceCustomFloats[StartIndex]);
 			*CustomData = InCustomData;
