Added a cvar, mass.LODSubsystem.IncludeAllPlayerControllers, to include all player controllers, even those without a camera or pawn. This fixes the player being skipped by Mass until the player has stasis removed after spawning, which in turn meant that collision wasn't set correctly for some large static mesh ISMs in the game.

diff --git a/Engine/Plugins/Runtime/MassGameplay/Source/MassLOD/Private/MassLODSubsystem.cpp b/Engine/Plugins/Runtime/MassGameplay/Source/MassLOD/Private/MassLODSubsystem.cpp
--- a/Engine/Plugins/Runtime/MassGameplay/Source/MassLOD/Private/MassLODSubsystem.cpp
+++ b/Engine/Plugins/Runtime/MassGameplay/Source/MassLOD/Private/MassLODSubsystem.cpp
@@ -25,6 +25,18 @@ namespace UE::MassLOD
 		FColor::White,
 	};
 	
+	namespace Tweakables
+	{
+
+		bool bLODSubsystemIncludeAllPlayerControllers = true;
+		namespace
+		{
+			static FAutoConsoleVariableRef AnonymousCVars[] = {
+				{ TEXT("mass.LODSubsystem.IncludeAllPlayerControllers"), bLODSubsystemIncludeAllPlayerControllers, TEXT("Include all player controllers, even those without a camera or pawn."), ECVF_Default }
+			};
+		}
+	}  // UE::Mass::Tweakables
+
 #if WITH_MASSGAMEPLAY_DEBUG
 	namespace Debug
 	{
@@ -393,8 +405,10 @@ void UMassLODSubsystem::AddPlayerViewer(APlayerController& PlayerController)
 		return;
 	}
 #endif // WITH_EDITOR
+
 	// ignore players that don't have a pawn nor a camera
-	if (PlayerController.GetPawn() == nullptr
+	if (UE::MassLOD::Tweakables::bLODSubsystemIncludeAllPlayerControllers == false &&
+		PlayerController.GetPawn() == nullptr
 		&& (bool(PlayerController.PlayerCameraManager) == false
 			|| PlayerController.PlayerCameraManager->GetLastFrameCameraCacheTime() == 0.f)
 		)	
