Added a cvar for gating UMassSimulation's game-time ticking

diff --git a/Engine/Plugins/Runtime/MassGameplay/Source/MassSimulation/Private/MassSimulationSubsystem.cpp b/Engine/Plugins/Runtime/MassGameplay/Source/MassSimulation/Private/MassSimulationSubsystem.cpp
--- a/Engine/Plugins/Runtime/MassGameplay/Source/MassSimulation/Private/MassSimulationSubsystem.cpp
+++ b/Engine/Plugins/Runtime/MassGameplay/Source/MassSimulation/Private/MassSimulationSubsystem.cpp
@@ -13,10 +13,19 @@
 
 DEFINE_LOG_CATEGORY(LogMassSim);
 
-namespace UE::MassSimulation
+namespace UE::Mass::Simulation
 {
-	int32 bDoEntityCompaction = 1;
-	FAutoConsoleVariableRef CVarEntityCompaction(TEXT("mass.EntityCompaction"), bDoEntityCompaction, TEXT("Maximize the number of entities per chunk"), ECVF_Cheat);
+	bool bDoEntityCompaction = true;
+	bool bSimulationTickingEnabled = true;
+
+	namespace Private
+	{
+	FAutoConsoleVariableRef CVars[] = {
+		{TEXT("mass.EntityCompaction"), bDoEntityCompaction, TEXT("Maximize the number of entities per chunk"), ECVF_Cheat}
+		, {TEXT("mass.SimulationTickingEnabled"), bSimulationTickingEnabled, TEXT("Controls whether Mass simulation ticking is allowed in game worlds. Upon changing the value it also stops/starts Mass simulation ticking, if needed"), 
+			FConsoleVariableDelegate::CreateStatic(UMassSimulationSubsystem::HandleSimulationTickingEnabledCVarChange)}
+		};
+	}
 }
 
 //----------------------------------------------------------------------//
@@ -143,6 +152,40 @@ void UMassSimulationSubsystem::UnregisterDynamicProcessor(UMassProcessor& Proces
 	PhaseManager.UnregisterDynamicProcessor(Processor);
 }
 
+void UMassSimulationSubsystem::HandleSimulationTickingEnabledCVarChange(IConsoleVariable*)
+{
+	for (const FWorldContext& Context : GEngine->GetWorldContexts())
+	{
+		if (Context.WorldType != EWorldType::Inactive)
+		{
+			UWorld* World = Context.World();
+			
+			// we only want to affect game worlds
+			if (World->IsGameWorld() == false)
+			{
+				continue;
+			}
+
+			if (UMassSimulationSubsystem* MassSimulationSubsystem = UWorld::GetSubsystem<UMassSimulationSubsystem>(World))
+			{
+				if (UE::Mass::Simulation::bSimulationTickingEnabled)
+				{
+					// if enabling call StartSimulation ony if the world has already begun play. Otherwise the 
+					// simulation ticking will be started at the right time automatically
+					if (World->HasBegunPlay())
+					{
+						MassSimulationSubsystem->StartSimulation(*World);
+					}
+				}
+				else
+				{
+					MassSimulationSubsystem->StopSimulation();
+				}
+			}
+		}
+	}
+}
+
 void UMassSimulationSubsystem::RebuildTickPipeline()
 {
 	TConstArrayView<FMassProcessingPhaseConfig> ProcessingPhasesConfig = GET_MASS_CONFIG_VALUE(GetProcessingPhasesConfig());
@@ -162,15 +205,28 @@ void UMassSimulationSubsystem::RebuildTickPipeline()
 
 void UMassSimulationSubsystem::StartSimulation(UWorld& InWorld)
 {
-	PhaseManager.Start(InWorld);
+	if (bSimulationStarted)
+	{
+		return;
+	}
+
+	if (UE::Mass::Simulation::bSimulationTickingEnabled)
+	{
+		PhaseManager.Start(InWorld);
 
-	bSimulationStarted = true;
+		bSimulationStarted = true;
 
-	OnSimulationStarted.Broadcast(&InWorld);
+		OnSimulationStarted.Broadcast(&InWorld);
+	}
 }
 
 void UMassSimulationSubsystem::StopSimulation()
 {
+	if (bSimulationStarted == false)
+	{
+		return;
+	}
+
 	PhaseManager.Stop();
 
 	bSimulationStarted = false;
@@ -182,7 +238,7 @@ void UMassSimulationSubsystem::OnProcessingPhaseStarted(const float DeltaSeconds
 	{
 		case EMassProcessingPhase::PrePhysics:
 			{
-				if (UE::MassSimulation::bDoEntityCompaction && GET_MASSSIMULATION_CONFIG_VALUE(bEntityCompactionEnabled))
+				if (UE::Mass::Simulation::bDoEntityCompaction && GET_MASSSIMULATION_CONFIG_VALUE(bEntityCompactionEnabled))
 				{
 					TRACE_CPUPROFILER_EVENT_SCOPE(DoEntityCompaction);
 					check(EntityManager);
diff --git a/Engine/Plugins/Runtime/MassGameplay/Source/MassSimulation/Public/MassSimulationSubsystem.h b/Engine/Plugins/Runtime/MassGameplay/Source/MassSimulation/Public/MassSimulationSubsystem.h
--- a/Engine/Plugins/Runtime/MassGameplay/Source/MassSimulation/Public/MassSimulationSubsystem.h
+++ b/Engine/Plugins/Runtime/MassGameplay/Source/MassSimulation/Public/MassSimulationSubsystem.h
@@ -10,6 +10,7 @@
 
 
 struct FMassEntityManager;
+class IConsoleVariable;
 
 DECLARE_LOG_CATEGORY_EXTERN(LogMassSim, Log, All);
 
@@ -36,6 +37,9 @@ public:
 	/** @return whether hosted EntityManager is currently, actively being used for processing purposes. Equivalent to calling FMassEntityManager.IsProcessing() */
 	bool IsDuringMassProcessing() const;
 
+	/** Starts/stops simulation ticking for all worlds, based on new `mass.SimulationTickingEnabled` cvar value */
+	static void HandleSimulationTickingEnabledCVarChange(IConsoleVariable*);
+
 protected:
 	// UWorldSubsystem BEGIN
 	virtual void Initialize(FSubsystemCollectionBase& Collection) override;
