Horde Issue 397489
Non unity compilation fix

diff --git a/Engine/Plugins/Runtime/MassGameplay/Source/MassSimulation/Private/MassSimulationSubsystem.cpp b/Engine/Plugins/Runtime/MassGameplay/Source/MassSimulation/Private/MassSimulationSubsystem.cpp
--- a/Engine/Plugins/Runtime/MassGameplay/Source/MassSimulation/Private/MassSimulationSubsystem.cpp
+++ b/Engine/Plugins/Runtime/MassGameplay/Source/MassSimulation/Private/MassSimulationSubsystem.cpp
@@ -7,6 +7,7 @@
 #include "MassSimulationSettings.h"
 #include "VisualLogger/VisualLogger.h"
 #include "MassEntitySettings.h"
+#include "Engine/Engine.h"
 #if WITH_EDITOR
 #include "Editor.h"
 #endif // WITH_EDITOR
