Fixed to a failed check that can trigger when an actor being released vis Mass channels no longer has a valid outer World

While making the change encountered a couple of other minor things to improve. Most notably changed the signature of UMassRepresentationActorManagement::OnPreActorSpawn and UMassRepresentationActorManagement::OnPostActorSpawn to use TShatedPtr<FMassEntityManager> rather than a raw pointer

The changes required updating of CitySample code.

diff --git a/Engine/Plugins/Runtime/MassGameplay/Source/MassActors/Private/MassActorSubsystem.cpp b/Engine/Plugins/Runtime/MassGameplay/Source/MassActors/Private/MassActorSubsystem.cpp
--- a/Engine/Plugins/Runtime/MassGameplay/Source/MassActors/Private/MassActorSubsystem.cpp
+++ b/Engine/Plugins/Runtime/MassGameplay/Source/MassActors/Private/MassActorSubsystem.cpp
@@ -2,6 +2,7 @@
 
 #include "MassActorSubsystem.h"
 #include "MassCommonTypes.h"
+#include "MassActorTypes.h"
 #include "Engine/World.h"
 #include "MassEntityManager.h"
 #include "MassEntitySubsystem.h"
@@ -24,15 +25,21 @@ void FMassActorFragment::SetAndUpdateHandleMap(const FMassEntityHandle MassAgent
 	}
 }
 
-void FMassActorFragment::ResetAndUpdateHandleMap()
+void FMassActorFragment::ResetAndUpdateHandleMap(UMassActorSubsystem* CachedActorSubsystem)
 {
 	if (AActor* ActorPtr = Cast<AActor>(Actor.Get()))
 	{
-		UWorld* World = Actor->GetWorld();
-		check(World);
-		if (UMassActorSubsystem* MassActorSubsystem = UWorld::GetSubsystem<UMassActorSubsystem>(World))
+		if (CachedActorSubsystem == nullptr)
 		{
-			MassActorSubsystem->RemoveHandleForActor(ActorPtr);
+			UWorld* World = Actor->GetWorld();
+			UE_CLOG(World == nullptr, LogMassActor, Warning, TEXT("%hs: got Null while fetching World for actor %s. Can cause issues down the line. Pass in the optional CachedActorSubsystem parameter to address")
+				, __FUNCTION__, *GetNameSafe(ActorPtr));
+			CachedActorSubsystem = UWorld::GetSubsystem<UMassActorSubsystem>(World);
+		}
+
+		if (CachedActorSubsystem)
+		{
+			CachedActorSubsystem->RemoveHandleForActor(ActorPtr);
 		}
 	}
 
diff --git a/Engine/Plugins/Runtime/MassGameplay/Source/MassActors/Private/MassAgentComponent.cpp b/Engine/Plugins/Runtime/MassGameplay/Source/MassActors/Private/MassAgentComponent.cpp
--- a/Engine/Plugins/Runtime/MassGameplay/Source/MassActors/Private/MassAgentComponent.cpp
+++ b/Engine/Plugins/Runtime/MassGameplay/Source/MassActors/Private/MassAgentComponent.cpp
@@ -271,21 +271,24 @@ void UMassAgentComponent::ClearEntityHandle()
 
 void UMassAgentComponent::ClearEntityHandleInternal()
 {
-	if (const UMassAgentSubsystem* AgentSubsystem = UWorld::GetSubsystem<UMassAgentSubsystem>(GetWorld()))
+	if (UWorld* World = GetWorld())
 	{
-		AgentSubsystem->NotifyMassAgentComponentEntityDetaching(*this);
-	}
+		if (const UMassAgentSubsystem* AgentSubsystem = World->GetSubsystem<UMassAgentSubsystem>())
+		{
+			AgentSubsystem->NotifyMassAgentComponentEntityDetaching(*this);
+		}
 
-	// Sync up with mass
-	if (const UMassEntitySubsystem* EntitySubsystem = UWorld::GetSubsystem<UMassEntitySubsystem>(GetWorld()))
-	{
 		if (IsNetSimulating())
 		{
-			const FMassEntityView EntityView(EntitySubsystem->GetEntityManager(), AgentHandle);
-			if (FMassActorFragment* ActorInfo = EntityView.GetFragmentDataPtr<FMassActorFragment>())
+			// Sync up with mass
+			if (const UMassEntitySubsystem* EntitySubsystem = World->GetSubsystem<UMassEntitySubsystem>())
 			{
-				checkf(!ActorInfo->IsValid() || ActorInfo->Get() == GetOwner(), TEXT("Expecting actor pointer to be the Component\'s owner"));
-				ActorInfo->ResetAndUpdateHandleMap();
+				const FMassEntityView EntityView(EntitySubsystem->GetEntityManager(), AgentHandle);
+				if (FMassActorFragment* ActorInfo = EntityView.GetFragmentDataPtr<FMassActorFragment>())
+				{
+					checkf(!ActorInfo->IsValid() || ActorInfo->Get() == GetOwner(), TEXT("Expecting actor pointer to be the Component\'s owner"));
+					ActorInfo->ResetAndUpdateHandleMap(World->GetSubsystem<UMassActorSubsystem>());
+				}
 			}
 		}
 	}
diff --git a/Engine/Plugins/Runtime/MassGameplay/Source/MassActors/Public/MassActorSubsystem.h b/Engine/Plugins/Runtime/MassGameplay/Source/MassActors/Public/MassActorSubsystem.h
--- a/Engine/Plugins/Runtime/MassGameplay/Source/MassActors/Public/MassActorSubsystem.h
+++ b/Engine/Plugins/Runtime/MassGameplay/Source/MassActors/Public/MassActorSubsystem.h
@@ -13,6 +13,7 @@
 struct FMassEntityHandle;
 class AActor;
 struct FMassEntityManager;
+class UMassActorSubsystem;
 
 USTRUCT()
 struct MASSACTORS_API FMassGuidFragment : public FObjectWrapperFragment
@@ -45,8 +46,12 @@ struct MASSACTORS_API FMassActorFragment : public FObjectWrapperFragment
 	 */
 	void SetAndUpdateHandleMap(const FMassEntityHandle MassAgent, AActor* InActor, const bool bInIsOwnedByMass);
 
-	/** Resets the actor pointed by this fragment, will also keep the map back in MassActorSubsystem up to date */
-	void ResetAndUpdateHandleMap();
+	/** 
+	 * Resets the actor pointed by this fragment, will also keep the map back in UMassActorSubsystem up to date 
+	 * @param CachedActorSubsystem if provided will be used directly, otherwise an instance of UMassActorSubsystem will 
+	 *	be deduced from Actor's world (at additional runtime cost)
+	 */
+	void ResetAndUpdateHandleMap(UMassActorSubsystem* CachedActorSubsystem = nullptr);
 
 	/**
 	 * Set the actor associated to a mass agent, will NOT keep map back in MassActorSubsystem up to date.
diff --git a/Engine/Plugins/Runtime/MassGameplay/Source/MassRepresentation/Private/MassRepresentationActorManagement.cpp b/Engine/Plugins/Runtime/MassGameplay/Source/MassRepresentation/Private/MassRepresentationActorManagement.cpp
--- a/Engine/Plugins/Runtime/MassGameplay/Source/MassRepresentation/Private/MassRepresentationActorManagement.cpp
+++ b/Engine/Plugins/Runtime/MassGameplay/Source/MassRepresentation/Private/MassRepresentationActorManagement.cpp
@@ -17,9 +17,11 @@ AActor* UMassRepresentationActorManagement::GetOrSpawnActor(UMassRepresentationS
 	, const FMassEntityHandle MassAgent, const FTransform& Transform, const int16 TemplateActorIndex
 	, FMassActorSpawnRequestHandle& InOutSpawnRequestHandle, const float Priority) const
 {
+	TSharedRef<FMassEntityManager> SharedEntityManager = EntityManager.AsShared();
+
 	return RepresentationSubsystem.GetOrSpawnActorFromTemplate(MassAgent, Transform, TemplateActorIndex, InOutSpawnRequestHandle, Priority,
-		FMassActorPreSpawnDelegate::CreateUObject(this, &UMassRepresentationActorManagement::OnPreActorSpawn, &EntityManager),
-		FMassActorPostSpawnDelegate::CreateUObject(this, &UMassRepresentationActorManagement::OnPostActorSpawn, &EntityManager));
+		FMassActorPreSpawnDelegate::CreateUObject(this, &UMassRepresentationActorManagement::OnPreActorSpawn, SharedEntityManager),
+		FMassActorPostSpawnDelegate::CreateUObject(this, &UMassRepresentationActorManagement::OnPostActorSpawn, SharedEntityManager));
 }
 
 
@@ -51,10 +53,8 @@ void UMassRepresentationActorManagement::TeleportActor(const FTransform& Transfo
 	}
 }
 
-void UMassRepresentationActorManagement::OnPreActorSpawn(const FMassActorSpawnRequestHandle& SpawnRequestHandle, FConstStructView SpawnRequest, FMassEntityManager* EntityManager) const
+void UMassRepresentationActorManagement::OnPreActorSpawn(const FMassActorSpawnRequestHandle& SpawnRequestHandle, FConstStructView SpawnRequest, TSharedRef<FMassEntityManager> EntityManager) const
 {
-	check(EntityManager);
-
 	const FMassActorSpawnRequest& MassActorSpawnRequest = SpawnRequest.Get<const FMassActorSpawnRequest>();
 	const FMassEntityView EntityView(*EntityManager, MassActorSpawnRequest.MassAgent);
 	FMassActorFragment& ActorInfo = EntityView.GetFragmentData<FMassActorFragment>();
@@ -82,10 +82,9 @@ void UMassRepresentationActorManagement::OnPreActorSpawn(const FMassActorSpawnRe
 	}
 }
 
-EMassActorSpawnRequestAction UMassRepresentationActorManagement::OnPostActorSpawn(const FMassActorSpawnRequestHandle& SpawnRequestHandle, FConstStructView SpawnRequest, FMassEntityManager* EntityManager) const
+EMassActorSpawnRequestAction UMassRepresentationActorManagement::OnPostActorSpawn(const FMassActorSpawnRequestHandle& SpawnRequestHandle
+	, FConstStructView SpawnRequest, TSharedRef<FMassEntityManager> EntityManager) const
 {
-	check(EntityManager);
-
 	const FMassActorSpawnRequest& MassActorSpawnRequest = SpawnRequest.Get<const FMassActorSpawnRequest>();
 	checkf(MassActorSpawnRequest.SpawnedActor, TEXT("Expecting valid spawned actor"));
 
@@ -115,7 +114,8 @@ void UMassRepresentationActorManagement::ReleaseAnyActorOrCancelAnySpawning(FMas
 }
 
 void UMassRepresentationActorManagement::ReleaseAnyActorOrCancelAnySpawning(UMassRepresentationSubsystem& RepresentationSubsystem
-	, const FMassEntityHandle MassAgent, FMassActorFragment& ActorInfo, FMassRepresentationFragment& Representation)
+	, const FMassEntityHandle MassAgent, FMassActorFragment& ActorInfo, FMassRepresentationFragment& Representation
+	, UMassActorSubsystem* CachedActorSubsystem)
 {
 	// This method can only release owned by mass actors
 	AActor* Actor = ActorInfo.GetOwnedByMassMutable();
@@ -124,7 +124,7 @@ void UMassRepresentationActorManagement::ReleaseAnyActorOrCancelAnySpawning(UMas
 		// WARNING!
 		// Need to reset before ReleaseTemplateActorOrCancelSpawning as this action might move the entity to a new archetype and
 		// so the Fragment passed in parameters would not be valid anymore.
-		ActorInfo.ResetAndUpdateHandleMap();
+		ActorInfo.ResetAndUpdateHandleMap(CachedActorSubsystem);
 	}
 	// Try releasing both as we can have a low res actor and a high res spawning request
 	if (Representation.HighResTemplateActorIndex != INDEX_NONE)
@@ -138,6 +138,10 @@ void UMassRepresentationActorManagement::ReleaseAnyActorOrCancelAnySpawning(UMas
 	check(!Representation.ActorSpawnRequestHandle.IsValid());
 }
 
+
+//-----------------------------------------------------------------------------
+// DEPRECATED
+//-----------------------------------------------------------------------------
 AActor* UMassRepresentationActorManagement::GetOrSpawnActor(UMassRepresentationSubsystem& RepresentationSubsystem, FMassEntityManager& EntityManager
 	, const FMassEntityHandle MassAgent, FMassActorFragment&/* OutActorInfo*/, const FTransform& Transform, const int16 TemplateActorIndex
 	, FMassActorSpawnRequestHandle& InOutSpawnRequestHandle, const float Priority) const
@@ -145,4 +149,22 @@ AActor* UMassRepresentationActorManagement::GetOrSpawnActor(UMassRepresentationS
 	return GetOrSpawnActor(RepresentationSubsystem, EntityManager
 		, MassAgent, Transform, TemplateActorIndex
 		, InOutSpawnRequestHandle, Priority);
-}
\ No newline at end of file
+}
+
+void UMassRepresentationActorManagement::OnPreActorSpawn(const FMassActorSpawnRequestHandle& SpawnRequestHandle, FConstStructView SpawnRequest, FMassEntityManager* EntityManager) const
+{
+	if (EntityManager)
+	{
+		OnPreActorSpawn(SpawnRequestHandle, SpawnRequest, EntityManager->AsShared());
+	}
+}
+
+EMassActorSpawnRequestAction UMassRepresentationActorManagement::OnPostActorSpawn(const FMassActorSpawnRequestHandle& SpawnRequestHandle, FConstStructView SpawnRequest, FMassEntityManager* EntityManager) const
+{
+	if (EntityManager)
+	{
+		return OnPostActorSpawn(SpawnRequestHandle, SpawnRequest, EntityManager->AsShared());
+	}
+
+	return EMassActorSpawnRequestAction::Remove;
+}
diff --git a/Engine/Plugins/Runtime/MassGameplay/Source/MassRepresentation/Private/MassRepresentationProcessor.cpp b/Engine/Plugins/Runtime/MassGameplay/Source/MassRepresentation/Private/MassRepresentationProcessor.cpp
--- a/Engine/Plugins/Runtime/MassGameplay/Source/MassRepresentation/Private/MassRepresentationProcessor.cpp
+++ b/Engine/Plugins/Runtime/MassGameplay/Source/MassRepresentation/Private/MassRepresentationProcessor.cpp
@@ -428,6 +428,7 @@ void UMassRepresentationFragmentDestructor::ConfigureQueries()
 	EntityQuery.AddRequirement<FMassActorFragment>(EMassFragmentAccess::ReadWrite);
 	EntityQuery.AddConstSharedRequirement<FMassRepresentationParameters>();
 	EntityQuery.AddSharedRequirement<FMassRepresentationSubsystemSharedFragment>(EMassFragmentAccess::ReadWrite);
+	EntityQuery.AddSubsystemRequirement<UMassActorSubsystem>(EMassFragmentAccess::ReadWrite);
 }
 
 void UMassRepresentationFragmentDestructor::Execute(FMassEntityManager& EntityManager, FMassExecutionContext& Context)
@@ -436,6 +437,7 @@ void UMassRepresentationFragmentDestructor::Execute(FMassEntityManager& EntityMa
 	{
 		UMassRepresentationSubsystem* RepresentationSubsystem = Context.GetMutableSharedFragment<FMassRepresentationSubsystemSharedFragment>().RepresentationSubsystem;
 		check(RepresentationSubsystem);
+		UMassActorSubsystem* ActorSubsystem = Context.GetMutableSubsystem<UMassActorSubsystem>();
 
 		const TArrayView<FMassRepresentationFragment> RepresentationList = Context.GetMutableFragmentView<FMassRepresentationFragment>();
 		const TArrayView<FMassActorFragment> ActorList = Context.GetMutableFragmentView<FMassActorFragment>();
@@ -448,7 +450,7 @@ void UMassRepresentationFragmentDestructor::Execute(FMassEntityManager& EntityMa
 
 			const FMassEntityHandle MassAgentEntityHandle = Context.GetEntity(i);
 
-			UMassRepresentationActorManagement::ReleaseAnyActorOrCancelAnySpawning(*RepresentationSubsystem, MassAgentEntityHandle, ActorInfo, Representation);
+			UMassRepresentationActorManagement::ReleaseAnyActorOrCancelAnySpawning(*RepresentationSubsystem, MassAgentEntityHandle, ActorInfo, Representation, ActorSubsystem);
 		}
 	});
 }
diff --git a/Engine/Plugins/Runtime/MassGameplay/Source/MassRepresentation/Public/MassRepresentationActorManagement.h b/Engine/Plugins/Runtime/MassGameplay/Source/MassRepresentation/Public/MassRepresentationActorManagement.h
--- a/Engine/Plugins/Runtime/MassGameplay/Source/MassRepresentation/Public/MassRepresentationActorManagement.h
+++ b/Engine/Plugins/Runtime/MassGameplay/Source/MassRepresentation/Public/MassRepresentationActorManagement.h
@@ -15,6 +15,7 @@ struct FMassActorFragment;
 struct FMassRepresentationLODFragment;
 struct FMassRepresentationFragment;
 class UMassRepresentationSubsystem;
+class UMassActorSubsystem;
 
 UCLASS()
 class MASSREPRESENTATION_API UMassRepresentationActorManagement : public UObject
@@ -68,18 +69,18 @@ public:
 	 * Method that will be bound to a delegate called before the spawning of an actor to let the requester prepare it
 	 * @param SpawnRequestHandle the handle of the spawn request that is about to spawn
 	 * @param SpawnRequest of the actor that is about to spawn
-	 * @param EntitySubsystem to use to retrieve the mass agent fragments
+	 * @param EntityManager to use to retrieve the mass agent fragments
 	 */
-	virtual void OnPreActorSpawn(const FMassActorSpawnRequestHandle& SpawnRequestHandle, FConstStructView SpawnRequest, FMassEntityManager* EntitySubsystem) const;
+	virtual void OnPreActorSpawn(const FMassActorSpawnRequestHandle& SpawnRequestHandle, FConstStructView SpawnRequest, TSharedRef<FMassEntityManager> EntityManager) const;
 
 	/**
 	 * Method that will be bound to a delegate used post-spawn to notify and let the requester configure the actor
 	 * @param SpawnRequestHandle the handle of the spawn request that was just spawned
 	 * @param SpawnRequest of the actor that just spawned
-	 * @param EntitySubsystem to use to retrieve the mass agent fragments
+	 * @param EntityManager to use to retrieve the mass agent fragments
 	 * @return The action to take on the spawn request, either keep it there or remove it.
 	 */
-	virtual EMassActorSpawnRequestAction OnPostActorSpawn(const FMassActorSpawnRequestHandle& SpawnRequestHandle, FConstStructView SpawnRequest, FMassEntityManager* EntitySubsystem) const;
+	virtual EMassActorSpawnRequestAction OnPostActorSpawn(const FMassActorSpawnRequestHandle& SpawnRequestHandle, FConstStructView SpawnRequest, TSharedRef<FMassEntityManager> EntityManager) const;
 
 	/**
 	 * Static methods to Release an actor or cancel its spawning (calls ReleaseAnyActorOrCancelAnySpawning)
@@ -99,11 +100,20 @@ public:
 	 * @param MassAgent is the handle to the associated mass agent
 	 * @param ActorInfo is the fragment where we are going to store the actor pointer
 	 * @param Representation fragment containing the current and previous visual state
+	 * @param ActorSubsystem passed over to FMassActorFragment::ResetAndUpdateHandleMap, used to avoid fetching the subsystem
+	 *	from UWorld every time.
 	 */
-	static void ReleaseAnyActorOrCancelAnySpawning(UMassRepresentationSubsystem& RepresentationSubsystem, const FMassEntityHandle MassAgent, FMassActorFragment& ActorInfo, FMassRepresentationFragment& Representation);
+	static void ReleaseAnyActorOrCancelAnySpawning(UMassRepresentationSubsystem& RepresentationSubsystem, const FMassEntityHandle MassAgent
+		, FMassActorFragment& ActorInfo, FMassRepresentationFragment& Representation, UMassActorSubsystem* ActorSubsystem = nullptr);
 
 	UE_DEPRECATED(5.4, "This flavor of GetOrSpawnActor has been deprecated due to a defunct parameter, OutActorInfo, that was never being used.")
 	virtual AActor* GetOrSpawnActor(UMassRepresentationSubsystem& RepresentationSubsystem, FMassEntityManager& EntitySubsystem
 		, const FMassEntityHandle MassAgent, FMassActorFragment&/* OutActorInfo*/, const FTransform& Transform, const int16 TemplateActorIndex
 		, FMassActorSpawnRequestHandle& InOutSpawnRequestHandle, const float Priority) const final;
+
+	UE_DEPRECATED(5.4, "This flavor of OnPreActorSpawn has been deprecated in favor of the version using TSharedPtr<FMassEntityManager> as a parameter. Use that one instead.")
+	virtual void OnPreActorSpawn(const FMassActorSpawnRequestHandle& SpawnRequestHandle, FConstStructView SpawnRequest, FMassEntityManager* EntityManager) const final;
+
+	UE_DEPRECATED(5.4, "This flavor of OnPostActorSpawn has been deprecated in favor of the version using TSharedPtr<FMassEntityManager> as a parameter. Use that one instead.")
+	virtual EMassActorSpawnRequestAction OnPostActorSpawn(const FMassActorSpawnRequestHandle& SpawnRequestHandle, FConstStructView SpawnRequest, FMassEntityManager* EntityManager) const final;
 };
\ No newline at end of file
