Improve determinism of Mass, specifically as used in CitySample.
- added private random streams to some objects to reduce thread ordering risks
- added mass.RandomSeedOverride and helper functions UE::Mass::Utils::IsDeterministic, UE::Mass::Utils::OverrideRandomSeedForTesting, and UE::Mass::Utils::GenerateRandomSeed to make it easier to replace the random seed with a fixed value
- gated the use of fixed random seed on FApp::bUseFixedSeed (controlled through  -FixedSeed or -Deterministic)

diff --git a/Engine/Plugins/Runtime/MassGameplay/Source/MassCommon/Private/MassCommonUtils.cpp b/Engine/Plugins/Runtime/MassGameplay/Source/MassCommon/Private/MassCommonUtils.cpp
--- a/Engine/Plugins/Runtime/MassGameplay/Source/MassCommon/Private/MassCommonUtils.cpp
+++ b/Engine/Plugins/Runtime/MassGameplay/Source/MassCommon/Private/MassCommonUtils.cpp
@@ -1,9 +1,20 @@
 // Copyright Epic Games, Inc. All Rights Reserved.
 
 #include "MassCommonUtils.h"
+#include "Misc/App.h"
+
+namespace UE::Mass::Private
+{
+    int32 RandomSeedOverride = 7;
+    FAutoConsoleVariableRef CVarRandomSeedOverride(
+		TEXT("mass.RandomSeedOverride"), 
+		RandomSeedOverride, 
+		TEXT("If FApp::bUseFixedSeed is true (e.g., -FixedSeed or -Deterministic) this value is used as random seed for all Mass uses."));
+}
 
 namespace UE::Mass::Utils
 {
+
 	TArray<FMassEntityHandle> EntityQueueToArray(TQueue<FMassEntityHandle, EQueueMode::Mpsc>& EntitiesQueue, const int32 EntitiesCount)
 	{
 		check(EntitiesCount > 0);
@@ -20,4 +31,31 @@ namespace UE::Mass::Utils
 
 		return MoveTemp(EntitiesArray);
 	}
-} // namespace UE::Mass::Utils
\ No newline at end of file
+
+#if !UE_BUILD_SHIPPING
+
+	bool IsDeterministic()
+	{
+		return FApp::bUseFixedSeed;
+	}
+
+	int32 OverrideRandomSeedForTesting(int32 InSeed)
+	{
+		if (IsDeterministic())
+		{
+			return UE::Mass::Private::RandomSeedOverride;
+		}
+		return InSeed;
+	}
+
+	int32 GenerateRandomSeed()
+	{
+		if (IsDeterministic())
+		{
+			return UE::Mass::Private::RandomSeedOverride;
+		}
+		return FMath::Rand();
+	}
+#endif
+
+} // namespace UE::Mass::Utils
diff --git a/Engine/Plugins/Runtime/MassGameplay/Source/MassCommon/Public/MassCommonUtils.h b/Engine/Plugins/Runtime/MassGameplay/Source/MassCommon/Public/MassCommonUtils.h
--- a/Engine/Plugins/Runtime/MassGameplay/Source/MassCommon/Public/MassCommonUtils.h
+++ b/Engine/Plugins/Runtime/MassGameplay/Source/MassCommon/Public/MassCommonUtils.h
@@ -13,4 +13,32 @@ namespace UE::Mass::Utils
 	 * ensure EntitiesCount > 0, otherwise the function will fail a check (with explosive results).
 	 */
 	MASSCOMMON_API TArray<FMassEntityHandle> EntityQueueToArray(TQueue<FMassEntityHandle, EQueueMode::Mpsc>& EntitiesQueue, const int32 EntitiesCount);
+
+#if !UE_BUILD_SHIPPING
+
+	/**
+	 * If this is true, then the mass systems should strive to be as deterministic as possible, this will also enable the fixed random seed
+	 * Currently maps to FApp::bUseFixedSeed
+	 */
+	MASSCOMMON_API bool IsDeterministic();
+
+	/**
+	 * If IsDeterministic() returns true, then this function will return the value of ai.massrepresentation.OverrideRandomSeed in place of InSeed
+	 */
+	MASSCOMMON_API int32 OverrideRandomSeedForTesting(int32 InSeed);
+
+	/**
+	 * If IsDeterministic() returns true, then this function will return the value of ai.massrepresentation.OverrideRandomSeed in place of FMath::Rand()
+	 */
+	MASSCOMMON_API int32 GenerateRandomSeed();
+
+#else
+
+	FORCEINLINE constexpr bool IsDeterministic() { return false; }
+
+	FORCEINLINE int32 OverrideRandomSeedForTesting(int32 InSeed) { return InSeed; }
+
+	FORCEINLINE int32 GenerateRandomSeed() { return FMath::Rand(); }
+#endif
+
 } // namespace UE::Mass::Utils
\ No newline at end of file
diff --git a/Engine/Plugins/Runtime/MassGameplay/Source/MassLOD/Public/MassLODTickRateController.h b/Engine/Plugins/Runtime/MassGameplay/Source/MassLOD/Public/MassLODTickRateController.h
--- a/Engine/Plugins/Runtime/MassGameplay/Source/MassLOD/Public/MassLODTickRateController.h
+++ b/Engine/Plugins/Runtime/MassGameplay/Source/MassLOD/Public/MassLODTickRateController.h
@@ -4,6 +4,7 @@
 
 #include "MassLODLogic.h"
 #include "MassLODUtils.h"
+#include "MassCommonUtils.h"
 #include "MassCommandBuffer.h"
 
 /**
@@ -125,12 +126,14 @@ bool TMassLODTickRateController<TVariableTickChunkFragment, FLODLogic>::UpdateTi
 		const float TickRate = TickRates[ChunkLOD];
 		if (bFirstUpdate)
 		{
-			TimeUntilNextTick = FMath::RandRange(0.0f, TickRate);
+			// @todo: Add some randomization for deterministic runs too. The randomization is used to distribute the infrequent ticks evenly on different frames.
+			TimeUntilNextTick = UE::Mass::Utils::IsDeterministic() ? TickRate * 0.5f : FMath::RandRange(0.0f, TickRate);
 		}
 		else if(bWasChunkTicked)
 		{
 			// Reset DeltaTime if we ticked last frame and start tracking chunk modifications
-			TimeUntilNextTick = TickRate * (1.0f + FMath::RandRange(-0.1f, 0.1f));
+			// @todo: Add some randomization for deterministic runs too. The randomization is used to distribute the infrequent ticks evenly on different frames.
+			TimeUntilNextTick = UE::Mass::Utils::IsDeterministic() ? TickRate : (TickRate * (1.0f + FMath::RandRange(-0.1f, 0.1f)));
 			NewChunkSerialModificationNumber = ChunkSerialModificationNumber;
 		}
 		else
diff --git a/Engine/Plugins/Runtime/MassGameplay/Source/MassMovement/Private/MassVelocityRandomizerTrait.cpp b/Engine/Plugins/Runtime/MassGameplay/Source/MassMovement/Private/MassVelocityRandomizerTrait.cpp
--- a/Engine/Plugins/Runtime/MassGameplay/Source/MassMovement/Private/MassVelocityRandomizerTrait.cpp
+++ b/Engine/Plugins/Runtime/MassGameplay/Source/MassMovement/Private/MassVelocityRandomizerTrait.cpp
@@ -5,7 +5,7 @@
 #include "MassExecutionContext.h"
 #include "Engine/World.h"
 #include "MassMovementFragments.h"
-
+#include "MassCommonUtils.h"
 
 //----------------------------------------------------------------------//
 //  UMassVelocityRandomizerTrait
@@ -28,6 +28,7 @@ UMassRandomVelocityInitializer::UMassRandomVelocityInitializer()
 {
 	ObservedType = FMassVelocityFragment::StaticStruct();
 	Operation = EMassObservedOperation::Add;
+	RandomStream.Initialize(UE::Mass::Utils::GenerateRandomSeed());
 }
 
 void UMassRandomVelocityInitializer::ConfigureQueries()
@@ -46,10 +47,10 @@ void UMassRandomVelocityInitializer::Execute(FMassEntityManager& EntityManager,
 				// the given VelocityFragment's value is encoding the initialization parameters, as per comment in
 				// UMassVelocityRandomizerTrait::BuildTemplate
 				const FVector RandomVector = VelocityFragment.Value.Z != 0.f 
-					? FVector(FMath::FRandRange(-1.f, 1.f), FMath::FRandRange(-1.f, 1.f), FMath::FRandRange(-1.f, 1.f)).GetSafeNormal()
-					: FVector(FMath::FRandRange(-1.f, 1.f), FMath::FRandRange(-1.f, 1.f), 0).GetSafeNormal();
+					? FVector(RandomStream.FRandRange(-1.f, 1.f), RandomStream.FRandRange(-1.f, 1.f), RandomStream.FRandRange(-1.f, 1.f)).GetSafeNormal()
+					: FVector(RandomStream.FRandRange(-1.f, 1.f), RandomStream.FRandRange(-1.f, 1.f), 0).GetSafeNormal();
 
-				VelocityFragment.Value = RandomVector * FMath::FRandRange(VelocityFragment.Value.X, VelocityFragment.Value.Y);
+				VelocityFragment.Value = RandomVector * RandomStream.FRandRange(VelocityFragment.Value.X, VelocityFragment.Value.Y);
 			}
 		}));
 }
diff --git a/Engine/Plugins/Runtime/MassGameplay/Source/MassMovement/Public/Example/MassVelocityRandomizerTrait.h b/Engine/Plugins/Runtime/MassGameplay/Source/MassMovement/Public/Example/MassVelocityRandomizerTrait.h
--- a/Engine/Plugins/Runtime/MassGameplay/Source/MassMovement/Public/Example/MassVelocityRandomizerTrait.h
+++ b/Engine/Plugins/Runtime/MassGameplay/Source/MassMovement/Public/Example/MassVelocityRandomizerTrait.h
@@ -54,4 +54,5 @@ protected:
 
 	UPROPERTY()
 	bool bSetZComponent = false;
+	FRandomStream RandomStream;
 };
diff --git a/Engine/Plugins/Runtime/MassGameplay/Source/MassMovement/Public/MassMovementFragments.h b/Engine/Plugins/Runtime/MassGameplay/Source/MassMovement/Public/MassMovementFragments.h
--- a/Engine/Plugins/Runtime/MassGameplay/Source/MassMovement/Public/MassMovementFragments.h
+++ b/Engine/Plugins/Runtime/MassGameplay/Source/MassMovement/Public/MassMovementFragments.h
@@ -64,7 +64,7 @@ struct MASSMOVEMENT_API FMassMovementParameters : public FMassSharedFragment
 			}
 		}
 		
-		return DesiredSpeed * FMath::RandRange(1.0f - DesiredSpeedVariance, 1.0f + DesiredSpeedVariance);
+		return DesiredSpeed * UE::RandomSequence::RandRange(UniqueId, 1.0f - DesiredSpeedVariance, 1.0f + DesiredSpeedVariance);;
 	}
 	
 	/** Maximum speed (cm/s). */
diff --git a/Engine/Plugins/Runtime/MassGameplay/Source/MassRepresentation/Private/MassRepresentationProcessor.cpp b/Engine/Plugins/Runtime/MassGameplay/Source/MassRepresentation/Private/MassRepresentationProcessor.cpp
--- a/Engine/Plugins/Runtime/MassGameplay/Source/MassRepresentation/Private/MassRepresentationProcessor.cpp
+++ b/Engine/Plugins/Runtime/MassGameplay/Source/MassRepresentation/Private/MassRepresentationProcessor.cpp
@@ -11,6 +11,7 @@
 #include "MassEntityView.h"
 #include "Engine/World.h"
 #include "MassRepresentationActorManagement.h"
+#include "MassCommonUtils.h"
 #include "MassEntityUtils.h"
 #include "MassExecutionContext.h"
 
@@ -365,13 +366,15 @@ FMassVisualizationChunkFragment& UMassVisualizationProcessor::UpdateChunkVisibil
 		if (bFirstUpdate)
 		{
 			// A DeltaTime of 0.0f means it will tick this frame.
-			DeltaTime = FMath::RandRange(0.0f, RepresentationParams.NotVisibleUpdateRate);
+			// @todo: Add some randomization for deterministic runs too. The randomization is used to distribute the infrequent ticks evenly on different frames.
+			DeltaTime = UE::Mass::Utils::IsDeterministic() ? RepresentationParams.NotVisibleUpdateRate * 0.5f : FMath::RandRange(0.0f, RepresentationParams.NotVisibleUpdateRate);
 		}
 		else 
 		{
 			if (DeltaTime < 0.0f)
 			{
-				DeltaTime += RepresentationParams.NotVisibleUpdateRate * (1.0f + FMath::RandRange(-0.1f, 0.1f));
+				// @todo: Add some randomization for deterministic runs too. The randomization is used to distribute the infrequent ticks evenly on different frames.
+				DeltaTime += UE::Mass::Utils::IsDeterministic() ? RepresentationParams.NotVisibleUpdateRate : (RepresentationParams.NotVisibleUpdateRate * (1.0f + FMath::RandRange(-0.1f, 0.1f)));
 			}
 			DeltaTime -= Context.GetDeltaTimeSeconds();
 		}
diff --git a/Engine/Plugins/Runtime/MassGameplay/Source/MassSpawner/Private/MassEntityEQSSpawnPointsGenerator.cpp b/Engine/Plugins/Runtime/MassGameplay/Source/MassSpawner/Private/MassEntityEQSSpawnPointsGenerator.cpp
--- a/Engine/Plugins/Runtime/MassGameplay/Source/MassSpawner/Private/MassEntityEQSSpawnPointsGenerator.cpp
+++ b/Engine/Plugins/Runtime/MassGameplay/Source/MassSpawner/Private/MassEntityEQSSpawnPointsGenerator.cpp
@@ -5,6 +5,7 @@
 #include "MassSpawnLocationProcessor.h"
 #include "VisualLogger/VisualLogger.h"
 #include "MassGameplaySettings.h"
+#include "MassCommonUtils.h"
 
 
 UMassEntityEQSSpawnPointsGenerator::UMassEntityEQSSpawnPointsGenerator()
@@ -56,7 +57,7 @@ void UMassEntityEQSSpawnPointsGenerator::OnEQSQueryFinished(TSharedPtr<FEnvQuery
 	EQSResult->GetAllAsLocations(Locations);
 
 	// Randomize them
-	FRandomStream RandomStream(GetRandomSelectionSeed());
+	FRandomStream RandomStream(UE::Mass::Utils::OverrideRandomSeedForTesting(GetRandomSelectionSeed()));
 	for (int32 I = 0; I < Locations.Num(); ++I)
 	{
 		const int32 J = RandomStream.RandHelper(Locations.Num());
diff --git a/Engine/Plugins/Runtime/MassGameplay/Source/MassSpawner/Private/MassEntityZoneGraphSpawnPointsGenerator.cpp b/Engine/Plugins/Runtime/MassGameplay/Source/MassSpawner/Private/MassEntityZoneGraphSpawnPointsGenerator.cpp
--- a/Engine/Plugins/Runtime/MassGameplay/Source/MassSpawner/Private/MassEntityZoneGraphSpawnPointsGenerator.cpp
+++ b/Engine/Plugins/Runtime/MassGameplay/Source/MassSpawner/Private/MassEntityZoneGraphSpawnPointsGenerator.cpp
@@ -9,6 +9,7 @@
 #include "MassGameplaySettings.h"
 #include "MassSpawnLocationProcessor.h"
 #include "Engine/World.h"
+#include "MassCommonUtils.h"
 
 void UMassEntityZoneGraphSpawnPointsGenerator::Generate(UObject& QueryOwner, TConstArrayView<FMassSpawnedEntityType> EntityTypes, int32 Count, FFinishedGeneratingSpawnDataSignature& FinishedGeneratingSpawnPointsDelegate) const
 {
@@ -27,7 +28,7 @@ void UMassEntityZoneGraphSpawnPointsGenerator::Generate(UObject& QueryOwner, TCo
 
 	TArray<FVector> Locations;
 	
-	const FRandomStream RandomStream(GetRandomSelectionSeed());
+	const FRandomStream RandomStream(UE::Mass::Utils::OverrideRandomSeedForTesting(GetRandomSelectionSeed()));
 	const TConstArrayView<FRegisteredZoneGraphData> RegisteredZoneGraphs = ZoneGraph->GetRegisteredZoneGraphData();
 	if (RegisteredZoneGraphs.IsEmpty())
 	{
diff --git a/Engine/Plugins/Runtime/MassGameplay/Source/MassSpawner/Private/MassSpawnLocationProcessor.cpp b/Engine/Plugins/Runtime/MassGameplay/Source/MassSpawner/Private/MassSpawnLocationProcessor.cpp
--- a/Engine/Plugins/Runtime/MassGameplay/Source/MassSpawner/Private/MassSpawnLocationProcessor.cpp
+++ b/Engine/Plugins/Runtime/MassGameplay/Source/MassSpawner/Private/MassSpawnLocationProcessor.cpp
@@ -5,6 +5,7 @@
 #include "MassEntityManager.h"
 #include "MassExecutionContext.h"
 #include "MassSpawnerTypes.h"
+#include "MassCommonUtils.h"
 #include "Engine/World.h"
 #include "VisualLogger/VisualLogger.h"
 
@@ -15,6 +16,7 @@ UMassSpawnLocationProcessor::UMassSpawnLocationProcessor()
 	: EntityQuery(*this)
 {
 	bAutoRegisterWithProcessingPhases = false;
+	RandomStream.Initialize(UE::Mass::Utils::GenerateRandomSeed());
 }
 
 void UMassSpawnLocationProcessor::ConfigureQueries()
@@ -61,11 +63,11 @@ void UMassSpawnLocationProcessor::Execute(FMassEntityManager& EntityManager, FMa
 		Transforms.AddUninitialized(NumToAdd);
 		for (int i = 0; i < NumToAdd; ++i)
 		{
-			Transforms[NumSpawnTransforms + i] = Transforms[FMath::RandRange(0, NumSpawnTransforms - 1)];
+			Transforms[NumSpawnTransforms + i] = Transforms[RandomStream.RandRange(0, NumSpawnTransforms - 1)];
 		}
 	}
 
-	if (AuxData.bRandomize)
+	if (AuxData.bRandomize && !UE::Mass::Utils::IsDeterministic())
 	{
 		EntityQuery.ForEachEntityChunk(EntityManager, Context, [&Transforms, this](FMassExecutionContext& Context)
 			{
@@ -73,7 +75,7 @@ void UMassSpawnLocationProcessor::Execute(FMassEntityManager& EntityManager, FMa
 				const int32 NumEntities = Context.GetNumEntities();
 				for (int32 i = 0; i < NumEntities; ++i)
 				{
-					const int32 AuxIndex = FMath::RandRange(0, Transforms.Num() - 1);
+					const int32 AuxIndex = RandomStream.RandRange(0, Transforms.Num() - 1);
 					LocationList[i].GetMutableTransform() = Transforms[AuxIndex];
 					Transforms.RemoveAtSwap(AuxIndex, 1, /*bAllowShrinking=*/false);
 				}
diff --git a/Engine/Plugins/Runtime/MassGameplay/Source/MassSpawner/Public/MassSpawnLocationProcessor.h b/Engine/Plugins/Runtime/MassGameplay/Source/MassSpawner/Public/MassSpawnLocationProcessor.h
--- a/Engine/Plugins/Runtime/MassGameplay/Source/MassSpawner/Public/MassSpawnLocationProcessor.h
+++ b/Engine/Plugins/Runtime/MassGameplay/Source/MassSpawner/Public/MassSpawnLocationProcessor.h
@@ -18,4 +18,5 @@ protected:
 	virtual void Execute(FMassEntityManager& EntityManager, FMassExecutionContext& Context) override;
 
 	FMassEntityQuery EntityQuery;
+	FRandomStream RandomStream;
 };
