Renamed last of "EntitySubsystem" occurrences to "EntityManager" in Mass code (where applicable, there are still "UMassEntitySubsystem* EntitySubsystem" variables).

#tests compilation, PIE

diff --git a/Engine/Plugins/AI/MassAI/Source/MassAIBehavior/Public/MassLookAtProcessors.h b/Engine/Plugins/AI/MassAI/Source/MassAIBehavior/Public/MassLookAtProcessors.h
--- a/Engine/Plugins/AI/MassAI/Source/MassAIBehavior/Public/MassLookAtProcessors.h
+++ b/Engine/Plugins/AI/MassAI/Source/MassAIBehavior/Public/MassLookAtProcessors.h
@@ -31,17 +31,17 @@ protected:
 	virtual void Execute(FMassEntityManager& EntityManager, FMassExecutionContext& Context) override;
 
 	/** Selects a nearby target if possible or use a random fixed direction */
-	void FindNewGazeTarget(const UMassNavigationSubsystem& MassNavSystem, const FMassEntityManager& EntitySubsystem, const double CurrentTime, const FTransform& Transform, FMassLookAtFragment& LookAt) const;
+	void FindNewGazeTarget(const UMassNavigationSubsystem& MassNavSystem, const FMassEntityManager& EntityManager, const double CurrentTime, const FTransform& Transform, FMassLookAtFragment& LookAt) const;
 
 	/** Updates look direction based on look at trajectory. */
 	void UpdateLookAtTrajectory(const FTransform& Transform, const FMassZoneGraphLaneLocationFragment& ZoneGraphLocation,
 								const FMassLookAtTrajectoryFragment& LookAtTrajectory, const bool bDisplayDebug, FMassLookAtFragment& LookAt) const;
 
 	/** Updates look at based on tracked entity. */
-	void UpdateLookAtTrackedEntity(const FMassEntityManager& EntitySubsystem, const FTransform& Transform, const bool bDisplayDebug, FMassLookAtFragment& LookAt) const;
+	void UpdateLookAtTrackedEntity(const FMassEntityManager& EntityManager, const FTransform& Transform, const bool bDisplayDebug, FMassLookAtFragment& LookAt) const;
 
 	/** Updates gaze based on tracked entity. */
-	bool UpdateGazeTrackedEntity(const FMassEntityManager& EntitySubsystem, const FTransform& Transform, const bool bDisplayDebug, FMassLookAtFragment& LookAt) const;
+	bool UpdateGazeTrackedEntity(const FMassEntityManager& EntityManager, const FTransform& Transform, const bool bDisplayDebug, FMassLookAtFragment& LookAt) const;
 
 	/** Builds look at trajectory along the current path. */
 	void BuildTrajectory(const UZoneGraphSubsystem& ZoneGraphSubsystem, const FMassZoneGraphLaneLocationFragment& LaneLocation, const FMassZoneGraphShortPathFragment& ShortPath,
diff --git a/Engine/Plugins/Runtime/MassEntity/Source/MassEntityTestSuite/Private/MassEntityTest.cpp b/Engine/Plugins/Runtime/MassEntity/Source/MassEntityTestSuite/Private/MassEntityTest.cpp
--- a/Engine/Plugins/Runtime/MassEntity/Source/MassEntityTestSuite/Private/MassEntityTest.cpp
+++ b/Engine/Plugins/Runtime/MassEntity/Source/MassEntityTestSuite/Private/MassEntityTest.cpp
@@ -157,10 +157,10 @@ struct FEntityTest_AddingRedundantFragment : FEntityTestBase
 {
 	virtual bool InstantTest() override
 	{
-		CA_ASSUME(EntitySubsystem);
-		const FMassEntityHandle Entity = EntitySubsystem->CreateEntity(FloatsArchetype);
-		EntitySubsystem->AddFragmentToEntity(Entity, FTestFragment_Float::StaticStruct());		
-		AITEST_EQUAL("Adding a fragment that a given entity\s archetype already has should do nothing", EntitySubsystem->GetArchetypeForEntity(Entity), FloatsArchetype);
+		CA_ASSUME(EntityManager);
+		const FMassEntityHandle Entity = EntityManager->CreateEntity(FloatsArchetype);
+		EntityManager->AddFragmentToEntity(Entity, FTestFragment_Float::StaticStruct());		
+		AITEST_EQUAL("Adding a fragment that a given entity\s archetype already has should do nothing", EntityManager->GetArchetypeForEntity(Entity), FloatsArchetype);
 
 		return true;
 	}
diff --git a/Engine/Plugins/Runtime/MassGameplay/Source/MassRepresentation/Public/MassRepresentationActorManagement.h b/Engine/Plugins/Runtime/MassGameplay/Source/MassRepresentation/Public/MassRepresentationActorManagement.h
--- a/Engine/Plugins/Runtime/MassGameplay/Source/MassRepresentation/Public/MassRepresentationActorManagement.h
+++ b/Engine/Plugins/Runtime/MassGameplay/Source/MassRepresentation/Public/MassRepresentationActorManagement.h
@@ -33,7 +33,7 @@ public:
 	/**
 	 * Returns an actor of the template type and setup fragments values from it
 	 * @param RepresentationSubsystem to use to get or spawn the actor
-	 * @param EntitySubsystem associated to the mass agent
+	 * @param EntityManager associated to the mass agent
 	 * @param MassAgent is the handle to the associated mass agent
 	 * @param ActorInfo is the fragment where we are going to store the actor pointer
 	 * @param Transform is the spatial information about where to spawn the actor
@@ -42,7 +42,7 @@ public:
 	 * @param Priority of this spawn request in comparison with the others, lower value means higher priority
 	 * @return the actor spawned
 	 */
-	virtual AActor* GetOrSpawnActor(UMassRepresentationSubsystem& RepresentationSubsystem, FMassEntityManager& EntitySubsystem
+	virtual AActor* GetOrSpawnActor(UMassRepresentationSubsystem& RepresentationSubsystem, FMassEntityManager& EntityManager
 		, const FMassEntityHandle MassAgent, const FTransform& Transform, const int16 TemplateActorIndex
 		, FMassActorSpawnRequestHandle& InOutSpawnRequestHandle, const float Priority) const;
 
@@ -86,7 +86,7 @@ public:
 	 * Static methods to Release an actor or cancel its spawning (calls ReleaseAnyActorOrCancelAnySpawning)
 	 * WARNING: This method will destroy the associated actor in any and by the same fact might also move the entity into a new archetype.
 	 *          So any reference to fragment might become invalid.
-	 * @param EntitySubsystem to use to retrieve the mass agent fragments
+	 * @param EntityManager to use to retrieve the mass agent fragments
 	 * @param MassAgent is the handle to the associated mass agent
 	 * @return True if actor was release or spawning request was canceled
 	 */
@@ -107,7 +107,7 @@ public:
 		, FMassActorFragment& ActorInfo, FMassRepresentationFragment& Representation, UMassActorSubsystem* ActorSubsystem = nullptr);
 
 	UE_DEPRECATED(5.4, "This flavor of GetOrSpawnActor has been deprecated due to a defunct parameter, OutActorInfo, that was never being used.")
-	virtual AActor* GetOrSpawnActor(UMassRepresentationSubsystem& RepresentationSubsystem, FMassEntityManager& EntitySubsystem
+	virtual AActor* GetOrSpawnActor(UMassRepresentationSubsystem& RepresentationSubsystem, FMassEntityManager& EntityManager
 		, const FMassEntityHandle MassAgent, FMassActorFragment&/* OutActorInfo*/, const FTransform& Transform, const int16 TemplateActorIndex
 		, FMassActorSpawnRequestHandle& InOutSpawnRequestHandle, const float Priority) const final;
 
diff --git a/Engine/Plugins/Runtime/MassGameplay/Source/MassRepresentation/Public/MassRepresentationProcessor.h b/Engine/Plugins/Runtime/MassGameplay/Source/MassRepresentation/Public/MassRepresentationProcessor.h
--- a/Engine/Plugins/Runtime/MassGameplay/Source/MassRepresentation/Public/MassRepresentationProcessor.h
+++ b/Engine/Plugins/Runtime/MassGameplay/Source/MassRepresentation/Public/MassRepresentationProcessor.h
@@ -54,7 +54,7 @@ protected:
 
 	/** 
 	 * Execution method for this processor 
-	 * @param EntitySubsystem is the system to execute the lambdas on each entity chunk
+	 * @param EntityManager is the system to execute the lambdas on each entity chunk
 	 * @param Context is the execution context to be passed when executing the lambdas
 	 */
 	virtual void Execute(FMassEntityManager& EntityManager, FMassExecutionContext& Context) override;
@@ -103,7 +103,7 @@ protected:
 
 	/**
 	 * Execution method for this processor
-	 * @param EntitySubsystem is the system to execute the lambdas on each entity chunk
+	 * @param EntityManager is the system to execute the lambdas on each entity chunk
 	 * @param Context is the execution context to be passed when executing the lambdas
 	 */
 	virtual void Execute(FMassEntityManager& EntityManager, FMassExecutionContext& Context) override;
