Minor massage of Mass Representation types' API to support fetching FMassISMCSharedData based on representation index

diff --git a/Engine/Plugins/Runtime/MassGameplay/Source/MassRepresentation/Private/MassRepresentationSubsystem.cpp b/Engine/Plugins/Runtime/MassGameplay/Source/MassRepresentation/Private/MassRepresentationSubsystem.cpp
--- a/Engine/Plugins/Runtime/MassGameplay/Source/MassRepresentation/Private/MassRepresentationSubsystem.cpp
+++ b/Engine/Plugins/Runtime/MassGameplay/Source/MassRepresentation/Private/MassRepresentationSubsystem.cpp
@@ -38,6 +38,12 @@ int16 UMassRepresentationSubsystem::AddVisualDescWithISMComponents(const FStatic
 	return VisualizationComponent->AddVisualDescWithISMComponents(Desc, ISMComponents);
 }
 
+const FMassISMCSharedData* UMassRepresentationSubsystem::GetISMCSharedDataForDescriptionIndex(const int32 DescriptionIndex) const
+{
+	check(VisualizationComponent);
+	return VisualizationComponent->GetISMCSharedDataForDescriptionIndex(DescriptionIndex);
+}
+
 void UMassRepresentationSubsystem::RemoveISMComponent(UInstancedStaticMeshComponent& ISMComponent)
 {
 PRAGMA_DISABLE_DEPRECATION_WARNINGS
diff --git a/Engine/Plugins/Runtime/MassGameplay/Source/MassRepresentation/Private/MassVisualizationComponent.cpp b/Engine/Plugins/Runtime/MassGameplay/Source/MassRepresentation/Private/MassVisualizationComponent.cpp
--- a/Engine/Plugins/Runtime/MassGameplay/Source/MassRepresentation/Private/MassVisualizationComponent.cpp
+++ b/Engine/Plugins/Runtime/MassGameplay/Source/MassRepresentation/Private/MassVisualizationComponent.cpp
@@ -138,6 +138,11 @@ int16 UMassVisualizationComponent::AddVisualDescWithISMComponents(const FStaticM
 	return (int16)VisualIndex;
 }
 
+const FMassISMCSharedData* UMassVisualizationComponent::GetISMCSharedDataForDescriptionIndex(const int32 DescriptionIndex) const
+{
+	return ISMCSharedData.GetDataForIndex(DescriptionIndex);
+}
+
 void UMassVisualizationComponent::RemoveVisualDescByIndex(const int32 VisualizationIndex)
 {
 	UE_MT_SCOPED_WRITE_ACCESS(InstancedStaticMeshInfosDetector);
@@ -191,9 +196,9 @@ void UMassVisualizationComponent::ConstructStaticMeshComponents()
 		for (const FMassStaticMeshInstanceVisualizationMeshDesc& MeshDesc : Info.Desc.Meshes)
 		{
 			FMassISMCSharedData* SharedData = ISMCSharedData.Find(GetTypeHash(MeshDesc));
-			UInstancedStaticMeshComponent* ISMC = SharedData ? SharedData->GetISMComponent() : nullptr;
+			UInstancedStaticMeshComponent* ISMC = SharedData ? SharedData->GetMutableISMComponent() : nullptr;
 
-			if (SharedData == nullptr || SharedData->GetISMComponent() == nullptr)
+			if (ISMC == nullptr)
 			{
 				ISMC = NewObject<UInstancedStaticMeshComponent>(ActorOwner, MeshDesc.ISMComponentClass);
 				CA_ASSUME(ISMC);
@@ -304,7 +309,7 @@ void UMassVisualizationComponent::ClearAllVisualInstances()
 	// Pool should already be empty, got a problem if it's not
 	for (int32 SharedDataIndex = 0; SharedDataIndex < ISMCSharedData.Num(); ++SharedDataIndex)
 	{
-		if (UInstancedStaticMeshComponent* InstancedStaticMeshComponent = ISMCSharedData.GetAtIndex(SharedDataIndex).GetISMComponent())
+		if (UInstancedStaticMeshComponent* InstancedStaticMeshComponent = ISMCSharedData.GetAtIndex(SharedDataIndex).GetMutableISMComponent())
 		{
 			InstancedStaticMeshComponent->ClearInstances();
 			InstancedStaticMeshComponent->DestroyComponent();
@@ -360,7 +365,7 @@ void UMassVisualizationComponent::HandleChangesWithExternalIDTracking(UInstanced
 	{
 		INC_DWORD_STAT_BY(STAT_Mass_VisualizationComponent_InstancesAddedNum, MassInstanceIds.Num());
 
-		FMassISMCSharedData::FIdMap& SharedIdMap = SharedData.GetIdMap();
+		FMassISMCSharedData::FIdMap& SharedIdMap = SharedData.GetMutableIdMap();
 		TArray<Experimental::FHashElementId> ElementIds;
 		ElementIds.SetNumUninitialized(MassInstanceIds.Num());
 		// Filter out all updates & duplicate adds
@@ -422,7 +427,7 @@ void UMassVisualizationComponent::ProcessRemoves(UInstancedStaticMeshComponent&
 {
 	if (!SharedData.GetRemoveInstanceIds().IsEmpty())
 	{
-		FMassISMCSharedData::FIdMap& SharedIdMap = SharedData.GetIdMap();
+		FMassISMCSharedData::FIdMap& SharedIdMap = SharedData.GetMutableIdMap();
 		INC_DWORD_STAT_BY(STAT_Mass_VisualizationComponent_InstancesRemovedNum, SharedData.GetRemoveInstanceIds().Num());
 
 		//RemoveInstanceWithIds(SharedData.GetRemoveInstanceIds());
@@ -467,7 +472,7 @@ void UMassVisualizationComponent::EndVisualChanges()
 	{
 		FMassISMCSharedData& SharedData = *It;
 
-		UInstancedStaticMeshComponent* ISMComponent = SharedData.GetISMComponent();
+		UInstancedStaticMeshComponent* ISMComponent = SharedData.GetMutableISMComponent();
 		// @todo need to check validity this way since Mass used to rely on the assumption that all the ISM components used were
 		// under its control. That's no longer the case, but the system has not been updated to take that into consideration.
 		// This is a temporary fix. 
@@ -496,7 +501,7 @@ void UMassVisualizationComponent::EndVisualChanges()
 					// Loop over all the instances in the update and 
 					// 1. Sort the data such that all Adds are last
 					// 2. Remove any duplicates (unsure if they may exist)
-					FMassISMCSharedData::FIdMap& SharedIdMap = SharedData.GetIdMap();
+					FMassISMCSharedData::FIdMap& SharedIdMap = SharedData.GetMutableIdMap();
 					// Filter out all updates & duplicate adds
 					TBitArray<> Unprocessed;
 					Unprocessed.SetNum(SharedIdMap.GetMaxIndex(), true);
@@ -588,6 +593,9 @@ void UMassVisualizationComponent::EndVisualChanges()
 					// note that we're not clearing the dirty flag on purpose - these components require constant updates
 				}
 			}
+
+			// bump the touch counter so that anyone caching data based on contents of this SharedData can detect the change
+			++SharedData.ComponentInstanceIdTouchCounter;
 		}
 		
 		SharedData.ResetAccumulatedData();
@@ -604,9 +612,9 @@ void FMassInstancedStaticMeshInfo::ClearVisualInstance(FMassISMCSharedDataMap& I
 	{
 		const uint32 MeshDescHash = GetTypeHash(Desc.Meshes[i]);
 		FMassISMCSharedData* SharedData = ISMCSharedData.Find(MeshDescHash);
-		if (SharedData && SharedData->ReleaseReference() == 0)
+		if (SharedData && SharedData->OnISMComponentReferenceReleased() == 0)
 		{
-			if (UInstancedStaticMeshComponent* ISMC = SharedData->GetISMComponent())
+			if (UInstancedStaticMeshComponent* ISMC = SharedData->GetMutableISMComponent())
 			{
 				ISMC->ClearInstances();
 				ISMC->DestroyComponent();
@@ -626,14 +634,14 @@ void FMassInstancedStaticMeshInfo::ClearVisualInstance(FMassISMCSharedDataMap& I
 void FMassLODSignificanceRange::AddBatchedTransform(const int32 InstanceId, const FTransform& Transform, const FTransform& PrevTransform, const TArray<uint32>& ExcludeStaticMeshRefs)
 {
 	check(ISMCSharedDataPtr);
-	for (int i = 0; i < StaticMeshRefs.Num(); i++)
+	for (int32 StaticMeshIndex = 0; StaticMeshIndex < StaticMeshRefs.Num(); ++StaticMeshIndex)
 	{
-		if (ExcludeStaticMeshRefs.Contains(StaticMeshRefs[i]))
+		if (ExcludeStaticMeshRefs.Contains(StaticMeshRefs[StaticMeshIndex]))
 		{
 			continue;
 		}
 
-		FMassISMCSharedData& SharedData = ISMCSharedDataPtr->GetAndMarkDirtyChecked(StaticMeshRefs[i]);
+		FMassISMCSharedData& SharedData = ISMCSharedDataPtr->GetAndMarkDirtyChecked(StaticMeshRefs[StaticMeshIndex]);
 
 		SharedData.UpdateInstanceIds.Add(InstanceId);
 		SharedData.StaticMeshInstanceTransforms.Add(Transform);
@@ -644,14 +652,14 @@ void FMassLODSignificanceRange::AddBatchedTransform(const int32 InstanceId, cons
 void FMassLODSignificanceRange::AddBatchedCustomDataFloats(const TArray<float>& CustomFloats, const TArray<uint32>& ExcludeStaticMeshRefs)
 {
 	check(ISMCSharedDataPtr);
-	for (int i = 0; i < StaticMeshRefs.Num(); i++)
+	for (int32 StaticMeshIndex = 0; StaticMeshIndex < StaticMeshRefs.Num(); ++StaticMeshIndex)
 	{
-		if (ExcludeStaticMeshRefs.Contains(StaticMeshRefs[i]))
+		if (ExcludeStaticMeshRefs.Contains(StaticMeshRefs[StaticMeshIndex]))
 		{
 			continue;
 		}
 
-		FMassISMCSharedData& SharedData = ISMCSharedDataPtr->GetAndMarkDirtyChecked(StaticMeshRefs[i]);
+		FMassISMCSharedData& SharedData = ISMCSharedDataPtr->GetAndMarkDirtyChecked(StaticMeshRefs[StaticMeshIndex]);
 		SharedData.StaticMeshInstanceCustomFloats.Append(CustomFloats);
 	}
 }
@@ -659,9 +667,9 @@ void FMassLODSignificanceRange::AddBatchedCustomDataFloats(const TArray<float>&
 void FMassLODSignificanceRange::AddInstance(const int32 InstanceId, const FTransform& Transform)
 {
 	check(ISMCSharedDataPtr);
-	for (int i = 0; i < StaticMeshRefs.Num(); i++)
+	for (int32 StaticMeshIndex = 0; StaticMeshIndex < StaticMeshRefs.Num(); ++StaticMeshIndex)
 	{
-		FMassISMCSharedData& SharedData = ISMCSharedDataPtr->GetAndMarkDirtyChecked(StaticMeshRefs[i]);
+		FMassISMCSharedData& SharedData = ISMCSharedDataPtr->GetAndMarkDirtyChecked(StaticMeshRefs[StaticMeshIndex]);
 		SharedData.UpdateInstanceIds.Add(InstanceId);
 		SharedData.StaticMeshInstanceTransforms.Add(Transform);
 		SharedData.StaticMeshInstancePrevTransforms.Add(Transform);
@@ -671,9 +679,9 @@ void FMassLODSignificanceRange::AddInstance(const int32 InstanceId, const FTrans
 void FMassLODSignificanceRange::RemoveInstance(const int32 InstanceId)
 {
 	check(ISMCSharedDataPtr);
-	for (int i = 0; i < StaticMeshRefs.Num(); i++)
+	for (int32 StaticMeshIndex = 0; StaticMeshIndex < StaticMeshRefs.Num(); ++StaticMeshIndex)
 	{
-		if (FMassISMCSharedData* SharedData = ISMCSharedDataPtr->GetAndMarkDirty(StaticMeshRefs[i]))
+		if (FMassISMCSharedData* SharedData = ISMCSharedDataPtr->GetAndMarkDirty(StaticMeshRefs[StaticMeshIndex]))
 		{
 			SharedData->RemoveInstanceIds.Add(InstanceId);
 		}
diff --git a/Engine/Plugins/Runtime/MassGameplay/Source/MassRepresentation/Public/MassRepresentationSubsystem.h b/Engine/Plugins/Runtime/MassGameplay/Source/MassRepresentation/Public/MassRepresentationSubsystem.h
--- a/Engine/Plugins/Runtime/MassGameplay/Source/MassRepresentation/Public/MassRepresentationSubsystem.h
+++ b/Engine/Plugins/Runtime/MassGameplay/Source/MassRepresentation/Public/MassRepresentationSubsystem.h
@@ -50,6 +50,11 @@ public:
 	 */
 	int16 AddVisualDescWithISMComponents(const FStaticMeshInstanceVisualizationDesc& Desc, TArrayView<TObjectPtr<UInstancedStaticMeshComponent>> ISMComponents);
 
+	/**
+	 * Fetches FMassISMCSharedData indicated by DescriptionIndex, or nullptr if it's not a valid index
+	 */
+	const FMassISMCSharedData* GetISMCSharedDataForDescriptionIndex(const int32 DescriptionIndex) const;
+
 	/**
 	 * Removes the visualization data associated with the given ISM component. Note that this is safe to do only when
 	 * there are no entities relying on this data. No entity data patching will take place.
diff --git a/Engine/Plugins/Runtime/MassGameplay/Source/MassRepresentation/Public/MassRepresentationTypes.h b/Engine/Plugins/Runtime/MassGameplay/Source/MassRepresentation/Public/MassRepresentationTypes.h
--- a/Engine/Plugins/Runtime/MassGameplay/Source/MassRepresentation/Public/MassRepresentationTypes.h
+++ b/Engine/Plugins/Runtime/MassGameplay/Source/MassRepresentation/Public/MassRepresentationTypes.h
@@ -148,6 +148,7 @@ class UInstancedStaticMeshComponent;
 
 struct MASSREPRESENTATION_API FMassISMCSharedData
 {
+PRAGMA_DISABLE_DEPRECATION_WARNINGS
 	FMassISMCSharedData() 
 		: bRequiresExternalInstanceIDTracking(false)
 	{		
@@ -158,15 +159,20 @@ struct MASSREPRESENTATION_API FMassISMCSharedData
 	{
 	}
 
+	FMassISMCSharedData(const FMassISMCSharedData& Other) = default;
+	FMassISMCSharedData& operator=(const FMassISMCSharedData& Other) = default;
+PRAGMA_ENABLE_DEPRECATION_WARNINGS
+
 	void SetISMComponent(UInstancedStaticMeshComponent& InISMC)
 	{
-		check(ISMC == nullptr && RefCount == 0);
+		check(ISMC == nullptr && ISMComponentReferencesCount == 0);
 		ISMC = &InISMC;
 	}
 
-	UInstancedStaticMeshComponent* GetISMComponent() { return ISMC; }
-	int32 StoreReference() { return ++RefCount; }
-	int32 ReleaseReference() { ensure(RefCount >= 0); return --RefCount; }
+	UInstancedStaticMeshComponent* GetMutableISMComponent() { return ISMC; }
+	const UInstancedStaticMeshComponent* GetISMComponent() const { return ISMC; }
+	int32 OnISMComponentReferenceStored() { return ++ISMComponentReferencesCount; }
+	int32 OnISMComponentReferenceReleased() { ensure(ISMComponentReferencesCount >= 0); return --ISMComponentReferencesCount; }
 
 	void ResetAccumulatedData()
 	{
@@ -210,7 +216,10 @@ struct MASSREPRESENTATION_API FMassISMCSharedData
 
 	using FIdMap = Experimental::TRobinHoodHashMap<int32, FPrimitiveInstanceId>;
 
-	FIdMap &GetIdMap() { return MassInstanceIdToComponentInstanceIdMap; }
+	FIdMap& GetMutableIdMap() { return MassInstanceIdToComponentInstanceIdMap; }
+	const FIdMap& GetIdMap() const { return MassInstanceIdToComponentInstanceIdMap; }
+
+	int16 GetComponentInstanceIdTouchCounter() const { return ComponentInstanceIdTouchCounter; }
 
 protected:
 	friend FMassLODSignificanceRange;
@@ -228,7 +237,7 @@ protected:
 	int32 WriteIterator = 0;
 
 	UInstancedStaticMeshComponent* ISMC = nullptr;
-	int32 RefCount = 0;
+	int32 ISMComponentReferencesCount = 0;
 
 	/** 
 	 * When set to true will result in MassVisualizationComponent manually perform Instance ID-related operations 
@@ -237,8 +246,23 @@ protected:
 	 *	instance ID logic. WIP as of Jun 17th 2023 
 	 */
 	uint8 bRequiresExternalInstanceIDTracking : 1;
+	
+private:
+	/** Indicates that mutating changes, that can affect MassInstanceIdToComponentInstanceIdMap, have been performed.
+	 *	Can be used to validate whether cached data stored in other placed needs to be re-cached. */
+	uint16 ComponentInstanceIdTouchCounter = 0;
 
+protected:
 	FIdMap MassInstanceIdToComponentInstanceIdMap;
+
+	UE_DEPRECATED(5.4, "RefCount is deprecated, use ISMComponentReferencesCount instead")
+	int32 RefCount = 0;
+
+public:
+	UE_DEPRECATED(5.4, "StoreReference is deprecated, use OnISMComponentReferenceStored instead")
+	int32 StoreReference() { return OnISMComponentReferenceStored(); }
+	UE_DEPRECATED(5.4, "ReleaseReference is deprecated, use OnISMComponentReferenceReleased instead")
+	int32 ReleaseReference() { return OnISMComponentReferenceReleased(); }
 };
 
 
@@ -315,12 +339,13 @@ struct FMassISMCSharedDataMap
 		return nullptr;
 	}
 	
-	FMassISMCSharedData& FindOrAdd(const uint32 Hash, const FMassISMCSharedData& NewData)
+	template<typename... TArgs>
+	FMassISMCSharedData& FindOrAdd(const uint32 Hash, TArgs&&... InNewInstanceArgs)
 	{
 		const int32* DataIndex = Map.Find(Hash);
 		if (DataIndex == nullptr)
 		{
-			return Add(Hash, NewData);
+			return Add(Hash, Forward<TArgs>(InNewInstanceArgs)...);
 		}
 		check(Data.IsValidIndex(*DataIndex));
 		return Data[*DataIndex];
@@ -332,7 +357,8 @@ struct FMassISMCSharedDataMap
 		return (DataIndex == nullptr || *DataIndex == INDEX_NONE) ? (FMassISMCSharedData*)nullptr : &Data[*DataIndex];
 	}
 
-	FMassISMCSharedData& Add(const uint32 Hash, const FMassISMCSharedData& NewData)
+	template<typename... TArgs>
+	FMassISMCSharedData& Add(const uint32 Hash, TArgs&&... InNewInstanceArgs)
 	{
 		const int32 DataIndex = FreeIndices.Num() ? FreeIndices.Pop() : Data.Num();
 		Map.Add(Hash, DataIndex);
@@ -341,12 +367,12 @@ struct FMassISMCSharedDataMap
 		{
 			DirtyData.Add(false, DataIndex - DirtyData.Num() + 1);
 			DirtyData[DataIndex] = true;
-			return Data.Add_GetRef(NewData);
+			return Data.Add_GetRef(FMassISMCSharedData(Forward<TArgs>(InNewInstanceArgs)...));
 		}
 		else
 		{
 			DirtyData[DataIndex] = true;
-			Data[DataIndex] = NewData;
+			Data[DataIndex] = FMassISMCSharedData(Forward<TArgs>(InNewInstanceArgs)...);
 			return Data[DataIndex];
 		}
 	}
@@ -394,6 +420,11 @@ struct FMassISMCSharedDataMap
 		*this = FMassISMCSharedDataMap();
 	}
 
+	const FMassISMCSharedData* GetDataForIndex(const int32 Index) const
+	{
+		return Data.IsValidIndex(Index) ? &Data[Index] : nullptr;
+	}
+
 protected:
 	TArray<FMassISMCSharedData> Data;
 	/** Mapping from Hash of data represented by FMassISMCSharedData to an index to Data */
@@ -576,8 +607,8 @@ public:
 	{
 		if (ensure(SharedData.GetISMComponent()))
 		{
-			InstancedStaticMeshComponents.Add(SharedData.GetISMComponent());
-			SharedData.StoreReference();
+			InstancedStaticMeshComponents.Add(SharedData.GetMutableISMComponent());
+			SharedData.OnISMComponentReferenceStored();
 		}
 	}
 
diff --git a/Engine/Plugins/Runtime/MassGameplay/Source/MassRepresentation/Public/MassVisualizationComponent.h b/Engine/Plugins/Runtime/MassGameplay/Source/MassRepresentation/Public/MassVisualizationComponent.h
--- a/Engine/Plugins/Runtime/MassGameplay/Source/MassRepresentation/Public/MassVisualizationComponent.h
+++ b/Engine/Plugins/Runtime/MassGameplay/Source/MassRepresentation/Public/MassVisualizationComponent.h
@@ -42,6 +42,11 @@ public:
 	 */
 	int16 AddVisualDescWithISMComponents(const FStaticMeshInstanceVisualizationDesc& Desc, TArrayView<TObjectPtr<UInstancedStaticMeshComponent>> ISMComponents);
 
+	/**
+	 * Fetches FMassISMCSharedData indicated by DescriptionIndex, or nullptr if it's not a valid index
+	 */
+	const FMassISMCSharedData* GetISMCSharedDataForDescriptionIndex(const int32 DescriptionIndex) const;
+
 	/** 
 	 * Removes all the visualization data associated with the given ISM component. The function resolves the VisualizationIndex
 	 * associated with the given ISMComponent and calls RemoveVisualDescByIndex which will remove data on all ISMComponents
