Mass Visualization - Switched int16 StaticMeshDescIndex handles to new explicitly typed FStaticMeshInstanceVisualizationDescHandle which uses uint16 internally to double the number of possible ISMC visualization registrations.

This fixes an issue where extreme usage of UMassRepresentationSubsystem::AddVisualDescWithISMComponents for active registration of more than 32,767 ISMC groups would overflow the previous INT16_MAX handle cap. Now this lifts the limit to 65534 with 65535 (UINT16_MAX) reserved for representing invalid handles where previously INDEX_NONE was used.

#tests cooked

diff --git a/Engine/Plugins/AI/MassCrowd/Source/MassCrowd/Private/MassCrowdServerRepresentationTrait.cpp b/Engine/Plugins/AI/MassCrowd/Source/MassCrowd/Private/MassCrowdServerRepresentationTrait.cpp
--- a/Engine/Plugins/AI/MassCrowd/Source/MassCrowd/Private/MassCrowdServerRepresentationTrait.cpp
+++ b/Engine/Plugins/AI/MassCrowd/Source/MassCrowd/Private/MassCrowdServerRepresentationTrait.cpp
@@ -51,7 +51,7 @@ void UMassCrowdServerRepresentationTrait::BuildTemplate(FMassEntityTemplateBuild
 	BuildContext.AddConstSharedFragment(ParamsFragment);
 
 	FMassRepresentationFragment& RepresentationFragment = BuildContext.AddFragment_GetRef<FMassRepresentationFragment>();
-	RepresentationFragment.StaticMeshDescIndex = INDEX_NONE;
+	RepresentationFragment.StaticMeshDescHandle = FStaticMeshInstanceVisualizationDescHandle();
 	RepresentationFragment.HighResTemplateActorIndex = TemplateActor.Get() ? RepresentationSubsystem->FindOrAddTemplateActor(TemplateActor.Get()) : INDEX_NONE;
 	RepresentationFragment.LowResTemplateActorIndex = INDEX_NONE;
 
diff --git a/Engine/Plugins/Runtime/MassGameplay/Source/MassRepresentation/Private/MassDistanceVisualizationTrait.cpp b/Engine/Plugins/Runtime/MassGameplay/Source/MassRepresentation/Private/MassDistanceVisualizationTrait.cpp
--- a/Engine/Plugins/Runtime/MassGameplay/Source/MassRepresentation/Private/MassDistanceVisualizationTrait.cpp
+++ b/Engine/Plugins/Runtime/MassGameplay/Source/MassRepresentation/Private/MassDistanceVisualizationTrait.cpp
@@ -73,7 +73,7 @@ void UMassDistanceVisualizationTrait::BuildTemplate(FMassEntityTemplateBuildCont
 	FMassRepresentationFragment& RepresentationFragment = BuildContext.AddFragment_GetRef<FMassRepresentationFragment>();
 	if (bRegisterStaticMeshDesc)
 	{
-		RepresentationFragment.StaticMeshDescIndex = RepresentationSubsystem->FindOrAddStaticMeshDesc(StaticMeshInstanceDesc);
+		RepresentationFragment.StaticMeshDescHandle = RepresentationSubsystem->FindOrAddStaticMeshDesc(StaticMeshInstanceDesc);
 	}
 	RepresentationFragment.HighResTemplateActorIndex = HighResTemplateActor.Get() ? RepresentationSubsystem->FindOrAddTemplateActor(HighResTemplateActor.Get()) : INDEX_NONE;
 	RepresentationFragment.LowResTemplateActorIndex = LowResTemplateActor.Get() ? RepresentationSubsystem->FindOrAddTemplateActor(LowResTemplateActor.Get()) : INDEX_NONE;
diff --git a/Engine/Plugins/Runtime/MassGameplay/Source/MassRepresentation/Private/MassRepresentationSubsystem.cpp b/Engine/Plugins/Runtime/MassGameplay/Source/MassRepresentation/Private/MassRepresentationSubsystem.cpp
--- a/Engine/Plugins/Runtime/MassGameplay/Source/MassRepresentation/Private/MassRepresentationSubsystem.cpp
+++ b/Engine/Plugins/Runtime/MassGameplay/Source/MassRepresentation/Private/MassRepresentationSubsystem.cpp
@@ -20,19 +20,19 @@
 #include "MassEntityUtils.h"
 
 
-int16 UMassRepresentationSubsystem::FindOrAddStaticMeshDesc(const FStaticMeshInstanceVisualizationDesc& Desc)
+FStaticMeshInstanceVisualizationDescHandle UMassRepresentationSubsystem::FindOrAddStaticMeshDesc(const FStaticMeshInstanceVisualizationDesc& Desc)
 {
 	check(VisualizationComponent);
 	return VisualizationComponent->FindOrAddVisualDesc(Desc);
 }
 
-int16 UMassRepresentationSubsystem::AddVisualDescWithISMComponent(const FStaticMeshInstanceVisualizationDesc& Desc, UInstancedStaticMeshComponent& ISMComponent)
+FStaticMeshInstanceVisualizationDescHandle UMassRepresentationSubsystem::AddVisualDescWithISMComponent(const FStaticMeshInstanceVisualizationDesc& Desc, UInstancedStaticMeshComponent& ISMComponent)
 {
 	check(VisualizationComponent);
 	return VisualizationComponent->AddVisualDescWithISMComponent(Desc, ISMComponent);
 }
 
-int16 UMassRepresentationSubsystem::AddVisualDescWithISMComponents(const FStaticMeshInstanceVisualizationDesc& Desc, TArrayView<TObjectPtr<UInstancedStaticMeshComponent>> ISMComponents)
+FStaticMeshInstanceVisualizationDescHandle UMassRepresentationSubsystem::AddVisualDescWithISMComponents(const FStaticMeshInstanceVisualizationDesc& Desc, TArrayView<TObjectPtr<UInstancedStaticMeshComponent>> ISMComponents)
 {
 	check(VisualizationComponent);
 	return VisualizationComponent->AddVisualDescWithISMComponents(Desc, ISMComponents);
@@ -52,10 +52,10 @@ PRAGMA_DISABLE_DEPRECATION_WARNINGS
 PRAGMA_ENABLE_DEPRECATION_WARNINGS
 }
 
-void UMassRepresentationSubsystem::RemoveVisualDescByIndex(const int32 VisualizationIndex)
+void UMassRepresentationSubsystem::RemoveVisualDesc(const FStaticMeshInstanceVisualizationDescHandle VisualizationHandle)
 {
 	check(VisualizationComponent);
-	return VisualizationComponent->RemoveVisualDescByIndex(VisualizationIndex);
+	return VisualizationComponent->RemoveVisualDesc(VisualizationHandle);
 }
 
 FMassInstancedStaticMeshInfoArrayView UMassRepresentationSubsystem::GetMutableInstancedStaticMeshInfos()
diff --git a/Engine/Plugins/Runtime/MassGameplay/Source/MassRepresentation/Private/MassStationaryISMRepresentationFragmentDestructor.cpp b/Engine/Plugins/Runtime/MassGameplay/Source/MassRepresentation/Private/MassStationaryISMRepresentationFragmentDestructor.cpp
--- a/Engine/Plugins/Runtime/MassGameplay/Source/MassRepresentation/Private/MassStationaryISMRepresentationFragmentDestructor.cpp
+++ b/Engine/Plugins/Runtime/MassGameplay/Source/MassRepresentation/Private/MassStationaryISMRepresentationFragmentDestructor.cpp
@@ -41,7 +41,7 @@ void UMassStationaryISMRepresentationFragmentDestructor::Execute(FMassEntityMana
 			FMassRepresentationFragment& Representation = RepresentationList[EntityIdx];
 			if (Representation.CurrentRepresentation == EMassRepresentationType::StaticMeshInstance)
 			{
-				FMassInstancedStaticMeshInfo& ISMInfo = ISMInfosView[Representation.StaticMeshDescIndex];
+				FMassInstancedStaticMeshInfo& ISMInfo = ISMInfosView[Representation.StaticMeshDescHandle.ToIndex()];
 				if (FMassLODSignificanceRange* OldRange = ISMInfo.GetLODSignificanceRange(Representation.PrevLODSignificance))
 				{
 					const FMassEntityHandle EntityHandle = Context.GetEntity(EntityIdx);
diff --git a/Engine/Plugins/Runtime/MassGameplay/Source/MassRepresentation/Private/MassStationaryISMSwitcherProcessor.cpp b/Engine/Plugins/Runtime/MassGameplay/Source/MassRepresentation/Private/MassStationaryISMSwitcherProcessor.cpp
--- a/Engine/Plugins/Runtime/MassGameplay/Source/MassRepresentation/Private/MassStationaryISMSwitcherProcessor.cpp
+++ b/Engine/Plugins/Runtime/MassGameplay/Source/MassRepresentation/Private/MassStationaryISMSwitcherProcessor.cpp
@@ -63,12 +63,12 @@ void UMassStationaryISMSwitcherProcessor::ProcessContext(FMassExecutionContext&
 			continue;
 		}
 
-		if (!ensureMsgf(Representation.StaticMeshDescIndex != INDEX_NONE && ISMInfosView.IsValidIndex(Representation.StaticMeshDescIndex)
-						, TEXT("Invalid index %d for ISMInfosView"), Representation.StaticMeshDescIndex))
+		if (!ensureMsgf(Representation.StaticMeshDescHandle.IsValid() && ISMInfosView.IsValidIndex(Representation.StaticMeshDescHandle.ToIndex())
+						, TEXT("Invalid handle index %u for ISMInfosView"), Representation.StaticMeshDescHandle.ToIndex()))
 		{
 			continue;
 		}
-		FMassInstancedStaticMeshInfo& ISMInfo = ISMInfosView[Representation.StaticMeshDescIndex];
+		FMassInstancedStaticMeshInfo& ISMInfo = ISMInfosView[Representation.StaticMeshDescHandle.ToIndex()];
 
 		if (Representation.PrevRepresentation == EMassRepresentationType::StaticMeshInstance
 			&& Representation.CurrentRepresentation != EMassRepresentationType::StaticMeshInstance)
diff --git a/Engine/Plugins/Runtime/MassGameplay/Source/MassRepresentation/Private/MassUpdateISMProcessor.cpp b/Engine/Plugins/Runtime/MassGameplay/Source/MassRepresentation/Private/MassUpdateISMProcessor.cpp
--- a/Engine/Plugins/Runtime/MassGameplay/Source/MassRepresentation/Private/MassUpdateISMProcessor.cpp
+++ b/Engine/Plugins/Runtime/MassGameplay/Source/MassRepresentation/Private/MassUpdateISMProcessor.cpp
@@ -53,7 +53,7 @@ void UMassUpdateISMProcessor::Execute(FMassEntityManager& EntityManager, FMassEx
 
 			if (Representation.CurrentRepresentation == EMassRepresentationType::StaticMeshInstance)
 			{
-				UpdateISMTransform(GetTypeHash(Context.GetEntity(EntityIdx)), ISMInfo[Representation.StaticMeshDescIndex], TransformFragment.GetTransform(), Representation.PrevTransform, RepresentationLOD.LODSignificance, Representation.PrevLODSignificance);
+				UpdateISMTransform(GetTypeHash(Context.GetEntity(EntityIdx)), ISMInfo[Representation.StaticMeshDescHandle.ToIndex()], TransformFragment.GetTransform(), Representation.PrevTransform, RepresentationLOD.LODSignificance, Representation.PrevLODSignificance);
 			}
 			Representation.PrevTransform = TransformFragment.GetTransform();
 			Representation.PrevLODSignificance = RepresentationLOD.LODSignificance;
diff --git a/Engine/Plugins/Runtime/MassGameplay/Source/MassRepresentation/Private/MassVisualizationComponent.cpp b/Engine/Plugins/Runtime/MassGameplay/Source/MassRepresentation/Private/MassVisualizationComponent.cpp
--- a/Engine/Plugins/Runtime/MassGameplay/Source/MassRepresentation/Private/MassVisualizationComponent.cpp
+++ b/Engine/Plugins/Runtime/MassGameplay/Source/MassRepresentation/Private/MassVisualizationComponent.cpp
@@ -46,27 +46,30 @@ void UMassVisualizationComponent::PostInitProperties()
 	}
 }
 
-int32 UMassVisualizationComponent::AddInstancedStaticMeshInfo(const FStaticMeshInstanceVisualizationDesc& Desc)
+FStaticMeshInstanceVisualizationDescHandle UMassVisualizationComponent::AddInstancedStaticMeshInfo(const FStaticMeshInstanceVisualizationDesc& Desc)
 {
-	int32 Index = INDEX_NONE;
+	FStaticMeshInstanceVisualizationDescHandle Handle;
 	if (InstancedStaticMeshInfosFreeIndices.Num() > 0)
 	{
-		Index = InstancedStaticMeshInfosFreeIndices.Pop(/*bAllowShrinking=*/false);
-		new(&InstancedStaticMeshInfos[Index]) FMassInstancedStaticMeshInfo(Desc);
+		Handle = InstancedStaticMeshInfosFreeIndices.Pop(/*bAllowShrinking=*/false);
+		new(&InstancedStaticMeshInfos[Handle.ToIndex()]) FMassInstancedStaticMeshInfo(Desc);
 	}
 	else
 	{
-		Index = InstancedStaticMeshInfos.Emplace(Desc);
+		int32 AddedInfoIndex = InstancedStaticMeshInfos.Emplace(Desc);
+		Handle = FStaticMeshInstanceVisualizationDescHandle(AddedInfoIndex);
 	}
 
-	return Index;
+	return Handle;
 }
 
-int16 UMassVisualizationComponent::FindOrAddVisualDesc(const FStaticMeshInstanceVisualizationDesc& Desc)
+FStaticMeshInstanceVisualizationDescHandle UMassVisualizationComponent::FindOrAddVisualDesc(const FStaticMeshInstanceVisualizationDesc& Desc)
 {
 	UE_MT_SCOPED_WRITE_ACCESS(InstancedStaticMeshInfosDetector);
-	int32 VisualIndex = InstancedStaticMeshInfos.IndexOfByPredicate([&Desc](const FMassInstancedStaticMeshInfo& Info) { return Info.GetDesc() == Desc; });
-	if (VisualIndex == INDEX_NONE)
+	// First check to see if we already have a matching Desc already and reuse / return that
+	// Note: FStaticMeshInstanceVisualizationDescHandle(int32) handles the INDEX_NONE case here, generating an invalid handle in this case
+	FStaticMeshInstanceVisualizationDescHandle VisualDescHandle(InstancedStaticMeshInfos.IndexOfByPredicate([&Desc](const FMassInstancedStaticMeshInfo& Info) { return Info.GetDesc() == Desc; }));
+	if (!VisualDescHandle.IsValid())
 	{
 		bool bValidDescription = false;
 
@@ -81,29 +84,30 @@ int16 UMassVisualizationComponent::FindOrAddVisualDesc(const FStaticMeshInstance
 
 		if (bValidDescription)
 		{
-			VisualIndex = AddInstancedStaticMeshInfo(Desc);
-			BuildLODSignificanceForInfo(InstancedStaticMeshInfos[VisualIndex]);
+			VisualDescHandle = AddInstancedStaticMeshInfo(Desc);
+			BuildLODSignificanceForInfo(InstancedStaticMeshInfos[VisualDescHandle.ToIndex()]);
 
-			InstancedSMComponentsRequiringConstructing.Add(VisualIndex);
+			InstancedSMComponentsRequiringConstructing.Add(VisualDescHandle);
 		}
 	}
-	checkf(VisualIndex < INT16_MAX, TEXT("%hs resulting VisualIndex is out of expected bounds"), __FUNCTION__);
-	return (int16)VisualIndex;
+
+	check(VisualDescHandle.IsValid());
+	return VisualDescHandle;
 }
 
-int16 UMassVisualizationComponent::AddVisualDescWithISMComponent(const FStaticMeshInstanceVisualizationDesc& Desc, UInstancedStaticMeshComponent& ISMComponent)
+FStaticMeshInstanceVisualizationDescHandle UMassVisualizationComponent::AddVisualDescWithISMComponent(const FStaticMeshInstanceVisualizationDesc& Desc, UInstancedStaticMeshComponent& ISMComponent)
 {
 	TObjectPtr<UInstancedStaticMeshComponent> AsObjectPtr = ISMComponent;
 	return AddVisualDescWithISMComponents(Desc, MakeArrayView(&AsObjectPtr, 1));
 }
 
-int16 UMassVisualizationComponent::AddVisualDescWithISMComponents(const FStaticMeshInstanceVisualizationDesc& Desc, TArrayView<TObjectPtr<UInstancedStaticMeshComponent>> ISMComponents)
+FStaticMeshInstanceVisualizationDescHandle UMassVisualizationComponent::AddVisualDescWithISMComponents(const FStaticMeshInstanceVisualizationDesc& Desc, TArrayView<TObjectPtr<UInstancedStaticMeshComponent>> ISMComponents)
 {
 	check(Desc.Meshes.Num() == ISMComponents.Num());
 	
 	UE_MT_SCOPED_WRITE_ACCESS(InstancedStaticMeshInfosDetector);
 
-	int32 VisualIndex = INDEX_NONE;
+	FStaticMeshInstanceVisualizationDescHandle VisualHandle;
 	TArray<uint32> ISMComponentPathHashes;
 
 	for (int32 EntryIndex = 0; EntryIndex < Desc.Meshes.Num(); ++EntryIndex)
@@ -116,26 +120,26 @@ int16 UMassVisualizationComponent::AddVisualDescWithISMComponents(const FStaticM
 			continue;
 		}
 	
-		if (VisualIndex == INDEX_NONE)
+		if (!VisualHandle.IsValid())
 		{
-			VisualIndex = AddInstancedStaticMeshInfo(Desc);
+			VisualHandle = AddInstancedStaticMeshInfo(Desc);
+			check(VisualHandle.IsValid());
 		}
 
 		const uint32 ISMComponentPathHash = GetTypeHash(ISMComponents[EntryIndex]->GetPathName());
 		FMassISMCSharedData& NewData = ISMCSharedData.FindOrAdd(ISMComponentPathHash, FMassISMCSharedData(ISMComponents[EntryIndex], /*bInRequiresExternalInstanceIDTracking=*/true));
-		InstancedStaticMeshInfos[VisualIndex].AddISMComponent(NewData);
+		InstancedStaticMeshInfos[VisualHandle.ToIndex()].AddISMComponent(NewData);
 		ISMComponentPathHashes.Add(ISMComponentPathHash);
 	
-		ISMComponentMap.Add(ISMComponentPathHash, VisualIndex);
+		ISMComponentMap.Add(ISMComponentPathHash, VisualHandle);
 	}
 
-	if (VisualIndex != INDEX_NONE)
+	if (VisualHandle.IsValid())
 	{
-		BuildLODSignificanceForInfo(InstancedStaticMeshInfos[VisualIndex], ISMComponentPathHashes);
+		BuildLODSignificanceForInfo(InstancedStaticMeshInfos[VisualHandle.ToIndex()], ISMComponentPathHashes);
 	}
 
-	checkf(VisualIndex < INT16_MAX, TEXT("%hs resulting VisualIndex is out of expected bounds"), __FUNCTION__);
-	return (int16)VisualIndex;
+	return VisualHandle;
 }
 
 const FMassISMCSharedData* UMassVisualizationComponent::GetISMCSharedDataForDescriptionIndex(const int32 DescriptionIndex) const
@@ -143,24 +147,24 @@ const FMassISMCSharedData* UMassVisualizationComponent::GetISMCSharedDataForDesc
 	return ISMCSharedData.GetDataForIndex(DescriptionIndex);
 }
 
-void UMassVisualizationComponent::RemoveVisualDescByIndex(const int32 VisualizationIndex)
+void UMassVisualizationComponent::RemoveVisualDesc(const FStaticMeshInstanceVisualizationDescHandle VisualizationHandle)
 {
 	UE_MT_SCOPED_WRITE_ACCESS(InstancedStaticMeshInfosDetector);
-
-	if (ensure(InstancedStaticMeshInfos.IsValidIndex(VisualizationIndex))
-		&& ensureMsgf(InstancedStaticMeshInfos[VisualizationIndex].IsValid(), TEXT("Trying to remove visualization data that has already been cleaned")))
+	
+	if (ensure(InstancedStaticMeshInfos.IsValidIndex(VisualizationHandle.ToIndex()))
+		&& ensureMsgf(InstancedStaticMeshInfos[VisualizationHandle.ToIndex()].IsValid(), TEXT("Trying to remove visualization data that has already been cleaned")))
 	{
-		for (TObjectPtr<UInstancedStaticMeshComponent>& ISMComponent : InstancedStaticMeshInfos[VisualizationIndex].InstancedStaticMeshComponents)
+		for (TObjectPtr<UInstancedStaticMeshComponent>& ISMComponent : InstancedStaticMeshInfos[VisualizationHandle.ToIndex()].InstancedStaticMeshComponents)
 		{
 			const uint32 ISMComponentPathHash = GetTypeHash(ISMComponent.GetPathName());
-			const int32 StoredVisualizationIndex = ISMComponentMap.FindAndRemoveChecked(ISMComponentPathHash);
-			ensure(StoredVisualizationIndex == VisualizationIndex);
+			const FStaticMeshInstanceVisualizationDescHandle StoredVisualizationDescHandle = ISMComponentMap.FindAndRemoveChecked(ISMComponentPathHash);
+			ensure(StoredVisualizationDescHandle == VisualizationHandle);
 		
 			ISMCSharedData.Remove(ISMComponentPathHash);
 		}
 		
-		InstancedStaticMeshInfos[VisualizationIndex].Reset();
-		InstancedStaticMeshInfosFreeIndices.Add(VisualizationIndex);
+		InstancedStaticMeshInfos[VisualizationHandle.ToIndex()].Reset();
+		InstancedStaticMeshInfosFreeIndices.Add(VisualizationHandle);
 	}
 }
 
@@ -170,16 +174,16 @@ void UMassVisualizationComponent::ConstructStaticMeshComponents()
 	check(ActorOwner);
 	
 	UE_MT_SCOPED_WRITE_ACCESS(InstancedStaticMeshInfosDetector);
-	for (const int32 VisualIndex : InstancedSMComponentsRequiringConstructing)
+	for (const FStaticMeshInstanceVisualizationDescHandle VisualDescHandle : InstancedSMComponentsRequiringConstructing)
 	{
-		if (!ensureMsgf(InstancedStaticMeshInfos.IsValidIndex(VisualIndex)
-			, TEXT("InstancedStaticMeshInfos (size: %d) is never expected to shrink, so VisualIndex (value: %d) being invalid indicates it was wrong from the start.")
-			, InstancedStaticMeshInfos.Num(), VisualIndex))
+		if (!ensureMsgf(InstancedStaticMeshInfos.IsValidIndex(VisualDescHandle.ToIndex())
+			, TEXT("InstancedStaticMeshInfos (size: %d) is never expected to shrink, so VisualDescHandle (value: %u) being invalid indicates it was wrong from the start.")
+			, InstancedStaticMeshInfos.Num(), VisualDescHandle.ToIndex()))
 		{
 			continue;
 		}
 
-		FMassInstancedStaticMeshInfo& Info = InstancedStaticMeshInfos[VisualIndex];
+		FMassInstancedStaticMeshInfo& Info = InstancedStaticMeshInfos[VisualDescHandle.ToIndex()];
 
 		// Check if it is already created
 		if (!Info.InstancedStaticMeshComponents.IsEmpty())
@@ -232,7 +236,7 @@ void UMassVisualizationComponent::ConstructStaticMeshComponents()
 
 				const uint32 ISMComponentPathHash = GetTypeHash(ISMC->GetPathName());
 				ensureMsgf(ISMComponentMap.Find(ISMComponentPathHash) == nullptr, TEXT("We've just created the ISMC that's being used here, so this check failing indicates hash-clash."));
-				ISMComponentMap.Add(ISMComponentPathHash, VisualIndex); 
+				ISMComponentMap.Add(ISMComponentPathHash, VisualDescHandle); 
 			}
 
 			check(SharedData);
@@ -725,10 +729,14 @@ void UMassVisualizationComponent::RemoveISMComponent(UInstancedStaticMeshCompone
 	UE_MT_SCOPED_WRITE_ACCESS(InstancedStaticMeshInfosDetector);
 
 	const uint32 ISMComponentPathHash = GetTypeHash(ISMComponent.GetPathName());
-	const int32* VisualIndexPtr = ISMComponentMap.Find(ISMComponentPathHash);
-	if (VisualIndexPtr)
+	const FStaticMeshInstanceVisualizationDescHandle* VisualDescHandlePtr = ISMComponentMap.Find(ISMComponentPathHash);
+	if (VisualDescHandlePtr)
 	{
-		RemoveVisualDescByIndex(*VisualIndexPtr);
+		RemoveVisualDesc(*VisualDescHandlePtr);
 	}
 }
 
+void UMassVisualizationComponent::RemoveVisualDescByIndex(const int32 VisualizationIndex)
+{
+	RemoveVisualDesc(FStaticMeshInstanceVisualizationDescHandle(VisualizationIndex));
+}
diff --git a/Engine/Plugins/Runtime/MassGameplay/Source/MassRepresentation/Private/MassVisualizationTrait.cpp b/Engine/Plugins/Runtime/MassGameplay/Source/MassRepresentation/Private/MassVisualizationTrait.cpp
--- a/Engine/Plugins/Runtime/MassGameplay/Source/MassRepresentation/Private/MassVisualizationTrait.cpp
+++ b/Engine/Plugins/Runtime/MassGameplay/Source/MassRepresentation/Private/MassVisualizationTrait.cpp
@@ -85,7 +85,7 @@ void UMassVisualizationTrait::BuildTemplate(FMassEntityTemplateBuildContext& Bui
 	FMassRepresentationFragment& RepresentationFragment = BuildContext.AddFragment_GetRef<FMassRepresentationFragment>();
 	if (bRegisterStaticMeshDesc)
 	{
-		RepresentationFragment.StaticMeshDescIndex = RepresentationSubsystem->FindOrAddStaticMeshDesc(StaticMeshInstanceDesc);
+		RepresentationFragment.StaticMeshDescHandle = RepresentationSubsystem->FindOrAddStaticMeshDesc(StaticMeshInstanceDesc);
 	}
 	RepresentationFragment.HighResTemplateActorIndex = HighResTemplateActor.Get() ? RepresentationSubsystem->FindOrAddTemplateActor(HighResTemplateActor.Get()) : INDEX_NONE;
 	RepresentationFragment.LowResTemplateActorIndex = LowResTemplateActor.Get() ? RepresentationSubsystem->FindOrAddTemplateActor(LowResTemplateActor.Get()) : INDEX_NONE;
diff --git a/Engine/Plugins/Runtime/MassGameplay/Source/MassRepresentation/Public/MassRepresentationFragments.h b/Engine/Plugins/Runtime/MassGameplay/Source/MassRepresentation/Public/MassRepresentationFragments.h
--- a/Engine/Plugins/Runtime/MassGameplay/Source/MassRepresentation/Public/MassRepresentationFragments.h
+++ b/Engine/Plugins/Runtime/MassGameplay/Source/MassRepresentation/Public/MassRepresentationFragments.h
@@ -63,7 +63,7 @@ struct MASSREPRESENTATION_API FMassRepresentationFragment : public FMassFragment
 	int16 LowResTemplateActorIndex = INDEX_NONE;
 
 	UPROPERTY()
-	int16 StaticMeshDescIndex = INDEX_NONE;
+	FStaticMeshInstanceVisualizationDescHandle StaticMeshDescHandle;
 
 	UPROPERTY()
 	FMassActorSpawnRequestHandle ActorSpawnRequestHandle;
diff --git a/Engine/Plugins/Runtime/MassGameplay/Source/MassRepresentation/Public/MassRepresentationSubsystem.h b/Engine/Plugins/Runtime/MassGameplay/Source/MassRepresentation/Public/MassRepresentationSubsystem.h
--- a/Engine/Plugins/Runtime/MassGameplay/Source/MassRepresentation/Public/MassRepresentationSubsystem.h
+++ b/Engine/Plugins/Runtime/MassGameplay/Source/MassRepresentation/Public/MassRepresentationSubsystem.h
@@ -35,20 +35,20 @@ public:
 	 * @param Desc is the information for the static mesh that will be instantiated later via AddStaticMeshInstance()
 	 * @return The index of the static mesh type 
 	 */
-	int16 FindOrAddStaticMeshDesc(const FStaticMeshInstanceVisualizationDesc& Desc);
+	FStaticMeshInstanceVisualizationDescHandle FindOrAddStaticMeshDesc(const FStaticMeshInstanceVisualizationDesc& Desc);
 
 	/**
 	 * Creates a dedicated visual type described by host Desc and ties ISMComponent to it.
 	 * @note this is a helper function for a common "single ISMComponent" case. Calls AddVisualDescWithISMComponents under the hood.
 	 * @return The index of the visual type
 	 */
-	int16 AddVisualDescWithISMComponent(const FStaticMeshInstanceVisualizationDesc& Desc, UInstancedStaticMeshComponent& ISMComponent);
+	FStaticMeshInstanceVisualizationDescHandle AddVisualDescWithISMComponent(const FStaticMeshInstanceVisualizationDesc& Desc, UInstancedStaticMeshComponent& ISMComponent);
 
 	/**
 	 * Creates a dedicated visual type described by host Desc and ties given ISMComponents to it.
 	 * @return The index of the visual type
 	 */
-	int16 AddVisualDescWithISMComponents(const FStaticMeshInstanceVisualizationDesc& Desc, TArrayView<TObjectPtr<UInstancedStaticMeshComponent>> ISMComponents);
+	FStaticMeshInstanceVisualizationDescHandle AddVisualDescWithISMComponents(const FStaticMeshInstanceVisualizationDesc& Desc, TArrayView<TObjectPtr<UInstancedStaticMeshComponent>> ISMComponents);
 
 	/**
 	 * Fetches FMassISMCSharedData indicated by DescriptionIndex, or nullptr if it's not a valid index
@@ -68,7 +68,7 @@ public:
 	 * Removes all data associated with a given VisualizationIndex. Note that this is safe to do only if there are no
 	 * entities relying on this index. No entity data patching will take place.
 	 */
-	void RemoveVisualDescByIndex(const int32 VisualizationIndex);
+	void RemoveVisualDesc(const FStaticMeshInstanceVisualizationDescHandle VisualizationHandle);
 
 	/** 
 	 * @return the array of all the static mesh instance component information
diff --git a/Engine/Plugins/Runtime/MassGameplay/Source/MassRepresentation/Public/MassRepresentationTypes.h b/Engine/Plugins/Runtime/MassGameplay/Source/MassRepresentation/Public/MassRepresentationTypes.h
--- a/Engine/Plugins/Runtime/MassGameplay/Source/MassRepresentation/Public/MassRepresentationTypes.h
+++ b/Engine/Plugins/Runtime/MassGameplay/Source/MassRepresentation/Public/MassRepresentationTypes.h
@@ -143,6 +143,67 @@ struct FStaticMeshInstanceVisualizationDesc : public FTableRowBase
 	}
 };
 
+/** Handle for FStaticMeshInstanceVisualizationDesc's registered with UMassRepresentationSubsystem */
+USTRUCT()
+struct alignas(2) FStaticMeshInstanceVisualizationDescHandle
+{
+	GENERATED_BODY()
+
+	static constexpr uint16 InvalidIndex = TNumericLimits<uint16>::Max();
+
+	FStaticMeshInstanceVisualizationDescHandle() = default;
+
+	explicit FStaticMeshInstanceVisualizationDescHandle(uint16 InIndex)
+	: Index(InIndex)
+	{}
+
+	explicit FStaticMeshInstanceVisualizationDescHandle(int32 InIndex) 
+	{
+		// Handle special case INDEX_NONE = InvalidIndex
+		if (InIndex == INDEX_NONE)
+		{
+			Index = InvalidIndex;
+		}
+		else
+		{
+			checkf(InIndex < static_cast<int32>(InvalidIndex), TEXT("Visualization description index InIndex %d is out of expected bounds (< %u)"), InIndex, InvalidIndex);
+			Index = static_cast<uint16>(InIndex);
+		}
+	}
+
+	FORCEINLINE int32 ToIndex() const
+	{
+		return IsValid() ? Index : INDEX_NONE;
+	}
+
+	bool IsValid() const
+	{
+		return Index != InvalidIndex;
+	}
+
+	bool operator==(const FStaticMeshInstanceVisualizationDescHandle& Other) const = default;
+
+	UE_DEPRECATED(5.4, "Referring to registered FStaticMeshInstanceVisualizationDesc's by raw int16 index has been deprecated. Please use strictly typed FStaticMeshInstanceVisualizationDescHandle instead.")
+	operator int16() const
+	{
+		if (!IsValid())
+		{
+			return INDEX_NONE;
+		}
+		return ensure(Index < TNumericLimits<int16>::Max()) ? static_cast<int16>(Index) : INDEX_NONE; 
+	}
+
+private:
+
+	UPROPERTY()
+	uint16 Index = InvalidIndex;
+
+	// @todo: Add a version / serial number to protect against recycled handle reuse. Leaving this out for now to keep size down due to 
+	// prevalent use in FMassRepresentationFragment. Perhaps serial number could be formed from the referenced 
+	// FStaticMeshInstanceVisualizationDesc's hash.
+};
+static_assert(sizeof(FStaticMeshInstanceVisualizationDescHandle) == sizeof(uint16), TEXT("FStaticMeshInstanceVisualizationDescHandle must be uint16 sized to ensure FMassRepresentationFragment memory isn't unexpectedly bloated"));
+
 class UInstancedStaticMeshComponent;
 
 
diff --git a/Engine/Plugins/Runtime/MassGameplay/Source/MassRepresentation/Public/MassVisualizationComponent.h b/Engine/Plugins/Runtime/MassGameplay/Source/MassRepresentation/Public/MassVisualizationComponent.h
--- a/Engine/Plugins/Runtime/MassGameplay/Source/MassRepresentation/Public/MassVisualizationComponent.h
+++ b/Engine/Plugins/Runtime/MassGameplay/Source/MassRepresentation/Public/MassVisualizationComponent.h
@@ -27,20 +27,20 @@ public:
 	 * @param Desc is the information for the visual that will be instantiated later via AddVisualInstance()
 	 * @return The index of the visual type 
 	 */
-	int16 FindOrAddVisualDesc(const FStaticMeshInstanceVisualizationDesc& Desc);
+	FStaticMeshInstanceVisualizationDescHandle FindOrAddVisualDesc(const FStaticMeshInstanceVisualizationDesc& Desc);
 
 	/** 
 	 * Creates a dedicated visual type described by host Desc and ties ISMComponent to it.
 	 * @note this is a helper function for a common "single ISMComponent" case. Calls AddVisualDescWithISMComponents under the hood.
 	 * @return The index of the visual type 
 	 */
-	int16 AddVisualDescWithISMComponent(const FStaticMeshInstanceVisualizationDesc& Desc, UInstancedStaticMeshComponent& ISMComponent);
+	FStaticMeshInstanceVisualizationDescHandle AddVisualDescWithISMComponent(const FStaticMeshInstanceVisualizationDesc& Desc, UInstancedStaticMeshComponent& ISMComponent);
 
 	/**
 	 * Creates a dedicated visual type described by host Desc and ties given ISMComponents to it.
 	 * @return The index of the visual type
 	 */
-	int16 AddVisualDescWithISMComponents(const FStaticMeshInstanceVisualizationDesc& Desc, TArrayView<TObjectPtr<UInstancedStaticMeshComponent>> ISMComponents);
+	FStaticMeshInstanceVisualizationDescHandle AddVisualDescWithISMComponents(const FStaticMeshInstanceVisualizationDesc& Desc, TArrayView<TObjectPtr<UInstancedStaticMeshComponent>> ISMComponents);
 
 	/**
 	 * Fetches FMassISMCSharedData indicated by DescriptionIndex, or nullptr if it's not a valid index
@@ -53,15 +53,22 @@ public:
 	 * associated with the index. Note that this is safe to do only when there are no entities relying on this data. 
 	 * No entity data patching will take place. 
 	 */
-	UE_DEPRECATED(5.4, "RemoveISMComponent has been deprecated in favor of RemoveVisualDescByIndex. Please use that instead.")
+	UE_DEPRECATED(5.4, "RemoveISMComponent has been deprecated in favor of RemoveVisualDesc. Please use that instead.")
 	void RemoveISMComponent(UInstancedStaticMeshComponent& ISMComponent);
 
 	/**
 	 * Removes all data associated with a given VisualizationIndex. Note that this is safe to do only if there are no
 	 * entities relying on this index. No entity data patching will take place.
 	 */
+	UE_DEPRECATED(5.4, "RemoveVisualDescByIndex has been deprecated in favor of RemoveVisualDesc. Please use that instead.")
 	void RemoveVisualDescByIndex(const int32 VisualizationIndex);
 
+	/**
+	 * Removes all data associated with a given VisualizationIndex. Note that this is safe to do only if there are no
+	 * entities relying on this index. No entity data patching will take place.
+	 */
+	void RemoveVisualDesc(const FStaticMeshInstanceVisualizationDescHandle VisualizationHandle);
+
 	/** Get the array of all visual instance informations */
 	FMassInstancedStaticMeshInfoArrayView GetMutableVisualInfos()
 	{
@@ -107,7 +114,7 @@ protected:
 	void BuildLODSignificanceForInfo(FMassInstancedStaticMeshInfo& Info, TConstArrayView<uint32> ForcedStaticMeshRefKeys = TConstArrayView<uint32>());
 
 	/** Either adds an element to InstancedStaticMeshInfos or reuses an existing entry based on InstancedStaticMeshInfosFreeIndices*/
-	int32 AddInstancedStaticMeshInfo(const FStaticMeshInstanceVisualizationDesc& Desc);
+	FStaticMeshInstanceVisualizationDescHandle AddInstancedStaticMeshInfo(const FStaticMeshInstanceVisualizationDesc& Desc);
 
 	/** The information of all the instanced static meshes. Make sure to use AddInstancedStaticMeshInfo to add elements to it */
 	UPROPERTY(Transient)
@@ -115,15 +122,15 @@ protected:
 	UE_MT_DECLARE_RW_RECURSIVE_ACCESS_DETECTOR(InstancedStaticMeshInfosDetector);
 
 	/** Indices to InstancedStaticMeshInfos that have been released and can be reused */
-	TArray<int32> InstancedStaticMeshInfosFreeIndices;
+	TArray<FStaticMeshInstanceVisualizationDescHandle> InstancedStaticMeshInfosFreeIndices;
 
-	/** Mapping from ISMComponent object path to corresponding VisualIndex */
-	TMap<uint32, int32> ISMComponentMap;
+	/** Mapping from ISMComponent object path hash to corresponding VisualDescHandle */
+	TMap<uint32, FStaticMeshInstanceVisualizationDescHandle> ISMComponentMap;
 
 	FMassISMCSharedDataMap ISMCSharedData;
 
 	/** Indicies to InstancedStaticMeshInfos that need their SMComponent constructed */
-	TArray<int32> InstancedSMComponentsRequiringConstructing;
+	TArray<FStaticMeshInstanceVisualizationDescHandle> InstancedSMComponentsRequiringConstructing;
 
 	UE_DEPRECATED(5.4, "This flavor of BuildLODSignificanceForInfo is no longer supported. Use the other one instead.")
 	void BuildLODSignificanceForInfo(FMassInstancedStaticMeshInfo& Info, const uint32 ForcedStaticMeshRefKey);
