Massaged the overzealous check in UMassVisualizationComponent::FindOrAddVisualDesc

diff --git a/Engine/Plugins/Runtime/MassGameplay/Source/MassRepresentation/Private/MassVisualizationComponent.cpp b/Engine/Plugins/Runtime/MassGameplay/Source/MassRepresentation/Private/MassVisualizationComponent.cpp
--- a/Engine/Plugins/Runtime/MassGameplay/Source/MassRepresentation/Private/MassVisualizationComponent.cpp
+++ b/Engine/Plugins/Runtime/MassGameplay/Source/MassRepresentation/Private/MassVisualizationComponent.cpp
@@ -88,10 +88,15 @@ FStaticMeshInstanceVisualizationDescHandle UMassVisualizationComponent::FindOrAd
 			BuildLODSignificanceForInfo(InstancedStaticMeshInfos[VisualDescHandle.ToIndex()]);
 
 			InstancedSMComponentsRequiringConstructing.Add(VisualDescHandle);
+
+			check(VisualDescHandle.IsValid());
+		}
+		else
+		{
+			UE_LOG(LogMassRepresentation, Warning, TEXT("%hs: invalid FStaticMeshInstanceVisualizationDesc passed in. Check the contained meshes."), __FUNCTION__);
 		}
 	}
 
-	check(VisualDescHandle.IsValid());
 	return VisualDescHandle;
 }
 
