[Mass] Prevent message log notifications from ValidateBuildContext where there is no Editor.

diff --git a/Engine/Plugins/Runtime/MassGameplay/Source/MassSpawner/Private/MassEntityTemplateRegistry.cpp b/Engine/Plugins/Runtime/MassGameplay/Source/MassSpawner/Private/MassEntityTemplateRegistry.cpp
--- a/Engine/Plugins/Runtime/MassGameplay/Source/MassSpawner/Private/MassEntityTemplateRegistry.cpp
+++ b/Engine/Plugins/Runtime/MassGameplay/Source/MassSpawner/Private/MassEntityTemplateRegistry.cpp
@@ -10,8 +10,10 @@
 #include "MassEntityTypes.h"
 #include "MassEntityTraitBase.h"
 #include "Logging/MessageLog.h"
-#include "MassEntityUtils.h"
 
+#if WITH_EDITOR
+#include "Editor.h"
+#endif
 
 #define LOCTEXT_NAMESPACE "Mass"
 
@@ -161,8 +163,8 @@ bool FMassEntityTemplateBuildContext::ValidateBuildContext(const UWorld& World)
 		}
 	}
 
-#if WITH_UNREAL_DEVELOPER_TOOLS
-	if (bFragmentHasMultipleOwners || bMissingFragmentDependencies)
+#if WITH_UNREAL_DEVELOPER_TOOLS && WITH_EDITOR
+	if (GEditor && (bFragmentHasMultipleOwners || bMissingFragmentDependencies))
 	{
 		FMessageLog EditorErrors("MassEntity");
 		if (bFragmentHasMultipleOwners)
