Added another console command for Mass LOD debugging - mass.debug.LODSubsystem.GatherPlayers - this one making the whole system not use playes as LOD viewers.

diff --git a/Engine/Plugins/Runtime/MassGameplay/Source/MassLOD/Private/MassLODSubsystem.cpp b/Engine/Plugins/Runtime/MassGameplay/Source/MassLOD/Private/MassLODSubsystem.cpp
--- a/Engine/Plugins/Runtime/MassGameplay/Source/MassLOD/Private/MassLODSubsystem.cpp
+++ b/Engine/Plugins/Runtime/MassGameplay/Source/MassLOD/Private/MassLODSubsystem.cpp
@@ -40,39 +40,64 @@ namespace UE::MassLOD
 #if WITH_MASSGAMEPLAY_DEBUG
 	namespace Debug
 	{
+		/** Returns whether getting the UMassLODSubsystem and the bool parameter was successful */
+		bool GetSubsystemAndBoolArgument(const TArray<FString>& Args, UWorld* World, UMassLODSubsystem*& OutMassLODSubsystem, bool& bOutBool)
+		{
+			if (!World)
+			{
+				UE_LOG(LogConsoleResponse, Display, TEXT("Error: invalid world"));
+				return false;
+			}
+
+			OutMassLODSubsystem = World->GetSubsystem<UMassLODSubsystem>();
+			if (OutMassLODSubsystem == nullptr)
+			{
+				UE_LOG(LogConsoleResponse, Display, TEXT("Error: Unable to fetch MassLODSubsystem instance"));
+				return false;
+			}
+
+			if (Args.Num() < 1)
+			{
+				UE_LOG(LogConsoleResponse, Display, TEXT("Error: Expecting 1 parameter"));
+				return false;
+			}
+
+			if (!LexTryParseString<bool>(bOutBool, *Args[0]))
+			{
+				UE_LOG(LogConsoleResponse, Display, TEXT("Error: parameter must be an integer or a boolean"));
+				return false;
+			}
+
+			return true;
+		}
+
 		FAutoConsoleCommandWithWorldArgsAndOutputDevice ToggleUsePlayerLocationCmd(
 			TEXT("mass.debug.LODSubsystem.UsePlayerLocation"),
-			TEXT("Sets UMassSubsystem::bUsePlayerPawnLocationInsteadOfCamera. Note that this is a command that doesn't retain state and usually needs running both for the client and the server"),
+			TEXT("Sets UMassLODSubsystem::bUsePlayerPawnLocationInsteadOfCamera. Note that this is a command that doesn't retain state and usually needs running both for the client and the server"),
 			FConsoleCommandWithWorldArgsAndOutputDeviceDelegate::CreateLambda([](const TArray<FString>& Args, UWorld* World, FOutputDevice& Ar)
 				{
-					if (!World)
-					{
-						UE_LOG(LogConsoleResponse, Display, TEXT("Error: invalid world"));
-						return;
-					}
+					UMassLODSubsystem* MassLODSubsystem = nullptr;
+					bool bNewValue = false;
 
-					UMassLODSubsystem* MassLODSubsystem = World->GetSubsystem<UMassLODSubsystem>();
-					if (MassLODSubsystem == nullptr)
-					{
-						UE_LOG(LogConsoleResponse, Display, TEXT("Error: Unable to fetch MassLODSubsystem instance"));
-						return;
-					}
-					
-					if (Args.Num() < 1)
+					if (GetSubsystemAndBoolArgument(Args, World, MassLODSubsystem, bNewValue))
 					{
-						UE_LOG(LogConsoleResponse, Display, TEXT("Error: Expecting 1 parameter"));
-						return;
+						MassLODSubsystem->DebugSetUsePlayerPawnLocationInsteadOfCamera(bNewValue);
 					}
+				}));
 
+		FAutoConsoleCommandWithWorldArgsAndOutputDevice ToggleGatherPlayers(
+			TEXT("mass.debug.LODSubsystem.GatherPlayers"),
+			TEXT("Sets UMassLODSubsystem::bGatherPlayerControllers. Note that this is a command that doesn't retain state and usually needs running both for the client and the server"),
+			FConsoleCommandWithWorldArgsAndOutputDeviceDelegate::CreateLambda([](const TArray<FString>& Args, UWorld* World, FOutputDevice& Ar)
+				{
+					UMassLODSubsystem* MassLODSubsystem = nullptr;
 					bool bNewValue = false;
-					if (!LexTryParseString<bool>(bNewValue, *Args[0]))
+
+					if (GetSubsystemAndBoolArgument(Args, World, MassLODSubsystem, bNewValue))
 					{
-						UE_LOG(LogConsoleResponse, Display, TEXT("Error: parameter must be an integer or a boolean"));
-						return;
+						MassLODSubsystem->DebugSetGatherPlayers(bNewValue);
 					}
-
-					MassLODSubsystem->DebugSetUsePlayerPawnLocationInsteadOfCamera(bNewValue);
-				}));
+				})); 
 	}
 #endif // WITH_MASSGAMEPLAY_DEBUG
 }
diff --git a/Engine/Plugins/Runtime/MassGameplay/Source/MassLOD/Public/MassLODSubsystem.h b/Engine/Plugins/Runtime/MassGameplay/Source/MassLOD/Public/MassLODSubsystem.h
--- a/Engine/Plugins/Runtime/MassGameplay/Source/MassLOD/Public/MassLODSubsystem.h
+++ b/Engine/Plugins/Runtime/MassGameplay/Source/MassLOD/Public/MassLODSubsystem.h
@@ -114,6 +114,7 @@ public:
 	bool IsUsingPlayerPawnLocationInsteadOfCamera() const { return bUsePlayerPawnLocationInsteadOfCamera; }
 
 #if WITH_MASSGAMEPLAY_DEBUG
+	void DebugSetGatherPlayers(const bool bInValue) { bGatherPlayerControllers = bInValue; }
 	void DebugSetUsePlayerPawnLocationInsteadOfCamera(const bool bInValue) { bUsePlayerPawnLocationInsteadOfCamera = bInValue; }
 	void DebugUnregisterActorViewer();
 #endif
