MassCrowd: Fix desired speed sometimes getting set to zero.
- FMassMovementParameters::GenerateDesiredSpeed() was calling the int version of UE::RandomSequence::RandRange(), resulting 0 or1 speed
- Added float override, in addtion to the FRandRange(), GenerateDesiredSpeed() should pick that one up now

diff --git a/Engine/Plugins/Runtime/MassGameplay/Source/MassCommon/Public/RandomSequence.h b/Engine/Plugins/Runtime/MassGameplay/Source/MassCommon/Public/RandomSequence.h
--- a/Engine/Plugins/Runtime/MassGameplay/Source/MassCommon/Public/RandomSequence.h
+++ b/Engine/Plugins/Runtime/MassGameplay/Source/MassCommon/Public/RandomSequence.h
@@ -28,7 +28,7 @@ namespace UE::RandomSequence
 	/**
 	 * @return 32 bit fibonacci hash at specified index.
 	 */
-	FORCEINLINE uint32 FibocciHash(const int32 SeqIndex)
+	FORCEINLINE uint32 FibonacciHash(const int32 SeqIndex)
 	{
 		constexpr uint32 K = 2654435769u; // 2^32 / phi (golden ratio)
 		// Offset the sequence by 1, so that index 0 is not always 0.
@@ -42,7 +42,7 @@ namespace UE::RandomSequence
 	FORCEINLINE float FRand(const int32 SeqIndex)
 	{
 		float Result;
-		*(uint32*)&Result = 0x3F800000U | (FibocciHash(SeqIndex) >> 9);
+		*(uint32*)&Result = 0x3F800000U | (FibonacciHash(SeqIndex) >> 9);
 		return Result - 1.0f; 
 	}
 
@@ -52,7 +52,7 @@ namespace UE::RandomSequence
 	 */
 	FORCEINLINE int32 RandHelper(const int32 SeqIndex, const int32 A)
 	{
-		return (int32)(((int64)FibocciHash(SeqIndex) * (int64)A) >> 32);
+		return (int32)(((int64)FibonacciHash(SeqIndex) * (int64)A) >> 32);
 	}
 
 	/** 
@@ -66,7 +66,7 @@ namespace UE::RandomSequence
 	}
 	
 	/** 
-	 * Helper function to return random int in specified range.
+	 * Helper function to return random float in specified range.
 	 * @return A random number >= Min and <= Max
 	 */
 	FORCEINLINE float FRandRange(const int32 SeqIndex, const float InMin, const float InMax)
@@ -74,4 +74,18 @@ namespace UE::RandomSequence
 		return InMin + (InMax - InMin) * FRand(SeqIndex);
 	}
 
+	/** 
+	 * Helper function to return random float in specified range.
+	 * @return A random number >= Min and <= Max
+	 */
+	FORCEINLINE float RandRange(const int32 SeqIndex, const float InMin, const float InMax)
+	{
+		return FRandRange(SeqIndex, InMin, InMax);
+	}
+
+	UE_DEPRECATED(5.4, "Please use the correctly spelled version FibonacciHash()")
+	FORCEINLINE uint32 FibocciHash(const int32 SeqIndex)
+	{
+		return FibonacciHash(SeqIndex);
+	}
 }; // UE::RandomSequence
\ No newline at end of file
