Replaced calls to FMassISMCSharedDataMap::GetAndMarkDirtyChecked with FMassISMCSharedDataMap::GetAndMarkDirty and modified the latter to catch the erroneous hash. Also updated UMassVisualizationComponent::RemoveVisualDesc in a similar way due to it being suspected of suffering from the same issue.

Note that the checks added in code are to force QA to report the occurrences. The change is deadended in this stream for that reason.

#deadend

diff --git a/Engine/Plugins/Runtime/MassGameplay/Source/MassRepresentation/Private/MassVisualizationComponent.cpp b/Engine/Plugins/Runtime/MassGameplay/Source/MassRepresentation/Private/MassVisualizationComponent.cpp
--- a/Engine/Plugins/Runtime/MassGameplay/Source/MassRepresentation/Private/MassVisualizationComponent.cpp
+++ b/Engine/Plugins/Runtime/MassGameplay/Source/MassRepresentation/Private/MassVisualizationComponent.cpp
@@ -131,7 +131,19 @@ FStaticMeshInstanceVisualizationDescHandle UMassVisualizationComponent::AddVisua
 			check(VisualHandle.IsValid());
 		}
 
-		const uint32 ISMComponentPathHash = GetTypeHash(ISMComponents[EntryIndex]->GetPathName());
+		const FString ISMCPath = ISMComponents[EntryIndex]->GetPathName();
+		const uint32 ISMComponentPathHash = GetTypeHash(ISMCPath); 
+#if WITH_MASSGAMEPLAY_DEBUG
+		TArray<FString>& DebugPaths = DebugHashToPathMap.FindOrAdd(ISMComponentPathHash, TArray<FString>());
+		if (!ensureMsgf(DebugPaths.Add(ISMCPath) == 0, TEXT("Multiple ISMC paths resulting in the same hash")))
+		{
+			UE_VLOG_UELOG(this, LogMassRepresentation, Error, TEXT("%hs multiple ISMComponents resulting in identical hash %u"), __FUNCTION__, ISMComponentPathHash);
+			for (const FString& Path : DebugPaths)
+			{
+				UE_VLOG_UELOG(this, LogMassRepresentation, Error, TEXT("\t%s"), *Path);
+			}
+		}
+#endif // WITH_MASSGAMEPLAY_DEBUG
 		FMassISMCSharedData& NewData = ISMCSharedData.FindOrAdd(ISMComponentPathHash, FMassISMCSharedData(ISMComponents[EntryIndex], /*bInRequiresExternalInstanceIDTracking=*/true));
 		InstancedStaticMeshInfos[VisualHandle.ToIndex()].AddISMComponent(NewData);
 		ISMComponentPathHashes.Add(ISMComponentPathHash);
@@ -161,11 +173,24 @@ void UMassVisualizationComponent::RemoveVisualDesc(const FStaticMeshInstanceVisu
 	{
 		for (TObjectPtr<UInstancedStaticMeshComponent>& ISMComponent : InstancedStaticMeshInfos[VisualizationHandle.ToIndex()].InstancedStaticMeshComponents)
 		{
-			const uint32 ISMComponentPathHash = GetTypeHash(ISMComponent.GetPathName());
-			const FStaticMeshInstanceVisualizationDescHandle StoredVisualizationDescHandle = ISMComponentMap.FindAndRemoveChecked(ISMComponentPathHash);
-			ensure(StoredVisualizationDescHandle == VisualizationHandle);
+			// @todo using ISMComponent.GetPathName() here might be wrong for cases where we use GetTypeHash(MeshDesc)
+			// to create the ISMComponentMap key
+			const FString ISMCPath = ISMComponent.GetPathName();
+			const uint32 ISMComponentPathHash = GetTypeHash(ISMCPath);
+
+			const bool bValidKey = ISMComponentMap.Contains(ISMComponentPathHash);
+			checkf(bValidKey, TEXT("Failed to find %u as a key in ISMComponentMap, ISMC path: %s"), ISMComponentPathHash, *ISMComponent.GetPathName());
+			if (bValidKey)
+			{
+				const FStaticMeshInstanceVisualizationDescHandle StoredVisualizationDescHandle = ISMComponentMap.FindAndRemoveChecked(ISMComponentPathHash);
+				ensure(StoredVisualizationDescHandle == VisualizationHandle);
+			}
 		
 			ISMCSharedData.Remove(ISMComponentPathHash);
+#if WITH_MASSGAMEPLAY_DEBUG
+			TArray<FString>& Paths = DebugHashToPathMap.FindChecked(ISMComponentPathHash);
+			ensure(Paths.Remove(ISMCPath));
+#endif // WITH_MASSGAMEPLAY_DEBUG
 		}
 		
 		InstancedStaticMeshInfos[VisualizationHandle.ToIndex()].Reset();
@@ -650,11 +675,12 @@ void FMassLODSignificanceRange::AddBatchedTransform(const int32 InstanceId, cons
 			continue;
 		}
 
-		FMassISMCSharedData& SharedData = ISMCSharedDataPtr->GetAndMarkDirtyChecked(StaticMeshRefs[StaticMeshIndex]);
-
-		SharedData.UpdateInstanceIds.Add(InstanceId);
-		SharedData.StaticMeshInstanceTransforms.Add(Transform);
-		SharedData.StaticMeshInstancePrevTransforms.Add(PrevTransform);
+		if (FMassISMCSharedData* SharedData = ISMCSharedDataPtr->GetAndMarkDirty(StaticMeshRefs[StaticMeshIndex]))
+		{
+			SharedData->UpdateInstanceIds.Add(InstanceId);
+			SharedData->StaticMeshInstanceTransforms.Add(Transform);
+			SharedData->StaticMeshInstancePrevTransforms.Add(PrevTransform);
+		}
 	}
 }
 
@@ -668,8 +694,10 @@ void FMassLODSignificanceRange::AddBatchedCustomDataFloats(const TArray<float>&
 			continue;
 		}
 
-		FMassISMCSharedData& SharedData = ISMCSharedDataPtr->GetAndMarkDirtyChecked(StaticMeshRefs[StaticMeshIndex]);
-		SharedData.StaticMeshInstanceCustomFloats.Append(CustomFloats);
+		if (FMassISMCSharedData* SharedData = ISMCSharedDataPtr->GetAndMarkDirty(StaticMeshRefs[StaticMeshIndex]))
+		{
+			SharedData->StaticMeshInstanceCustomFloats.Append(CustomFloats);
+		}
 	}
 }
 
@@ -678,10 +706,12 @@ void FMassLODSignificanceRange::AddInstance(const int32 InstanceId, const FTrans
 	check(ISMCSharedDataPtr);
 	for (int32 StaticMeshIndex = 0; StaticMeshIndex < StaticMeshRefs.Num(); ++StaticMeshIndex)
 	{
-		FMassISMCSharedData& SharedData = ISMCSharedDataPtr->GetAndMarkDirtyChecked(StaticMeshRefs[StaticMeshIndex]);
-		SharedData.UpdateInstanceIds.Add(InstanceId);
-		SharedData.StaticMeshInstanceTransforms.Add(Transform);
-		SharedData.StaticMeshInstancePrevTransforms.Add(Transform);
+		if (FMassISMCSharedData* SharedData = ISMCSharedDataPtr->GetAndMarkDirty(StaticMeshRefs[StaticMeshIndex]))
+		{
+			SharedData->UpdateInstanceIds.Add(InstanceId);
+			SharedData->StaticMeshInstanceTransforms.Add(Transform);
+			SharedData->StaticMeshInstancePrevTransforms.Add(Transform);
+		}
 	}
 }
 
@@ -707,17 +737,18 @@ void FMassLODSignificanceRange::WriteCustomDataFloatsAtStartIndex(int32 StaticMe
 			return;
 		}
 
-		FMassISMCSharedData& SharedData = ISMCSharedDataPtr->GetAndMarkDirtyChecked(StaticMeshRefs[StaticMeshIndex]);
-
-		int32 StartIndex = FloatsPerInstance * SharedData.WriteIterator + StartFloatIndex;
+		if (FMassISMCSharedData* SharedData = ISMCSharedDataPtr->GetAndMarkDirty(StaticMeshRefs[StaticMeshIndex]))
+		{
+			const int32 StartIndex = FloatsPerInstance * SharedData->WriteIterator + StartFloatIndex;
 
-		ensure(SharedData.StaticMeshInstanceCustomFloats.Num() >= StartIndex + CustomFloats.Num());
+			ensure(SharedData->StaticMeshInstanceCustomFloats.Num() >= StartIndex + CustomFloats.Num());
 
-		for (int CustomFloatIdx = 0; CustomFloatIdx < CustomFloats.Num(); CustomFloatIdx++)
-		{
-			SharedData.StaticMeshInstanceCustomFloats[StartIndex + CustomFloatIdx] = CustomFloats[CustomFloatIdx];
+			for (int CustomFloatIdx = 0; CustomFloatIdx < CustomFloats.Num(); CustomFloatIdx++)
+			{
+				SharedData->StaticMeshInstanceCustomFloats[StartIndex + CustomFloatIdx] = CustomFloats[CustomFloatIdx];
+			}
+			SharedData->WriteIterator++;
 		}
-		SharedData.WriteIterator++;
 	}
 }
 
diff --git a/Engine/Plugins/Runtime/MassGameplay/Source/MassRepresentation/Public/MassRepresentationTypes.h b/Engine/Plugins/Runtime/MassGameplay/Source/MassRepresentation/Public/MassRepresentationTypes.h
--- a/Engine/Plugins/Runtime/MassGameplay/Source/MassRepresentation/Public/MassRepresentationTypes.h
+++ b/Engine/Plugins/Runtime/MassGameplay/Source/MassRepresentation/Public/MassRepresentationTypes.h
@@ -338,7 +338,7 @@ public:
  * @see UMassVisualizationComponent::EndVisualChanges for iteration
  * @see FMassLODSignificanceRange methods for performing dirtying operations
  */
-struct FMassISMCSharedDataMap 
+struct FMassISMCSharedDataMap
 {
 	struct FDirtyIterator
 	{
@@ -392,7 +392,8 @@ struct FMassISMCSharedDataMap
 
 	FMassISMCSharedData* GetAndMarkDirty(const uint32 Hash)
 	{
-		if (const int32* DataIndex = Map.Find(Hash))
+		const int32* DataIndex = Map.Find(Hash);
+		if (ensureMsgf(DataIndex, TEXT("%hs Failed to find Hash %u"), __FUNCTION__, Hash))
 		{
 			DirtyData[*DataIndex] = true;
 			return &Data[*DataIndex];
diff --git a/Engine/Plugins/Runtime/MassGameplay/Source/MassRepresentation/Public/MassVisualizationComponent.h b/Engine/Plugins/Runtime/MassGameplay/Source/MassRepresentation/Public/MassVisualizationComponent.h
--- a/Engine/Plugins/Runtime/MassGameplay/Source/MassRepresentation/Public/MassVisualizationComponent.h
+++ b/Engine/Plugins/Runtime/MassGameplay/Source/MassRepresentation/Public/MassVisualizationComponent.h
@@ -3,6 +3,7 @@
 #pragma once
 
 #include "MassRepresentationTypes.h"
+#include "MassCommonTypes.h"
 #include "Components/ActorComponent.h"
 #include "MassRepresentationTypes.h"
 #include "Misc/MTAccessDetector.h"
@@ -134,4 +135,8 @@ protected:
 
 	UE_DEPRECATED(5.4, "This flavor of BuildLODSignificanceForInfo is no longer supported. Use the other one instead.")
 	void BuildLODSignificanceForInfo(FMassInstancedStaticMeshInfo& Info, const uint32 ForcedStaticMeshRefKey);
+
+#if WITH_MASSGAMEPLAY_DEBUG
+	TMap<uint32, TArray<FString>> DebugHashToPathMap;
+#endif // WITH_MASSGAMEPLAY_DEBUG
 };
