Removed header included twice in MassSmartObjectProcessor

diff --git a/Engine/Plugins/Runtime/MassGameplay/Source/MassSmartObjects/Private/MassSmartObjectProcessor.cpp b/Engine/Plugins/Runtime/MassGameplay/Source/MassSmartObjects/Private/MassSmartObjectProcessor.cpp
--- a/Engine/Plugins/Runtime/MassGameplay/Source/MassSmartObjects/Private/MassSmartObjectProcessor.cpp
+++ b/Engine/Plugins/Runtime/MassGameplay/Source/MassSmartObjects/Private/MassSmartObjectProcessor.cpp
@@ -3,23 +3,22 @@
 #include "MassSmartObjectProcessor.h"
 #include "MassCommandBuffer.h"
 #include "MassCommonTypes.h"
+#include "MassDebugger.h"
 #include "MassExecutionContext.h"
+#include "MassGameplayExternalTraits.h"
 #include "MassSignalSubsystem.h"
 #include "MassSmartObjectBehaviorDefinition.h"
 #include "MassSmartObjectFragments.h"
 #include "MassSmartObjectRequest.h"
 #include "MassSmartObjectSettings.h"
 #include "MassSmartObjectTypes.h"
-#include "SmartObjectZoneAnnotations.h"
 #include "Misc/ScopeExit.h"
 #include "SmartObjectOctree.h"
 #include "SmartObjectSubsystem.h"
+#include "SmartObjectZoneAnnotations.h"
 #include "VisualLogger/VisualLogger.h"
 #include "ZoneGraphAnnotationSubsystem.h"
-#include "MassGameplayExternalTraits.h"
 #include "ZoneGraphSubsystem.h"
-#include "MassGameplayExternalTraits.h"
-#include "MassDebugger.h"
 
 
 //----------------------------------------------------------------------//
