Changed the Mass ISM-related API's usage of cryptic int32 to identify entities (usually obtained by hashing the entity handle) in favor of using the FMassEntityHandle directly.

Added benefit is that we can directly identify the entity being processed by [REDACTED_PERSON], without additional lookups.

diff --git a/Engine/Plugins/Runtime/MassGameplay/Source/MassRepresentation/Private/MassStationaryISMRepresentationFragmentDestructor.cpp b/Engine/Plugins/Runtime/MassGameplay/Source/MassRepresentation/Private/MassStationaryISMRepresentationFragmentDestructor.cpp
--- a/Engine/Plugins/Runtime/MassGameplay/Source/MassRepresentation/Private/MassStationaryISMRepresentationFragmentDestructor.cpp
+++ b/Engine/Plugins/Runtime/MassGameplay/Source/MassRepresentation/Private/MassStationaryISMRepresentationFragmentDestructor.cpp
@@ -45,10 +45,9 @@ void UMassStationaryISMRepresentationFragmentDestructor::Execute(FMassEntityMana
 				if (FMassLODSignificanceRange* OldRange = ISMInfo.GetLODSignificanceRange(Representation.PrevLODSignificance))
 				{
 					const FMassEntityHandle EntityHandle = Context.GetEntity(EntityIdx);
-					const int32 EntityId = GetTypeHash(EntityHandle);
 					if (OldRange)
 					{
-						OldRange->RemoveInstance(EntityId);
+						OldRange->RemoveInstance(EntityHandle);
 					}
 				}
 				Representation.CurrentRepresentation = EMassRepresentationType::None;
diff --git a/Engine/Plugins/Runtime/MassGameplay/Source/MassRepresentation/Private/MassStationaryISMSwitcherProcessor.cpp b/Engine/Plugins/Runtime/MassGameplay/Source/MassRepresentation/Private/MassStationaryISMSwitcherProcessor.cpp
--- a/Engine/Plugins/Runtime/MassGameplay/Source/MassRepresentation/Private/MassStationaryISMSwitcherProcessor.cpp
+++ b/Engine/Plugins/Runtime/MassGameplay/Source/MassRepresentation/Private/MassStationaryISMSwitcherProcessor.cpp
@@ -73,11 +73,9 @@ void UMassStationaryISMSwitcherProcessor::ProcessContext(FMassExecutionContext&
 		if (Representation.PrevRepresentation == EMassRepresentationType::StaticMeshInstance
 			&& Representation.CurrentRepresentation != EMassRepresentationType::StaticMeshInstance)
 		{
-			const int32 EntityId = GetTypeHash(EntityHandle);
-
 			// note that we're using the PrevLODSignificance here, and the reason for it is that the Prev value matches the 
 			// PrevRepresentation - thus we need to remove from the "previously" used LODSignificance range.
-			ISMInfo.RemoveInstance(EntityId, Representation.PrevLODSignificance);
+			ISMInfo.RemoveInstance(EntityHandle, Representation.PrevLODSignificance);
 
 			// consume "prev" data
 			Representation.PrevRepresentation = Representation.CurrentRepresentation;
@@ -90,7 +88,6 @@ void UMassStationaryISMSwitcherProcessor::ProcessContext(FMassExecutionContext&
 		else if (Representation.PrevRepresentation != EMassRepresentationType::StaticMeshInstance
 			&& Representation.CurrentRepresentation == EMassRepresentationType::StaticMeshInstance)
 		{
-			const int32 EntityId = GetTypeHash(EntityHandle);
 			const FTransform& Transform = TransformFragment.GetTransform();
 			const FTransform& PrevTransform = Representation.PrevTransform;
 			const float LODSignificance = RepresentationLOD.LODSignificance;
@@ -102,11 +99,11 @@ void UMassStationaryISMSwitcherProcessor::ProcessContext(FMassExecutionContext&
 				{
 					const FTransform& TransformOffset = ISMInfo.GetTransformOffset();
 					const FTransform SMTransform = TransformOffset * Transform;
-					NewRange->AddInstance(EntityId, SMTransform);
+					NewRange->AddInstance(EntityHandle, SMTransform);
 				}
 				else
 				{
-					NewRange->AddInstance(EntityId, Transform);
+					NewRange->AddInstance(EntityHandle, Transform);
 				}
 			}
 
@@ -128,15 +125,14 @@ void UMassStationaryISMSwitcherProcessor::ProcessContext(FMassExecutionContext&
 			FMassLODSignificanceRange* NewRange = ISMInfo.GetLODSignificanceRange(RepresentationLOD.LODSignificance);
 			if (OldRange != NewRange)
 			{
-				const int32 EntityId = GetTypeHash(EntityHandle);
 				if (OldRange)
 				{
-					OldRange->RemoveInstance(EntityId);
+					OldRange->RemoveInstance(EntityHandle);
 				}
 				if (NewRange)
 				{
 					const FTransform& Transform = TransformFragment.GetTransform();
-					NewRange->AddInstance(EntityId, Transform);
+					NewRange->AddInstance(EntityHandle, Transform);
 				}
 			}
 		}
diff --git a/Engine/Plugins/Runtime/MassGameplay/Source/MassRepresentation/Private/MassUpdateISMProcessor.cpp b/Engine/Plugins/Runtime/MassGameplay/Source/MassRepresentation/Private/MassUpdateISMProcessor.cpp
--- a/Engine/Plugins/Runtime/MassGameplay/Source/MassRepresentation/Private/MassUpdateISMProcessor.cpp
+++ b/Engine/Plugins/Runtime/MassGameplay/Source/MassRepresentation/Private/MassUpdateISMProcessor.cpp
@@ -28,7 +28,7 @@ void UMassUpdateISMProcessor::ConfigureQueries()
 	EntityQuery.SetChunkFilter(&FMassVisualizationChunkFragment::AreAnyEntitiesVisibleInChunk);
 	EntityQuery.AddSharedRequirement<FMassRepresentationSubsystemSharedFragment>(EMassFragmentAccess::ReadWrite);
 
-	// ignore entitites configured to have their representation static (@todo maybe just check if there's not movement fragment?)
+	// ignore entities configured to have their representation static (@todo maybe just check if there's not movement fragment?)
 	EntityQuery.AddTagRequirement<FMassStaticRepresentationTag>(EMassFragmentPresence::None);
 }
 
@@ -53,7 +53,7 @@ void UMassUpdateISMProcessor::Execute(FMassEntityManager& EntityManager, FMassEx
 
 			if (Representation.CurrentRepresentation == EMassRepresentationType::StaticMeshInstance)
 			{
-				UpdateISMTransform(GetTypeHash(Context.GetEntity(EntityIdx)), ISMInfo[Representation.StaticMeshDescHandle.ToIndex()], TransformFragment.GetTransform(), Representation.PrevTransform, RepresentationLOD.LODSignificance, Representation.PrevLODSignificance);
+				UpdateISMTransform(Context.GetEntity(EntityIdx), ISMInfo[Representation.StaticMeshDescHandle.ToIndex()], TransformFragment.GetTransform(), Representation.PrevTransform, RepresentationLOD.LODSignificance, Representation.PrevLODSignificance);
 			}
 			Representation.PrevTransform = TransformFragment.GetTransform();
 			Representation.PrevLODSignificance = RepresentationLOD.LODSignificance;
@@ -61,7 +61,7 @@ void UMassUpdateISMProcessor::Execute(FMassEntityManager& EntityManager, FMassEx
 	});
 }
 
-void UMassUpdateISMProcessor::UpdateISMTransform(int32 EntityId, FMassInstancedStaticMeshInfo& ISMInfo, const FTransform& Transform, const FTransform& PrevTransform, const float LODSignificance, const float PrevLODSignificance/* = -1.0f*/)
+void UMassUpdateISMProcessor::UpdateISMTransform(FMassEntityHandle EntityHandle, FMassInstancedStaticMeshInfo& ISMInfo, const FTransform& Transform, const FTransform& PrevTransform, const float LODSignificance, const float PrevLODSignificance/* = -1.0f*/)
 {
 	if (ISMInfo.ShouldUseTransformOffset())
 	{
@@ -69,10 +69,10 @@ void UMassUpdateISMProcessor::UpdateISMTransform(int32 EntityId, FMassInstancedS
 		const FTransform SMTransform = TransformOffset * Transform;
 		const FTransform SMPrevTransform = TransformOffset * PrevTransform;
 
-		ISMInfo.AddBatchedTransform(EntityId, SMTransform, SMPrevTransform, LODSignificance, PrevLODSignificance);
+		ISMInfo.AddBatchedTransform(EntityHandle, SMTransform, SMPrevTransform, LODSignificance, PrevLODSignificance);
 	}
 	else
 	{
-		ISMInfo.AddBatchedTransform(EntityId, Transform, PrevTransform, LODSignificance, PrevLODSignificance);
+		ISMInfo.AddBatchedTransform(EntityHandle, Transform, PrevTransform, LODSignificance, PrevLODSignificance);
 	}
 }
diff --git a/Engine/Plugins/Runtime/MassGameplay/Source/MassRepresentation/Private/MassVisualizationComponent.cpp b/Engine/Plugins/Runtime/MassGameplay/Source/MassRepresentation/Private/MassVisualizationComponent.cpp
--- a/Engine/Plugins/Runtime/MassGameplay/Source/MassRepresentation/Private/MassVisualizationComponent.cpp
+++ b/Engine/Plugins/Runtime/MassGameplay/Source/MassRepresentation/Private/MassVisualizationComponent.cpp
@@ -394,19 +394,19 @@ void UMassVisualizationComponent::HandleChangesWithExternalIDTracking(UInstanced
 	ProcessRemoves(ISMComponent, SharedData, /*bUpdateNavigation=*/false);
 
 	// NOTE: This code path is designed to only perform Adds, never updates so updates are filtered out along with duplicates.
-	TArray<int32>& MassInstanceIds = SharedData.UpdateInstanceIds;
-	if (MassInstanceIds.Num())
+	TArray<FMassEntityHandle>& EntityHandles = SharedData.UpdateInstanceIds;
+	if (EntityHandles.Num())
 	{
-		INC_DWORD_STAT_BY(STAT_Mass_VisualizationComponent_InstancesAddedNum, MassInstanceIds.Num());
+		INC_DWORD_STAT_BY(STAT_Mass_VisualizationComponent_InstancesAddedNum, EntityHandles.Num());
 
-		FMassISMCSharedData::FIdMap& SharedIdMap = SharedData.GetMutableIdMap();
+		FMassISMCSharedData::FEntityToPrimitiveIdMap& SharedIdMap = SharedData.GetMutableEntityPrimitiveToIdMap();
 		TArray<Experimental::FHashElementId> ElementIds;
-		ElementIds.SetNumUninitialized(MassInstanceIds.Num());
+		ElementIds.SetNumUninitialized(EntityHandles.Num());
 		// Filter out all updates & duplicate adds
-		for (int32 IDIndex = MassInstanceIds.Num() - 1; IDIndex >= 0; --IDIndex)
+		for (int32 IDIndex = EntityHandles.Num() - 1; IDIndex >= 0; --IDIndex)
 		{
 			bool bWasAlreadyInMap = false;
-			Experimental::FHashElementId ElementId = SharedIdMap.FindOrAddId(MassInstanceIds[IDIndex], FPrimitiveInstanceId{INDEX_NONE}, bWasAlreadyInMap);
+			Experimental::FHashElementId ElementId = SharedIdMap.FindOrAddId(EntityHandles[IDIndex], FPrimitiveInstanceId{INDEX_NONE}, bWasAlreadyInMap);
 
 			if (bWasAlreadyInMap)
 			{
@@ -420,11 +420,11 @@ void UMassVisualizationComponent::HandleChangesWithExternalIDTracking(UInstanced
 		}
 
 		// it's possible the loop above removed all the data, so we do one last check
-		if (!MassInstanceIds.IsEmpty())
+		if (!EntityHandles.IsEmpty())
 		{
-			check(ElementIds.Num() == MassInstanceIds.Num());
+			check(ElementIds.Num() == EntityHandles.Num());
 
-			const TConstArrayView<int32> InstanceIds = SharedData.UpdateInstanceIds;
+			const TConstArrayView<FMassEntityHandle> InstanceIds = SharedData.UpdateInstanceIds;
 
 			const TArray<FTransform>& InstanceTransforms = SharedData.GetStaticMeshInstanceTransformsArray();
 			const int32 InNumCustomDataFloats = SharedData.GetStaticMeshInstanceCustomFloats().Num();
@@ -439,10 +439,10 @@ void UMassVisualizationComponent::HandleChangesWithExternalIDTracking(UInstanced
 				ISMComponent.SetNumCustomDataFloats(InNumCustomDataFloats);
 			}
 
-			check(MassInstanceIds.Num() == InstanceTransforms.Num());
+			check(EntityHandles.Num() == InstanceTransforms.Num());
 			TArray<FPrimitiveInstanceId> NewIds = ISMComponent.AddInstancesById(InstanceTransforms, /*bWorldSpace=*/true, /*bUpdateNavigation =*/bInitiallyEmpty);
-			check(MassInstanceIds.Num() == NewIds.Num());
-			for (int32 i = 0; i < MassInstanceIds.Num(); ++i)
+			check(EntityHandles.Num() == NewIds.Num());
+			for (int32 i = 0; i < EntityHandles.Num(); ++i)
 			{
 				SharedIdMap.GetByElementId(ElementIds[i]).Value = NewIds[i];
 			}
@@ -461,17 +461,16 @@ void UMassVisualizationComponent::ProcessRemoves(UInstancedStaticMeshComponent&
 {
 	if (!SharedData.GetRemoveInstanceIds().IsEmpty())
 	{
-		FMassISMCSharedData::FIdMap& SharedIdMap = SharedData.GetMutableIdMap();
+		FMassISMCSharedData::FEntityToPrimitiveIdMap& SharedIdMap = SharedData.GetMutableEntityPrimitiveToIdMap();
 		INC_DWORD_STAT_BY(STAT_Mass_VisualizationComponent_InstancesRemovedNum, SharedData.GetRemoveInstanceIds().Num());
 
-		//RemoveInstanceWithIds(SharedData.GetRemoveInstanceIds());
-		TConstArrayView<int32> InstanceIds = SharedData.GetRemoveInstanceIds();
+		TConstArrayView<FMassEntityHandle> EntityHandles = SharedData.GetRemoveInstanceIds();
 
 		TArray<FPrimitiveInstanceId> ISMInstanceIds;
-		ISMInstanceIds.Reserve(InstanceIds.Num());
+		ISMInstanceIds.Reserve(EntityHandles.Num());
 		
 		// Translate Mass IDs to ISMC IDs
-		for (const int32 MassInstanceId : InstanceIds)
+		for (const FMassEntityHandle MassInstanceId : EntityHandles)
 		{
 			Experimental::FHashElementId ElementId = SharedIdMap.FindId(MassInstanceId);
 			if (ElementId.IsValid())
@@ -530,12 +529,12 @@ void UMassVisualizationComponent::EndVisualChanges()
 				// its length is NumCustomDataFloats * InstanceTransforms.Num()
 				ensure(NumCustomDataFloats == 0 || (SharedData.StaticMeshInstanceCustomFloats.Num() == NumCustomDataFloats * SharedData.UpdateInstanceIds.Num()));
 				ISMComponent->SetNumCustomDataFloats(NumCustomDataFloats);
-				TArray<int32>& MassInstanceIds = SharedData.UpdateInstanceIds;
+				TArray<FMassEntityHandle>& EntityHandles = SharedData.UpdateInstanceIds;
 				{
 					// Loop over all the instances in the update and 
 					// 1. Sort the data such that all Adds are last
 					// 2. Remove any duplicates (unsure if they may exist)
-					FMassISMCSharedData::FIdMap& SharedIdMap = SharedData.GetMutableIdMap();
+					FMassISMCSharedData::FEntityToPrimitiveIdMap& SharedIdMap = SharedData.GetMutableEntityPrimitiveToIdMap();
 					// Filter out all updates & duplicate adds
 					TBitArray<> Unprocessed;
 					Unprocessed.SetNum(SharedIdMap.GetMaxIndex(), true);
@@ -554,11 +553,11 @@ void UMassVisualizationComponent::EndVisualChanges()
 						int32 IDIndex;
 					};
 					TArray<FAddItem> ToAdd;
-					ToAdd.Reserve(MassInstanceIds.Num());
-					for (int32 IDIndex = 0; IDIndex < MassInstanceIds.Num(); ++IDIndex)
+					ToAdd.Reserve(EntityHandles.Num());
+					for (int32 IDIndex = 0; IDIndex < EntityHandles.Num(); ++IDIndex)
 					{
 						bool bWasAlreadyInMap = false;
-						Experimental::FHashElementId ElementId = SharedIdMap.FindOrAddId(MassInstanceIds[IDIndex], FPrimitiveInstanceId{INDEX_NONE}, bWasAlreadyInMap);
+						Experimental::FHashElementId ElementId = SharedIdMap.FindOrAddId(EntityHandles[IDIndex], FPrimitiveInstanceId{INDEX_NONE}, bWasAlreadyInMap);
 
 						// if it was already in the map, it may be a duplicate if we have processed it already
 						bool bIsDuplicate = bWasAlreadyInMap && !Unprocessed[ElementId.GetIndex()];
@@ -665,7 +664,7 @@ void FMassInstancedStaticMeshInfo::ClearVisualInstance(FMassISMCSharedDataMap& I
 // FMassLODSignificanceRange
 //---------------------------------------------------------------
 
-void FMassLODSignificanceRange::AddBatchedTransform(const int32 InstanceId, const FTransform& Transform, const FTransform& PrevTransform, const TArray<uint32>& ExcludeStaticMeshRefs)
+void FMassLODSignificanceRange::AddBatchedTransform(const FMassEntityHandle EntityHandle, const FTransform& Transform, const FTransform& PrevTransform, const TArray<uint32>& ExcludeStaticMeshRefs)
 {
 	check(ISMCSharedDataPtr);
 	for (int32 StaticMeshIndex = 0; StaticMeshIndex < StaticMeshRefs.Num(); ++StaticMeshIndex)
@@ -677,7 +676,7 @@ void FMassLODSignificanceRange::AddBatchedTransform(const int32 InstanceId, cons
 
 		if (FMassISMCSharedData* SharedData = ISMCSharedDataPtr->GetAndMarkDirty(StaticMeshRefs[StaticMeshIndex]))
 		{
-			SharedData->UpdateInstanceIds.Add(InstanceId);
+			SharedData->UpdateInstanceIds.Add(EntityHandle);
 			SharedData->StaticMeshInstanceTransforms.Add(Transform);
 			SharedData->StaticMeshInstancePrevTransforms.Add(PrevTransform);
 		}
@@ -701,28 +700,28 @@ void FMassLODSignificanceRange::AddBatchedCustomDataFloats(const TArray<float>&
 	}
 }
 
-void FMassLODSignificanceRange::AddInstance(const int32 InstanceId, const FTransform& Transform)
+void FMassLODSignificanceRange::AddInstance(const FMassEntityHandle EntityHandle, const FTransform& Transform)
 {
 	check(ISMCSharedDataPtr);
 	for (int32 StaticMeshIndex = 0; StaticMeshIndex < StaticMeshRefs.Num(); ++StaticMeshIndex)
 	{
 		if (FMassISMCSharedData* SharedData = ISMCSharedDataPtr->GetAndMarkDirty(StaticMeshRefs[StaticMeshIndex]))
 		{
-			SharedData->UpdateInstanceIds.Add(InstanceId);
+			SharedData->UpdateInstanceIds.Add(EntityHandle);
 			SharedData->StaticMeshInstanceTransforms.Add(Transform);
 			SharedData->StaticMeshInstancePrevTransforms.Add(Transform);
 		}
 	}
 }
 
-void FMassLODSignificanceRange::RemoveInstance(const int32 InstanceId)
+void FMassLODSignificanceRange::RemoveInstance(const FMassEntityHandle EntityHandle)
 {
 	check(ISMCSharedDataPtr);
 	for (int32 StaticMeshIndex = 0; StaticMeshIndex < StaticMeshRefs.Num(); ++StaticMeshIndex)
 	{
 		if (FMassISMCSharedData* SharedData = ISMCSharedDataPtr->GetAndMarkDirty(StaticMeshRefs[StaticMeshIndex]))
 		{
-			SharedData->RemoveInstanceIds.Add(InstanceId);
+			SharedData->RemoveInstanceIds.Add(EntityHandle);
 		}
 	}
 }
diff --git a/Engine/Plugins/Runtime/MassGameplay/Source/MassRepresentation/Public/MassRepresentationTypes.h b/Engine/Plugins/Runtime/MassGameplay/Source/MassRepresentation/Public/MassRepresentationTypes.h
--- a/Engine/Plugins/Runtime/MassGameplay/Source/MassRepresentation/Public/MassRepresentationTypes.h
+++ b/Engine/Plugins/Runtime/MassGameplay/Source/MassRepresentation/Public/MassRepresentationTypes.h
@@ -257,7 +257,7 @@ PRAGMA_ENABLE_DEPRECATION_WARNINGS
 	}
 
 	bool HasUpdatesToApply() const { return UpdateInstanceIds.Num() || RemoveInstanceIds.Num(); }
-	TConstArrayView<int32> GetUpdateInstanceIds() const { return UpdateInstanceIds; }
+	TConstArrayView<FMassEntityHandle> GetUpdateInstanceIds() const { return UpdateInstanceIds; }
 	TConstArrayView<FTransform> GetStaticMeshInstanceTransforms() const { return StaticMeshInstanceTransforms; }
 	/** 
 	 * this function is a flavor we need to interact with older engine API that's using TArray references. 
@@ -265,7 +265,7 @@ PRAGMA_ENABLE_DEPRECATION_WARNINGS
 	 */
 	const TArray<FTransform>& GetStaticMeshInstanceTransformsArray() const { return StaticMeshInstanceTransforms; }
 	TConstArrayView<FTransform> GetStaticMeshInstancePrevTransforms() const { return StaticMeshInstancePrevTransforms; }
-	TConstArrayView<int32> GetRemoveInstanceIds() const { return RemoveInstanceIds; }
+	TConstArrayView<FMassEntityHandle> GetRemoveInstanceIds() const { return RemoveInstanceIds; }
 	TConstArrayView<float> GetStaticMeshInstanceCustomFloats() const { return StaticMeshInstanceCustomFloats; }
 	
 	bool RequiresExternalInstanceIDTracking() const { return bRequiresExternalInstanceIDTracking; }
@@ -275,10 +275,10 @@ PRAGMA_ENABLE_DEPRECATION_WARNINGS
 		*this = FMassISMCSharedData();
 	}
 
-	using FIdMap = Experimental::TRobinHoodHashMap<int32, FPrimitiveInstanceId>;
+	using FEntityToPrimitiveIdMap = Experimental::TRobinHoodHashMap<FMassEntityHandle, FPrimitiveInstanceId>;
 
-	FIdMap& GetMutableIdMap() { return MassInstanceIdToComponentInstanceIdMap; }
-	const FIdMap& GetIdMap() const { return MassInstanceIdToComponentInstanceIdMap; }
+	FEntityToPrimitiveIdMap& GetMutableEntityPrimitiveToIdMap() { return EntityHandleToPrimitiveIdMap; }
+	const FEntityToPrimitiveIdMap& GetEntityPrimitiveToIdMap() const { return EntityHandleToPrimitiveIdMap; }
 
 	int16 GetComponentInstanceIdTouchCounter() const { return ComponentInstanceIdTouchCounter; }
 
@@ -286,10 +286,10 @@ protected:
 	friend FMassLODSignificanceRange;
 	friend UMassVisualizationComponent;
 	/** Buffer holding current frame transforms for the static mesh instances, used to batch update the transforms */
-	TArray<int32> UpdateInstanceIds;
+	TArray<FMassEntityHandle> UpdateInstanceIds;
 	TArray<FTransform> StaticMeshInstanceTransforms;
 	TArray<FTransform> StaticMeshInstancePrevTransforms;
-	TArray<int32> RemoveInstanceIds;
+	TArray<FMassEntityHandle> RemoveInstanceIds;
 
 	/** Buffer holding current frame custom floats for the static mesh instances, used to batch update the ISMs custom data */
 	TArray<float> StaticMeshInstanceCustomFloats;
@@ -314,7 +314,7 @@ private:
 	uint16 ComponentInstanceIdTouchCounter = 0;
 
 protected:
-	FIdMap MassInstanceIdToComponentInstanceIdMap;
+	FEntityToPrimitiveIdMap EntityHandleToPrimitiveIdMap;
 
 	UE_DEPRECATED(5.4, "RefCount is deprecated, use ISMComponentReferencesCount instead")
 	int32 RefCount = 0;
@@ -504,7 +504,7 @@ struct MASSREPRESENTATION_API FMassLODSignificanceRange
 	GENERATED_BODY()
 public:
 
-	void AddBatchedTransform(const int32 InstanceId, const FTransform& Transform, const FTransform& PrevTransform, const TArray<uint32>& ExcludeStaticMeshRefs);
+	void AddBatchedTransform(const FMassEntityHandle EntityHandle, const FTransform& Transform, const FTransform& PrevTransform, const TArray<uint32>& ExcludeStaticMeshRefs);
 
 	// Adds the specified struct reinterpreted as custom floats to our custom data. Individual members of the specified struct should always fit into a float.
 	// When adding any custom data, the custom data must be added for every instance.
@@ -532,9 +532,9 @@ public:
 	void AddBatchedCustomDataFloats(const TArray<float>& CustomFloats, const TArray<uint32>& ExcludeStaticMeshRefs);
 
 	/** Single-instance version of AddBatchedCustomData when called to add entities (as opposed to modify existing ones).*/
-	void AddInstance(const int32 InstanceId, const FTransform& Transform);
+	void AddInstance(const FMassEntityHandle EntityHandle, const FTransform& Transform);
 
-	void RemoveInstance(const int32 InstanceId);
+	void RemoveInstance(const FMassEntityHandle EntityHandle);
 
 	void WriteCustomDataFloatsAtStartIndex(int32 StaticMeshIndex, const TArrayView<float>& CustomFloats, const int32 FloatsPerInstance, const int32 StartIndex, const TArray<uint32>& ExcludeStaticMeshRefs);
 
@@ -547,6 +547,17 @@ public:
 	TArray<uint32> StaticMeshRefs;
 
 	FMassISMCSharedDataMap* ISMCSharedDataPtr = nullptr;
+
+	//-----------------------------------------------------------------------------
+	// DEPRECATED
+	//-----------------------------------------------------------------------------
+	UE_DEPRECATED(5.4, "Deprecated in favor of new version taking FMassEntityHandle parameter instead of int32 to identify the entity. This deprecated function is now defunct.")
+	void AddBatchedTransform(const int32 InstanceId, const FTransform& Transform, const FTransform& PrevTransform, const TArray<uint32>& ExcludeStaticMeshRefs) {}
+	UE_DEPRECATED(5.4, "Deprecated in favor of new version taking FMassEntityHandle parameter instead of int32 to identify the entity. This deprecated function is now defunct.")
+	void AddInstance(const int32 InstanceId, const FTransform& Transform) {}
+	UE_DEPRECATED(5.4, "Deprecated in favor of new version taking FMassEntityHandle parameter instead of int32 to identify the entity. This deprecated function is now defunct.")
+	void RemoveInstance(const int32 InstanceId) {}
+
 };
 
 USTRUCT()
@@ -586,28 +597,28 @@ public:
 		return nullptr;
 	}
 
-	FORCEINLINE void AddBatchedTransform(const int32 InstanceId, const FTransform& Transform, const FTransform& PrevTransform, const float LODSignificance, const float PrevLODSignificance = -1.0f)
+	FORCEINLINE void AddBatchedTransform(const FMassEntityHandle EntityHandle, const FTransform& Transform, const FTransform& PrevTransform, const float LODSignificance, const float PrevLODSignificance = -1.0f)
 	{
 		if (FMassLODSignificanceRange* Range = GetLODSignificanceRange(LODSignificance))
 		{
-			Range->AddBatchedTransform(InstanceId, Transform, PrevTransform, TArray<uint32>());
+			Range->AddBatchedTransform(EntityHandle, Transform, PrevTransform, TArray<uint32>());
 			if(PrevLODSignificance >= 0.0f)
 			{
 				FMassLODSignificanceRange* PrevRange = GetLODSignificanceRange(PrevLODSignificance);
 				if (ensureMsgf(PrevRange, TEXT("Couldn't find a valid LODSignificanceRange for PrevLODSignificance %f"), PrevLODSignificance)
 					&& PrevRange != Range)
 				{
-					PrevRange->AddBatchedTransform(InstanceId, Transform, PrevTransform, Range->StaticMeshRefs);
+					PrevRange->AddBatchedTransform(EntityHandle, Transform, PrevTransform, Range->StaticMeshRefs);
 				}
 			}
 		}
 	}
 
-	FORCEINLINE void RemoveInstance(const int32 InstanceId, const float LODSignificance)
+	FORCEINLINE void RemoveInstance(const FMassEntityHandle EntityHandle, const float LODSignificance)
 	{
 		if (FMassLODSignificanceRange* Range = GetLODSignificanceRange(LODSignificance))
 		{
-			Range->RemoveInstance(InstanceId);
+			Range->RemoveInstance(EntityHandle);
 		}
 	}
 
@@ -698,6 +709,15 @@ protected:
 	TArray<FMassLODSignificanceRange> LODSignificanceRanges;
 
 	friend class UMassVisualizationComponent;
+
+	//-----------------------------------------------------------------------------
+	// DEPRECATED
+	//-----------------------------------------------------------------------------
+public:
+	UE_DEPRECATED(5.4, "Deprecated in favor of new version taking FMassEntityHandle parameter instead of int32 to identify the entity. This deprecated function is now defunct.")
+	void AddBatchedTransform(const int32 InstanceId, const FTransform& Transform, const FTransform& PrevTransform, const float LODSignificance, const float PrevLODSignificance = -1.0f) {}
+	UE_DEPRECATED(5.4, "Deprecated in favor of new version taking FMassEntityHandle parameter instead of int32 to identify the entity. This deprecated function is now defunct.")
+	void RemoveInstance(const int32 InstanceId, const float LODSignificance) {}
 };
 
 #if ENABLE_MT_DETECTOR
diff --git a/Engine/Plugins/Runtime/MassGameplay/Source/MassRepresentation/Public/MassUpdateISMProcessor.h b/Engine/Plugins/Runtime/MassGameplay/Source/MassRepresentation/Public/MassUpdateISMProcessor.h
--- a/Engine/Plugins/Runtime/MassGameplay/Source/MassRepresentation/Public/MassUpdateISMProcessor.h
+++ b/Engine/Plugins/Runtime/MassGameplay/Source/MassRepresentation/Public/MassUpdateISMProcessor.h
@@ -18,7 +18,11 @@ class MASSREPRESENTATION_API UMassUpdateISMProcessor : public UMassProcessor
 public:
 	UMassUpdateISMProcessor();
 
-	static void UpdateISMTransform(int32 EntityId, FMassInstancedStaticMeshInfo& ISMInfo, const FTransform& Transform, const FTransform& PrevTransform, const float LODSignificance, const float PrevLODSignificance = -1.0f);
+	static void UpdateISMTransform(FMassEntityHandle EntityHandle, FMassInstancedStaticMeshInfo& ISMInfo, const FTransform& Transform, const FTransform& PrevTransform, const float LODSignificance, const float PrevLODSignificance = -1.0f);
+
+	UE_DEPRECATED(5.4, "Deprecated in favor of new version taking FMassEntityHandle parameter instead of int32 to identify the entity. This deprecated function is now defunct.")
+	static void UpdateISMTransform(int32 EntityId, FMassInstancedStaticMeshInfo& ISMInfo, const FTransform& Transform, const FTransform& PrevTransform, const float LODSignificance, const float PrevLODSignificance = -1.0f) {}
+
 
 protected:
 
