PR #11400: [Mass] Run post spawn processors only on the freshly spawned entities

diff --git a/Engine/Plugins/Runtime/MassGameplay/Source/MassSpawner/Private/MassSpawner.cpp b/Engine/Plugins/Runtime/MassGameplay/Source/MassSpawner/Private/MassSpawner.cpp
--- a/Engine/Plugins/Runtime/MassGameplay/Source/MassSpawner/Private/MassSpawner.cpp
+++ b/Engine/Plugins/Runtime/MassGameplay/Source/MassSpawner/Private/MassSpawner.cpp
@@ -427,6 +427,9 @@ void AMassSpawner::SpawnGeneratedEntities(TConstArrayView<FMassEntitySpawnDataGe
 	UWorld* World = GetWorld();
 	check(World);
 
+	int32 TotalNum = 0;
+	const int32 StartIndex = AllSpawnedEntities.Num();
+
 	for (const FMassEntitySpawnDataGeneratorResult& Result : Results)
 	{
 		if (Result.NumEntities <= 0)
@@ -447,36 +450,54 @@ void AMassSpawner::SpawnGeneratedEntities(TConstArrayView<FMassEntitySpawnDataGe
 				FSpawnedEntities& SpawnedEntities = AllSpawnedEntities.AddDefaulted_GetRef();
 				SpawnedEntities.TemplateID = EntityTemplate.GetTemplateID();
 				SpawnerSystem->SpawnEntities(EntityTemplate.GetTemplateID(), Result.NumEntities, Result.SpawnData, Result.SpawnDataProcessor, SpawnedEntities.Entities);
+				TotalNum += SpawnedEntities.Entities.Num();
 			}
 		}
 	}
 
-	// Run post spawn processors on all Mass entities that matches the queries.
-	// @todo: we might need a way to specify that these are ran only on the freshly spawned entities.
-	
-	TArray<UMassProcessor*> Processors;
-	TSet<TSubclassOf<UMassProcessor>> AddedProcessors;
-
-	for (const FMassEntitySpawnDataGeneratorResult& Result : Results)
+	// Run post spawn processors only on the freshly spawned entities.
+	if (TotalNum)
 	{
-		for (const TSubclassOf<UMassProcessor>& ProcessorClass : Result.PostSpawnProcessors)
+		TArray<UMassProcessor*> Processors;
+		TSet<TSubclassOf<UMassProcessor>> AddedProcessorClasses;
+
+		for (const FMassEntitySpawnDataGeneratorResult& Result : Results)
 		{
-			if (AddedProcessors.Contains(ProcessorClass) == false)
+			for (const TSubclassOf<UMassProcessor>& ProcessorClass : Result.PostSpawnProcessors)
 			{
-				if (UMassProcessor* Processor = GetPostSpawnProcessor(ProcessorClass))
+				if (AddedProcessorClasses.Contains(ProcessorClass) == false)
 				{
-					Processors.Add(Processor);
+					if (UMassProcessor* Processor = GetPostSpawnProcessor(ProcessorClass))
+					{
+						Processors.Add(Processor);
+					}
+					AddedProcessorClasses.Add(ProcessorClass);
 				}
-				AddedProcessors.Add(ProcessorClass);
 			}
 		}
-	}
 
-	if (Processors.Num() > 0 && World)
-	{
-		FMassEntityManager& EntityManager = UE::Mass::Utils::GetEntityManagerChecked(*World);
-		FMassProcessingContext ProcessingContext(EntityManager, /*TimeDelta=*/0.0f);
-		UE::Mass::Executor::RunProcessorsView(Processors, ProcessingContext);
+		if (Processors.Num() > 0)
+		{
+			FMassEntityManager& EntityManager = UE::Mass::Utils::GetEntityManagerChecked(*World);
+			FMassProcessingContext ProcessingContext(EntityManager, /*TimeDelta=*/0.0f);
+
+			// gather freshly spawned entities
+			TArray<FMassEntityHandle> AllEntities;
+			AllEntities.Reserve(TotalNum);
+			
+			for (int32 Index = StartIndex; Index < AllSpawnedEntities.Num(); ++Index)
+			{
+				AllEntities.Append(AllSpawnedEntities[Index].Entities);
+			}
+
+			// create entity collections and run Processors on them. 
+			TArray<FMassArchetypeEntityCollection> EntityCollections;
+			UE::Mass::Utils::CreateEntityCollections(EntityManager, AllEntities, FMassArchetypeEntityCollection::NoDuplicates, EntityCollections);
+			for (const FMassArchetypeEntityCollection& Collection : EntityCollections)
+			{
+				UE::Mass::Executor::RunProcessorsView(Processors, ProcessingContext, &Collection);
+			}
+		}
 	}
 
 	OnSpawningFinishedEvent.Broadcast();
