PR #11378: MassSpawner add more safty check

mz> Skipping AMassSpawner::DoSpawning work when no entity types specified (user error).

diff --git a/Engine/Plugins/Runtime/MassGameplay/Source/MassSpawner/Private/MassSpawner.cpp b/Engine/Plugins/Runtime/MassGameplay/Source/MassSpawner/Private/MassSpawner.cpp
--- a/Engine/Plugins/Runtime/MassGameplay/Source/MassSpawner/Private/MassSpawner.cpp
+++ b/Engine/Plugins/Runtime/MassGameplay/Source/MassSpawner/Private/MassSpawner.cpp
@@ -285,6 +285,12 @@ void AMassSpawner::DoSpawning()
 		UE_VLOG_UELOG(this, LogMassSpawner, Warning, TEXT("No Spawn Data Generators configured."));
 		return;
 	}
+	
+	if (EntityTypes.Num() == 0)
+	{
+		UE_VLOG_UELOG(this, LogMassSpawner, Warning, TEXT("No EntityTypes configured."));
+		return;
+	}
 
 	AllGeneratedResults.Reset();
 	
