[Mass] updated some processors that shouldn't run in Editor

diff --git a/Engine/Plugins/AI/MassAI/Source/MassNavigation/Private/MassNavigationProcessors.cpp b/Engine/Plugins/AI/MassAI/Source/MassNavigation/Private/MassNavigationProcessors.cpp
--- a/Engine/Plugins/AI/MassAI/Source/MassNavigation/Private/MassNavigationProcessors.cpp
+++ b/Engine/Plugins/AI/MassAI/Source/MassNavigation/Private/MassNavigationProcessors.cpp
@@ -24,7 +24,7 @@
 UMassOffLODNavigationProcessor::UMassOffLODNavigationProcessor()
 	: EntityQuery_Conditional(*this)
 {
-	ExecutionFlags = (int32)EProcessorExecutionFlags::All;
+	ExecutionFlags = (int32)EProcessorExecutionFlags::AllNetModes;
 	ExecutionOrder.ExecuteInGroup = UE::Mass::ProcessorGroupNames::Movement;
 	ExecutionOrder.ExecuteAfter.Add(UE::Mass::ProcessorGroupNames::Avoidance); // @todo: remove this direct dependency
 }
diff --git a/Engine/Plugins/AI/MassAI/Source/MassZoneGraphNavigation/Private/MassZoneGraphNavigationProcessors.cpp b/Engine/Plugins/AI/MassAI/Source/MassZoneGraphNavigation/Private/MassZoneGraphNavigationProcessors.cpp
--- a/Engine/Plugins/AI/MassAI/Source/MassZoneGraphNavigation/Private/MassZoneGraphNavigationProcessors.cpp
+++ b/Engine/Plugins/AI/MassAI/Source/MassZoneGraphNavigation/Private/MassZoneGraphNavigationProcessors.cpp
@@ -43,7 +43,7 @@ UMassZoneGraphLocationInitializer::UMassZoneGraphLocationInitializer()
 {
 	ObservedType = FMassZoneGraphLaneLocationFragment::StaticStruct();
 	Operation = EMassObservedOperation::Add;
-	ExecutionFlags = (int32)EProcessorExecutionFlags::All;
+	ExecutionFlags = (int32)EProcessorExecutionFlags::AllNetModes;
 }
 
 void UMassZoneGraphLocationInitializer::ConfigureQueries()
diff --git a/Engine/Plugins/Runtime/MassGameplay/Source/MassActors/Private/Translators/MassCapsuleComponentTranslators.cpp b/Engine/Plugins/Runtime/MassGameplay/Source/MassActors/Private/Translators/MassCapsuleComponentTranslators.cpp
--- a/Engine/Plugins/Runtime/MassGameplay/Source/MassActors/Private/Translators/MassCapsuleComponentTranslators.cpp
+++ b/Engine/Plugins/Runtime/MassGameplay/Source/MassActors/Private/Translators/MassCapsuleComponentTranslators.cpp
@@ -13,7 +13,7 @@
 UMassCapsuleTransformToMassTranslator::UMassCapsuleTransformToMassTranslator()
 	: EntityQuery(*this)
 {
-	ExecutionFlags = (int32)EProcessorExecutionFlags::All;
+	ExecutionFlags = (int32)EProcessorExecutionFlags::AllNetModes;
 	RequiredTags.Add<FMassCapsuleTransformCopyToMassTag>();
 	ExecutionOrder.ExecuteInGroup = UE::Mass::ProcessorGroupNames::SyncWorldToMass;
 }
@@ -47,7 +47,7 @@ void UMassCapsuleTransformToMassTranslator::Execute(FMassEntityManager& EntityMa
 UMassTransformToActorCapsuleTranslator::UMassTransformToActorCapsuleTranslator()
 	: EntityQuery(*this)
 {
-	ExecutionFlags = (int32)EProcessorExecutionFlags::All;
+	ExecutionFlags = (int32)EProcessorExecutionFlags::AllNetModes;
 	RequiredTags.Add<FMassCapsuleTransformCopyToActorTag>();
 	ExecutionOrder.ExecuteInGroup = UE::Mass::ProcessorGroupNames::UpdateWorldFromMass;
 	ExecutionOrder.ExecuteAfter.Add(UE::Mass::ProcessorGroupNames::Movement);
diff --git a/Engine/Plugins/Runtime/MassGameplay/Source/MassActors/Private/Translators/MassCharacterMovementTranslators.cpp b/Engine/Plugins/Runtime/MassGameplay/Source/MassActors/Private/Translators/MassCharacterMovementTranslators.cpp
--- a/Engine/Plugins/Runtime/MassGameplay/Source/MassActors/Private/Translators/MassCharacterMovementTranslators.cpp
+++ b/Engine/Plugins/Runtime/MassGameplay/Source/MassActors/Private/Translators/MassCharacterMovementTranslators.cpp
@@ -16,7 +16,7 @@
 UMassCharacterMovementToMassTranslator::UMassCharacterMovementToMassTranslator()
 	: EntityQuery(*this)
 {
-	ExecutionFlags = (int32)EProcessorExecutionFlags::All;
+	ExecutionFlags = (int32)EProcessorExecutionFlags::AllNetModes;
 	ExecutionOrder.ExecuteInGroup = UE::Mass::ProcessorGroupNames::SyncWorldToMass;
 	RequiredTags.Add<FMassCharacterMovementCopyToMassTag>();
 }
@@ -57,7 +57,7 @@ void UMassCharacterMovementToMassTranslator::Execute(FMassEntityManager& EntityM
 UMassCharacterMovementToActorTranslator::UMassCharacterMovementToActorTranslator()
 	: EntityQuery(*this)
 {
-	ExecutionFlags = (int32)EProcessorExecutionFlags::All;
+	ExecutionFlags = (int32)EProcessorExecutionFlags::AllNetModes;
 	ExecutionOrder.ExecuteInGroup = UE::Mass::ProcessorGroupNames::UpdateWorldFromMass;
 	ExecutionOrder.ExecuteAfter.Add(UE::Mass::ProcessorGroupNames::Movement);
 	RequiredTags.Add<FMassCharacterMovementCopyToActorTag>();
@@ -95,7 +95,7 @@ void UMassCharacterMovementToActorTranslator::Execute(FMassEntityManager& Entity
 UMassCharacterOrientationToMassTranslator::UMassCharacterOrientationToMassTranslator()
 	: EntityQuery(*this)
 {
-	ExecutionFlags = (int32)EProcessorExecutionFlags::All;
+	ExecutionFlags = (int32)EProcessorExecutionFlags::AllNetModes;
 	ExecutionOrder.ExecuteInGroup = UE::Mass::ProcessorGroupNames::SyncWorldToMass;
 	RequiredTags.Add<FMassCharacterOrientationCopyToMassTag>();
 }
@@ -135,7 +135,7 @@ void UMassCharacterOrientationToMassTranslator::Execute(FMassEntityManager& Enti
 UMassCharacterOrientationToActorTranslator::UMassCharacterOrientationToActorTranslator()
 	: EntityQuery(*this)
 {
-	ExecutionFlags = (int32)EProcessorExecutionFlags::All;
+	ExecutionFlags = (int32)EProcessorExecutionFlags::AllNetModes;
 	ExecutionOrder.ExecuteInGroup = UE::Mass::ProcessorGroupNames::UpdateWorldFromMass;
 	ExecutionOrder.ExecuteAfter.Add(UE::Mass::ProcessorGroupNames::Movement);
 	RequiredTags.Add<FMassCharacterOrientationCopyToActorTag>();
diff --git a/Engine/Plugins/Runtime/MassGameplay/Source/MassLOD/Private/MassLODCollectorProcessor.cpp b/Engine/Plugins/Runtime/MassGameplay/Source/MassLOD/Private/MassLODCollectorProcessor.cpp
--- a/Engine/Plugins/Runtime/MassGameplay/Source/MassLOD/Private/MassLODCollectorProcessor.cpp
+++ b/Engine/Plugins/Runtime/MassGameplay/Source/MassLOD/Private/MassLODCollectorProcessor.cpp
@@ -11,7 +11,7 @@ UMassLODCollectorProcessor::UMassLODCollectorProcessor()
 {
 	bAutoRegisterWithProcessingPhases = false;
 
-	ExecutionFlags = (int32)EProcessorExecutionFlags::All;
+	ExecutionFlags = (int32)EProcessorExecutionFlags::AllNetModes;
 
 	ExecutionOrder.ExecuteInGroup = UE::Mass::ProcessorGroupNames::LODCollector;
 	ExecutionOrder.ExecuteAfter.Add(UE::Mass::ProcessorGroupNames::SyncWorldToMass);
diff --git a/Engine/Plugins/Runtime/MassGameplay/Source/MassLOD/Private/MassLODDistanceCollectorProcessor.cpp b/Engine/Plugins/Runtime/MassGameplay/Source/MassLOD/Private/MassLODDistanceCollectorProcessor.cpp
--- a/Engine/Plugins/Runtime/MassGameplay/Source/MassLOD/Private/MassLODDistanceCollectorProcessor.cpp
+++ b/Engine/Plugins/Runtime/MassGameplay/Source/MassLOD/Private/MassLODDistanceCollectorProcessor.cpp
@@ -15,7 +15,7 @@ UMassLODDistanceCollectorProcessor::UMassLODDistanceCollectorProcessor()
 {
 	bAutoRegisterWithProcessingPhases = true;
 
-	ExecutionFlags = (int32)EProcessorExecutionFlags::All;
+	ExecutionFlags = (int32)EProcessorExecutionFlags::AllNetModes;
 
 	ExecutionOrder.ExecuteInGroup = UE::Mass::ProcessorGroupNames::LODCollector;
 	ExecutionOrder.ExecuteAfter.Add(UE::Mass::ProcessorGroupNames::SyncWorldToMass);
diff --git a/Engine/Plugins/Runtime/MassGameplay/Source/MassRepresentation/Private/MassDistanceLODProcessor.cpp b/Engine/Plugins/Runtime/MassGameplay/Source/MassRepresentation/Private/MassDistanceLODProcessor.cpp
--- a/Engine/Plugins/Runtime/MassGameplay/Source/MassRepresentation/Private/MassDistanceLODProcessor.cpp
+++ b/Engine/Plugins/Runtime/MassGameplay/Source/MassRepresentation/Private/MassDistanceLODProcessor.cpp
@@ -7,7 +7,7 @@
 UMassDistanceLODProcessor::UMassDistanceLODProcessor()
 {
 	bAutoRegisterWithProcessingPhases = true;
-	ExecutionFlags = (int32)(EProcessorExecutionFlags::All);
+	ExecutionFlags = (int32)(EProcessorExecutionFlags::AllNetModes);
 	ExecutionOrder.ExecuteInGroup = UE::Mass::ProcessorGroupNames::LOD;
 	ExecutionOrder.ExecuteAfter.Add(UE::Mass::ProcessorGroupNames::LODCollector);
 }
diff --git a/Engine/Plugins/Runtime/MassGameplay/Source/MassRepresentation/Private/MassRepresentationProcessor.cpp b/Engine/Plugins/Runtime/MassGameplay/Source/MassRepresentation/Private/MassRepresentationProcessor.cpp
--- a/Engine/Plugins/Runtime/MassGameplay/Source/MassRepresentation/Private/MassRepresentationProcessor.cpp
+++ b/Engine/Plugins/Runtime/MassGameplay/Source/MassRepresentation/Private/MassRepresentationProcessor.cpp
@@ -421,7 +421,7 @@ UMassRepresentationFragmentDestructor::UMassRepresentationFragmentDestructor()
 {
 	ObservedType = FMassRepresentationFragment::StaticStruct();
 	Operation = EMassObservedOperation::Remove;
-	ExecutionFlags = (int32)EProcessorExecutionFlags::All;
+	ExecutionFlags = (int32)EProcessorExecutionFlags::AllNetModes;
 	bRequiresGameThreadExecution = true; // due to FMassRepresentationSubsystemSharedFragment.RepresentationSubsystem use
 }
 
diff --git a/Engine/Plugins/Runtime/MassGameplay/Source/MassSmartObjects/Private/MassSmartObjectProcessor.cpp b/Engine/Plugins/Runtime/MassGameplay/Source/MassSmartObjects/Private/MassSmartObjectProcessor.cpp
--- a/Engine/Plugins/Runtime/MassGameplay/Source/MassSmartObjects/Private/MassSmartObjectProcessor.cpp
+++ b/Engine/Plugins/Runtime/MassGameplay/Source/MassSmartObjects/Private/MassSmartObjectProcessor.cpp
@@ -389,7 +389,7 @@ UMassSmartObjectUserFragmentDeinitializer::UMassSmartObjectUserFragmentDeinitial
 {
 	ObservedType = FMassSmartObjectUserFragment::StaticStruct();
 	Operation = EMassObservedOperation::Remove;
-	ExecutionFlags = (int32)EProcessorExecutionFlags::All;
+	ExecutionFlags = (int32)EProcessorExecutionFlags::AllNetModes;
 }
 
 void UMassSmartObjectUserFragmentDeinitializer::ConfigureQueries()
