[Mass] added support for all owned queries for processors inheriting from UMassSignalProcessorBase

diff --git a/Engine/Plugins/Runtime/MassGameplay/Source/MassSignals/Private/MassSignalProcessorBase.cpp b/Engine/Plugins/Runtime/MassGameplay/Source/MassSignals/Private/MassSignalProcessorBase.cpp
--- a/Engine/Plugins/Runtime/MassGameplay/Source/MassSignals/Private/MassSignalProcessorBase.cpp
+++ b/Engine/Plugins/Runtime/MassGameplay/Source/MassSignals/Private/MassSignalProcessorBase.cpp
@@ -55,8 +55,10 @@ void UMassSignalProcessorBase::Execute(FMassEntityManager& EntityManager, FMassE
 		return;
 	}
 
-	EntityQuery.CacheArchetypes(EntityManager);
-	if (EntityQuery.GetArchetypes().Num() > 0)
+	TArray<FMassArchetypeHandle> ValidArchetypes;
+	GetArchetypesMatchingOwnedQueries(EntityManager, ValidArchetypes);
+
+	if (ValidArchetypes.Num() > 0)
 	{
 		// EntitySet stores unique array of entities per specified archetype.
 		// FMassArchetypeEntityCollection expects an array of entities, a set is used to detect unique ones.
@@ -72,7 +74,7 @@ void UMassSignalProcessorBase::Execute(FMassEntityManager& EntityManager, FMassE
 		};
 		TArray<FEntitySet> EntitySets;
 
-		for (const FMassArchetypeHandle& Archetype : EntityQuery.GetArchetypes())
+		for (const FMassArchetypeHandle& Archetype : ValidArchetypes)
 		{
 			FEntitySet& Set = EntitySets.AddDefaulted_GetRef();
 			Set.Archetype = Archetype;
