Changed how we hash ISMComponents in MassRepresentation

Using object path can cause issues in WP worlds since levels being removed might get their name changed which changes the hash and causes issues when removing ISM shared data.

diff --git a/Engine/Plugins/Runtime/MassGameplay/Source/MassRepresentation/Private/MassVisualizationComponent.cpp b/Engine/Plugins/Runtime/MassGameplay/Source/MassRepresentation/Private/MassVisualizationComponent.cpp
--- a/Engine/Plugins/Runtime/MassGameplay/Source/MassRepresentation/Private/MassVisualizationComponent.cpp
+++ b/Engine/Plugins/Runtime/MassGameplay/Source/MassRepresentation/Private/MassVisualizationComponent.cpp
@@ -23,6 +23,41 @@ DECLARE_DWORD_COUNTER_STAT(TEXT("VisualizationComp Instances Removed"), STAT_Mas
 DECLARE_DWORD_COUNTER_STAT(TEXT("VisualizationComp Instances Added"), STAT_Mass_VisualizationComponent_InstancesAddedNum, STATGROUP_Mass);
 
 
+namespace UE::Mass::Private
+{
+	uint32 CalculateComponentHash(const TObjectPtr<UInstancedStaticMeshComponent>& ISMComponent)
+	{
+		constexpr bool bUseObjectPathHash = false;
+		constexpr bool bUseAssetPathHash = false;
+
+		if constexpr (bUseObjectPathHash)
+		{
+			// This approach can result in (desired) reuse of hash if a given component is being recreated for 
+            // due to it or the owner being streamed in and out repeatedly. This approach however is fragile 
+            // to potential outer level name changes as part of GC preparation (depending on settings and context 
+            // might take place with WorldPartition approach)
+			check(ISMComponent);
+			const FString ISMCPath = ISMComponent->GetPathName();
+			return GetTypeHash(ISMCPath);
+		}
+		else if constexpr (bUseAssetPathHash)
+		{
+			// This approach is great if a given ISMComponent's outer ULevel is never getting 
+			// used twice at the same time (i.e. it will break if said ULevel gets instantiated
+			// multiple times)
+			const FSoftObjectPath ObjectPath(ISMComponent);
+			const FString AssetPathString = ObjectPath.GetAssetPathString();
+			return HashCombine(GetTypeHash(AssetPathString), GetTypeHash(ISMComponent.GetFName()));
+		}
+		else
+		{
+			// fallback hashing that will always work at runtime, but requires care when serializing
+			// Note that at the moment Mass doesn't have a native serialization implementation.
+			return PointerHash(ISMComponent.Get());
+		}
+	}
+}
+
 //---------------------------------------------------------------
 // UMassVisualizationComponent
 //---------------------------------------------------------------
@@ -131,9 +166,9 @@ FStaticMeshInstanceVisualizationDescHandle UMassVisualizationComponent::AddVisua
 			check(VisualHandle.IsValid());
 		}
 
-		const FString ISMCPath = ISMComponents[EntryIndex]->GetPathName();
-		const uint32 ISMComponentPathHash = GetTypeHash(ISMCPath); 
+		const uint32 ISMComponentPathHash = UE::Mass::Private::CalculateComponentHash(ISMComponents[EntryIndex]);
 #if WITH_MASSGAMEPLAY_DEBUG
+		const FString ISMCPath = ISMComponents[EntryIndex]->GetPathName();
 		TArray<FString>& DebugPaths = DebugHashToPathMap.FindOrAdd(ISMComponentPathHash, TArray<FString>());
 		if (!ensureMsgf(DebugPaths.Add(ISMCPath) == 0, TEXT("Multiple ISMC paths resulting in the same hash")))
 		{
@@ -175,8 +210,7 @@ void UMassVisualizationComponent::RemoveVisualDesc(const FStaticMeshInstanceVisu
 		{
 			// @todo using ISMComponent.GetPathName() here might be wrong for cases where we use GetTypeHash(MeshDesc)
 			// to create the ISMComponentMap key
-			const FString ISMCPath = ISMComponent.GetPathName();
-			const uint32 ISMComponentPathHash = GetTypeHash(ISMCPath);
+			const uint32 ISMComponentPathHash = UE::Mass::Private::CalculateComponentHash(ISMComponent);
 
 			const bool bValidKey = ISMComponentMap.Contains(ISMComponentPathHash);
 			checkf(bValidKey, TEXT("Failed to find %u as a key in ISMComponentMap, ISMC path: %s"), ISMComponentPathHash, *ISMComponent.GetPathName());
@@ -188,6 +222,7 @@ void UMassVisualizationComponent::RemoveVisualDesc(const FStaticMeshInstanceVisu
 		
 			ISMCSharedData.Remove(ISMComponentPathHash);
 #if WITH_MASSGAMEPLAY_DEBUG
+			const FString ISMCPath = ISMComponent.GetPathName();
 			TArray<FString>& Paths = DebugHashToPathMap.FindChecked(ISMComponentPathHash);
 			ensure(Paths.Remove(ISMCPath));
 #endif // WITH_MASSGAMEPLAY_DEBUG
@@ -264,7 +299,7 @@ void UMassVisualizationComponent::ConstructStaticMeshComponents()
 					SharedData->SetISMComponent(*ISMC);
 				}
 
-				const uint32 ISMComponentPathHash = GetTypeHash(ISMC->GetPathName());
+				const uint32 ISMComponentPathHash = UE::Mass::Private::CalculateComponentHash(ISMC);
 				ensureMsgf(ISMComponentMap.Find(ISMComponentPathHash) == nullptr, TEXT("We've just created the ISMC that's being used here, so this check failing indicates hash-clash."));
 				ISMComponentMap.Add(ISMComponentPathHash, VisualDescHandle); 
 			}
@@ -763,7 +798,7 @@ void UMassVisualizationComponent::RemoveISMComponent(UInstancedStaticMeshCompone
 {
 	UE_MT_SCOPED_WRITE_ACCESS(InstancedStaticMeshInfosDetector);
 
-	const uint32 ISMComponentPathHash = GetTypeHash(ISMComponent.GetPathName());
+	const uint32 ISMComponentPathHash = UE::Mass::Private::CalculateComponentHash(&ISMComponent);
 	const FStaticMeshInstanceVisualizationDescHandle* VisualDescHandlePtr = ISMComponentMap.Find(ISMComponentPathHash);
 	if (VisualDescHandlePtr)
 	{
