Switched MassVisualization from using hashes to represent ISMCs over to using TObjectKey. This change simplifies reasoning about the code as well as its stability - we no longer have hashes that are not being use and still dangling in the system.

diff --git a/Engine/Plugins/Runtime/MassGameplay/Source/MassRepresentation/Private/MassRepresentationTypes.cpp b/Engine/Plugins/Runtime/MassGameplay/Source/MassRepresentation/Private/MassRepresentationTypes.cpp
--- a/Engine/Plugins/Runtime/MassGameplay/Source/MassRepresentation/Private/MassRepresentationTypes.cpp
+++ b/Engine/Plugins/Runtime/MassGameplay/Source/MassRepresentation/Private/MassRepresentationTypes.cpp
@@ -78,6 +78,20 @@ default: \
 //-----------------------------------------------------------------------------
 // FMassInstancedStaticMeshInfo
 //-----------------------------------------------------------------------------
+void FMassInstancedStaticMeshInfo::ClearVisualInstance(UInstancedStaticMeshComponent& ISMComponent)
+{
+	if (InstancedStaticMeshComponents.RemoveSingleSwap(&ISMComponent, EAllowShrinking::No))
+	{
+		ISMComponent.ClearInstances();
+		ISMComponent.DestroyComponent();
+	}
+	else
+	{
+		InstancedStaticMeshComponents.Reset();
+		LODSignificanceRanges.Reset();
+	}
+}
+
 void FMassInstancedStaticMeshInfo::Reset()
 {
 	Desc.Reset();
@@ -92,3 +106,26 @@ FMassStaticMeshInstanceVisualizationMeshDesc::FMassStaticMeshInstanceVisualizati
 {
 	ISMComponentClass = UInstancedStaticMeshComponent::StaticClass();
 }
+
+//-----------------------------------------------------------------------------
+// FMassISMCSharedDataMap - DEPRECATED
+//-----------------------------------------------------------------------------
+FMassISMCSharedData& FMassISMCSharedDataMap::GetAndMarkDirtyChecked(const uint32 Hash)
+{
+	static FMassISMCSharedData Dummy;
+	return Dummy;
+}
+
+FMassISMCSharedData* FMassISMCSharedDataMap::GetAndMarkDirty(const uint32 Hash)
+{
+	return nullptr;
+}
+
+FMassISMCSharedData* FMassISMCSharedDataMap::Find(const uint32 Hash)
+{
+	return nullptr;
+}
+
+void FMassISMCSharedDataMap::Remove(const uint32 Hash)
+{
+}
diff --git a/Engine/Plugins/Runtime/MassGameplay/Source/MassRepresentation/Private/MassVisualizationComponent.cpp b/Engine/Plugins/Runtime/MassGameplay/Source/MassRepresentation/Private/MassVisualizationComponent.cpp
--- a/Engine/Plugins/Runtime/MassGameplay/Source/MassRepresentation/Private/MassVisualizationComponent.cpp
+++ b/Engine/Plugins/Runtime/MassGameplay/Source/MassRepresentation/Private/MassVisualizationComponent.cpp
@@ -22,42 +22,6 @@ DECLARE_CYCLE_STAT(TEXT("Mass Visualization HandleIDs"), STAT_Mass_Visualization
 DECLARE_DWORD_COUNTER_STAT(TEXT("VisualizationComp Instances Removed"), STAT_Mass_VisualizationComponent_InstancesRemovedNum, STATGROUP_Mass);
 DECLARE_DWORD_COUNTER_STAT(TEXT("VisualizationComp Instances Added"), STAT_Mass_VisualizationComponent_InstancesAddedNum, STATGROUP_Mass);
 
-
-namespace UE::Mass::Private
-{
-	uint32 CalculateComponentHash(const TObjectPtr<UInstancedStaticMeshComponent>& ISMComponent)
-	{
-		constexpr bool bUseObjectPathHash = false;
-		constexpr bool bUseAssetPathHash = false;
-
-		if constexpr (bUseObjectPathHash)
-		{
-			// This approach can result in (desired) reuse of hash if a given component is being recreated for 
-            // due to it or the owner being streamed in and out repeatedly. This approach however is fragile 
-            // to potential outer level name changes as part of GC preparation (depending on settings and context 
-            // might take place with WorldPartition approach)
-			check(ISMComponent);
-			const FString ISMCPath = ISMComponent->GetPathName();
-			return GetTypeHash(ISMCPath);
-		}
-		else if constexpr (bUseAssetPathHash)
-		{
-			// This approach is great if a given ISMComponent's outer ULevel is never getting 
-			// used twice at the same time (i.e. it will break if said ULevel gets instantiated
-			// multiple times)
-			const FSoftObjectPath ObjectPath(ISMComponent);
-			const FString AssetPathString = ObjectPath.GetAssetPathString();
-			return HashCombine(GetTypeHash(AssetPathString), GetTypeHash(ISMComponent.GetFName()));
-		}
-		else
-		{
-			// fallback hashing that will always work at runtime, but requires care when serializing
-			// Note that at the moment Mass doesn't have a native serialization implementation.
-			return PointerHash(ISMComponent.Get());
-		}
-	}
-}
-
 //---------------------------------------------------------------
 // UMassVisualizationComponent
 //---------------------------------------------------------------
@@ -112,7 +76,12 @@ FStaticMeshInstanceVisualizationDescHandle UMassVisualizationComponent::FindOrAd
 		{
 			if (MeshDesc.Mesh && MeshDesc.ISMComponentClass)
 			{
-				ISMCSharedData.FindOrAdd(GetTypeHash(MeshDesc), FMassISMCSharedData());
+				// if we've already encountered MeshDesc in the past MeshDescToISMCMap already contains information
+				// about actual ISMC used to represent it, and at the same time indicates the ISMCSharedData data
+				// tied to it. Regardless we need to process all MeshDesc instances here so that we have all the 
+				// data ready when InstancedSMComponentsRequiringConstructing gets processed next time
+				// UMassVisualizationComponent::ConstructStaticMeshComponents gets called.
+				MeshDescToISMCMap.FindOrAdd(GetTypeHash(MeshDesc), FISMCSharedDataKey());
 				bValidDescription = true;
 			}
 		}
@@ -120,11 +89,10 @@ FStaticMeshInstanceVisualizationDescHandle UMassVisualizationComponent::FindOrAd
 		if (bValidDescription)
 		{
 			VisualDescHandle = AddInstancedStaticMeshInfo(Desc);
-			BuildLODSignificanceForInfo(InstancedStaticMeshInfos[VisualDescHandle.ToIndex()]);
+			check(VisualDescHandle.IsValid());
 
+			// VisualDescHandle is a valid handle now, but there's initialization pending, performed in ConstructStaticMeshComponents
 			InstancedSMComponentsRequiringConstructing.Add(VisualDescHandle);
-
-			check(VisualDescHandle.IsValid());
 		}
 		else
 		{
@@ -148,7 +116,7 @@ FStaticMeshInstanceVisualizationDescHandle UMassVisualizationComponent::AddVisua
 	UE_MT_SCOPED_WRITE_ACCESS(InstancedStaticMeshInfosDetector);
 
 	FStaticMeshInstanceVisualizationDescHandle VisualHandle;
-	TArray<uint32> ISMComponentPathHashes;
+	TArray<UInstancedStaticMeshComponent*> ISMComponentsUsed;
 
 	for (int32 EntryIndex = 0; EntryIndex < Desc.Meshes.Num(); ++EntryIndex)
 	{
@@ -166,29 +134,16 @@ FStaticMeshInstanceVisualizationDescHandle UMassVisualizationComponent::AddVisua
 			check(VisualHandle.IsValid());
 		}
 
-		const uint32 ISMComponentPathHash = UE::Mass::Private::CalculateComponentHash(ISMComponents[EntryIndex]);
-#if WITH_MASSGAMEPLAY_DEBUG
-		const FString ISMCPath = ISMComponents[EntryIndex]->GetPathName();
-		TArray<FString>& DebugPaths = DebugHashToPathMap.FindOrAdd(ISMComponentPathHash, TArray<FString>());
-		if (!ensureMsgf(DebugPaths.Add(ISMCPath) == 0, TEXT("Multiple ISMC paths resulting in the same hash")))
-		{
-			UE_VLOG_UELOG(this, LogMassRepresentation, Error, TEXT("%hs multiple ISMComponents resulting in identical hash %u"), __FUNCTION__, ISMComponentPathHash);
-			for (const FString& Path : DebugPaths)
-			{
-				UE_VLOG_UELOG(this, LogMassRepresentation, Error, TEXT("\t%s"), *Path);
-			}
-		}
-#endif // WITH_MASSGAMEPLAY_DEBUG
-		FMassISMCSharedData& NewData = ISMCSharedData.FindOrAdd(ISMComponentPathHash, FMassISMCSharedData(ISMComponents[EntryIndex], /*bInRequiresExternalInstanceIDTracking=*/true));
+		FMassISMCSharedData& NewData = ISMCSharedData.FindOrAdd(ISMComponents[EntryIndex], FMassISMCSharedData(ISMComponents[EntryIndex], /*bInRequiresExternalInstanceIDTracking=*/true));
 		InstancedStaticMeshInfos[VisualHandle.ToIndex()].AddISMComponent(NewData);
-		ISMComponentPathHashes.Add(ISMComponentPathHash);
+		ISMComponentsUsed.Add(ISMComponents[EntryIndex]);
 	
-		ISMComponentMap.Add(ISMComponentPathHash, VisualHandle);
+		ISMComponentMap.Add(ISMComponents[EntryIndex], VisualHandle);
 	}
 
 	if (VisualHandle.IsValid())
 	{
-		BuildLODSignificanceForInfo(InstancedStaticMeshInfos[VisualHandle.ToIndex()], ISMComponentPathHashes);
+		BuildLODSignificanceForInfo(InstancedStaticMeshInfos[VisualHandle.ToIndex()], ISMComponentsUsed);
 	}
 
 	return VisualHandle;
@@ -208,24 +163,15 @@ void UMassVisualizationComponent::RemoveVisualDesc(const FStaticMeshInstanceVisu
 	{
 		for (TObjectPtr<UInstancedStaticMeshComponent>& ISMComponent : InstancedStaticMeshInfos[VisualizationHandle.ToIndex()].InstancedStaticMeshComponents)
 		{
-			// @todo using ISMComponent.GetPathName() here might be wrong for cases where we use GetTypeHash(MeshDesc)
-			// to create the ISMComponentMap key
-			const uint32 ISMComponentPathHash = UE::Mass::Private::CalculateComponentHash(ISMComponent);
-
-			const bool bValidKey = ISMComponentMap.Contains(ISMComponentPathHash);
-			checkf(bValidKey, TEXT("Failed to find %u as a key in ISMComponentMap, ISMC path: %s"), ISMComponentPathHash, *ISMComponent.GetPathName());
+			const bool bValidKey = ISMComponentMap.Contains(ISMComponent);
+			checkf(bValidKey, TEXT("Failed to find ISMC in ISMComponentMap, path: %s"), *ISMComponent.GetPathName());
 			if (bValidKey)
 			{
-				const FStaticMeshInstanceVisualizationDescHandle StoredVisualizationDescHandle = ISMComponentMap.FindAndRemoveChecked(ISMComponentPathHash);
+				const FStaticMeshInstanceVisualizationDescHandle StoredVisualizationDescHandle = ISMComponentMap.FindAndRemoveChecked(ISMComponent);
 				ensure(StoredVisualizationDescHandle == VisualizationHandle);
 			}
 		
-			ISMCSharedData.Remove(ISMComponentPathHash);
-#if WITH_MASSGAMEPLAY_DEBUG
-			const FString ISMCPath = ISMComponent.GetPathName();
-			TArray<FString>& Paths = DebugHashToPathMap.FindChecked(ISMComponentPathHash);
-			ensure(Paths.Remove(ISMCPath));
-#endif // WITH_MASSGAMEPLAY_DEBUG
+			ISMCSharedData.Remove(ISMComponent);
 		}
 		
 		InstancedStaticMeshInfos[VisualizationHandle.ToIndex()].Reset();
@@ -238,6 +184,8 @@ void UMassVisualizationComponent::ConstructStaticMeshComponents()
 	AActor* ActorOwner = GetOwner();
 	check(ActorOwner);
 	
+	TArray<UInstancedStaticMeshComponent*> TransientISMCs;
+
 	UE_MT_SCOPED_WRITE_ACCESS(InstancedStaticMeshInfosDetector);
 	for (const FStaticMeshInstanceVisualizationDescHandle VisualDescHandle : InstancedSMComponentsRequiringConstructing)
 	{
@@ -262,17 +210,30 @@ void UMassVisualizationComponent::ConstructStaticMeshComponents()
 			UE_LOG(LogMassRepresentation, Error, TEXT("No associated meshes for this instanced static mesh type"));
 			continue;
 		}
+
+		TransientISMCs.Reset();
 		for (const FMassStaticMeshInstanceVisualizationMeshDesc& MeshDesc : Info.Desc.Meshes)
 		{
-			FMassISMCSharedData* SharedData = ISMCSharedData.Find(GetTypeHash(MeshDesc));
+			// MeshDescToISMCMap here lets us figure out whether for the given MeshDesc we need to create a new ISM component
+			// or a one has already been created in the past. Note that we only need this intermediate map for 
+			// FMassStaticMeshInstanceVisualizationMeshDesc that has been added to the system without specifying an
+			// ISM component to handle the instances (i.e. added via FindOrAddVisualDesc rather than AddVisualDescWithISMComponents).
+			// This is the only kind of FMassStaticMeshInstanceVisualizationMeshDesc were processing here. 
+			FISMCSharedDataKey& ISMCKey = MeshDescToISMCMap.FindChecked(GetTypeHash(MeshDesc));
+			FMassISMCSharedData* SharedData = ISMCSharedData.Find(ISMCKey);
 			UInstancedStaticMeshComponent* ISMC = SharedData ? SharedData->GetMutableISMComponent() : nullptr;
 
 			if (ISMC == nullptr)
 			{
-				ISMC = NewObject<UInstancedStaticMeshComponent>(ActorOwner, MeshDesc.ISMComponentClass);
+				ISMC = NewObject<UInstancedStaticMeshComponent>(ActorOwner, MeshDesc.ISMComponentClass);	
 				CA_ASSUME(ISMC);
 				REDIRECT_OBJECT_TO_VLOG(ISMC, this);
 
+				// note that ISMCKey is a reference, so the assignment below actually sets a value in MeshDescToISMCMap
+				// and all subsequent handling of a given MeshDesc configuration (i.e. containing same values) will 
+				// result in referring to the ISMC we just created.
+				ISMCKey = ISMC;
+
 				ISMC->SetStaticMesh(MeshDesc.Mesh);
 				for (int32 ElementIndex = 0; ElementIndex < MeshDesc.MaterialOverrides.Num(); ++ElementIndex)
 				{
@@ -292,31 +253,33 @@ void UMassVisualizationComponent::ConstructStaticMeshComponents()
 
 				if (SharedData == nullptr)
 				{
-					SharedData = &ISMCSharedData.Add(GetTypeHash(MeshDesc), FMassISMCSharedData(ISMC));
+					SharedData = &ISMCSharedData.Add(ISMC, FMassISMCSharedData(ISMC));
 				}
 				else
 				{
 					SharedData->SetISMComponent(*ISMC);
 				}
 
-				const uint32 ISMComponentPathHash = UE::Mass::Private::CalculateComponentHash(ISMC);
-				ensureMsgf(ISMComponentMap.Find(ISMComponentPathHash) == nullptr, TEXT("We've just created the ISMC that's being used here, so this check failing indicates hash-clash."));
-				ISMComponentMap.Add(ISMComponentPathHash, VisualDescHandle); 
+				ensureMsgf(ISMComponentMap.Find(ISMC) == nullptr, TEXT("We've just created the ISMC that's being used here, so this check failing indicates hash-clash."));
+				ISMComponentMap.Add(ISMC, VisualDescHandle); 
 			}
 
+			TransientISMCs.Add(ISMC);
+
 			check(SharedData);
 			Info.AddISMComponent(*SharedData);
 		}
 
 		// Build the LOD significance ranges
-		if (Info.LODSignificanceRanges.Num() == 0)
+		if (TransientISMCs.Num())
 		{
-			BuildLODSignificanceForInfo(Info);
+			check(Info.LODSignificanceRanges.Num() == 0);
+			BuildLODSignificanceForInfo(Info, TransientISMCs);
 		}
 	}
 }
 
-void UMassVisualizationComponent::BuildLODSignificanceForInfo(FMassInstancedStaticMeshInfo& Info, TConstArrayView<uint32> ForcedStaticMeshRefKeys)
+void UMassVisualizationComponent::BuildLODSignificanceForInfo(FMassInstancedStaticMeshInfo& Info, TConstArrayView<UInstancedStaticMeshComponent*> StaticMeshRefKeys)
 {
 	TArray<float> AllLODSignificances;
 	auto UniqueInsertOrdered = [&AllLODSignificances](const float Significance)
@@ -358,8 +321,9 @@ void UMassVisualizationComponent::BuildLODSignificanceForInfo(FMassInstancedStat
 				const bool bAddMeshInRange = (Range.MinSignificance >= MeshDesc.MinLODSignificance && Range.MinSignificance < MeshDesc.MaxLODSignificance);
 				if (bAddMeshInRange)
 				{
-					Range.StaticMeshRefs.Add(ForcedStaticMeshRefKeys.IsValidIndex(MeshIndex) && ForcedStaticMeshRefKeys[MeshIndex]
-						? ForcedStaticMeshRefKeys[MeshIndex] : GetTypeHash(MeshDesc));
+					checkf(StaticMeshRefKeys.IsValidIndex(MeshIndex) && StaticMeshRefKeys[MeshIndex]
+						, TEXT("We don't expect receiving null ISMCs at this point"));
+					Range.StaticMeshRefs.Add(StaticMeshRefKeys[MeshIndex]);
 				}
 			}
 		}
@@ -369,13 +333,7 @@ void UMassVisualizationComponent::BuildLODSignificanceForInfo(FMassInstancedStat
 void UMassVisualizationComponent::ClearAllVisualInstances()
 {
 	UE_MT_SCOPED_WRITE_ACCESS(InstancedStaticMeshInfosDetector);
-	for (FMassInstancedStaticMeshInfo& Info : InstancedStaticMeshInfos)
-	{
-		Info.ClearVisualInstance(ISMCSharedData);
-	}
-	InstancedStaticMeshInfos.Reset();
 	
-	// Pool should already be empty, got a problem if it's not
 	for (int32 SharedDataIndex = 0; SharedDataIndex < ISMCSharedData.Num(); ++SharedDataIndex)
 	{
 		if (UInstancedStaticMeshComponent* InstancedStaticMeshComponent = ISMCSharedData.GetAtIndex(SharedDataIndex).GetMutableISMComponent())
@@ -385,8 +343,10 @@ void UMassVisualizationComponent::ClearAllVisualInstances()
 		}
 	}
 
+	MeshDescToISMCMap.Reset();
 	ISMCSharedData.Reset();
 	InstancedSMComponentsRequiringConstructing.Reset();
+	InstancedStaticMeshInfos.Reset();
 }
 
 void UMassVisualizationComponent::DirtyVisuals()
@@ -670,36 +630,11 @@ void UMassVisualizationComponent::EndVisualChanges()
 	}
 }
 
-//---------------------------------------------------------------
-// FMassInstancedStaticMeshInfo
-//---------------------------------------------------------------
-
-void FMassInstancedStaticMeshInfo::ClearVisualInstance(FMassISMCSharedDataMap& ISMCSharedData)
-{
-	for (int i = 0; i < Desc.Meshes.Num(); i++)
-	{
-		const uint32 MeshDescHash = GetTypeHash(Desc.Meshes[i]);
-		FMassISMCSharedData* SharedData = ISMCSharedData.Find(MeshDescHash);
-		if (SharedData && SharedData->OnISMComponentReferenceReleased() == 0)
-		{
-			if (UInstancedStaticMeshComponent* ISMC = SharedData->GetMutableISMComponent())
-			{
-				ISMC->ClearInstances();
-				ISMC->DestroyComponent();
-			}
-			ISMCSharedData.Remove(MeshDescHash);
-		}
-	}
-
-	InstancedStaticMeshComponents.Reset();
-	LODSignificanceRanges.Reset();
-}
-
 //---------------------------------------------------------------
 // FMassLODSignificanceRange
 //---------------------------------------------------------------
 
-void FMassLODSignificanceRange::AddBatchedTransform(const FMassEntityHandle EntityHandle, const FTransform& Transform, const FTransform& PrevTransform, const TArray<uint32>& ExcludeStaticMeshRefs)
+void FMassLODSignificanceRange::AddBatchedTransform(const FMassEntityHandle EntityHandle, const FTransform& Transform, const FTransform& PrevTransform, TConstArrayView<FISMCSharedDataKey> ExcludeStaticMeshRefs)
 {
 	check(ISMCSharedDataPtr);
 	for (int32 StaticMeshIndex = 0; StaticMeshIndex < StaticMeshRefs.Num(); ++StaticMeshIndex)
@@ -718,7 +653,7 @@ void FMassLODSignificanceRange::AddBatchedTransform(const FMassEntityHandle Enti
 	}
 }
 
-void FMassLODSignificanceRange::AddBatchedCustomDataFloats(const TArray<float>& CustomFloats, const TArray<uint32>& ExcludeStaticMeshRefs)
+void FMassLODSignificanceRange::AddBatchedCustomDataFloats(const TArray<float>& CustomFloats, const TArray<FISMCSharedDataKey>& ExcludeStaticMeshRefs)
 {
 	check(ISMCSharedDataPtr);
 	for (int32 StaticMeshIndex = 0; StaticMeshIndex < StaticMeshRefs.Num(); ++StaticMeshIndex)
@@ -761,7 +696,7 @@ void FMassLODSignificanceRange::RemoveInstance(const FMassEntityHandle EntityHan
 	}
 }
 
-void FMassLODSignificanceRange::WriteCustomDataFloatsAtStartIndex(int32 StaticMeshIndex, const TArrayView<float>& CustomFloats, const int32 FloatsPerInstance, const int32 StartFloatIndex, const TArray<uint32>& ExcludeStaticMeshRefs)
+void FMassLODSignificanceRange::WriteCustomDataFloatsAtStartIndex(int32 StaticMeshIndex, const TArrayView<float>& CustomFloats, const int32 FloatsPerInstance, const int32 StartFloatIndex, const TArray<FISMCSharedDataKey>& ExcludeStaticMeshRefs)
 {
 	check(ISMCSharedDataPtr);
 	if (StaticMeshRefs.IsValidIndex(StaticMeshIndex))
@@ -789,17 +724,11 @@ void FMassLODSignificanceRange::WriteCustomDataFloatsAtStartIndex(int32 StaticMe
 //-----------------------------------------------------------------------------
 // DEPRECATED
 //-----------------------------------------------------------------------------
-void UMassVisualizationComponent::BuildLODSignificanceForInfo(FMassInstancedStaticMeshInfo& Info, const uint32 ForcedStaticMeshRefKeys)
-{
-	BuildLODSignificanceForInfo(Info, MakeArrayView(&ForcedStaticMeshRefKeys, 1));
-}
-
 void UMassVisualizationComponent::RemoveISMComponent(UInstancedStaticMeshComponent& ISMComponent)
 {
 	UE_MT_SCOPED_WRITE_ACCESS(InstancedStaticMeshInfosDetector);
 
-	const uint32 ISMComponentPathHash = UE::Mass::Private::CalculateComponentHash(&ISMComponent);
-	const FStaticMeshInstanceVisualizationDescHandle* VisualDescHandlePtr = ISMComponentMap.Find(ISMComponentPathHash);
+	const FStaticMeshInstanceVisualizationDescHandle* VisualDescHandlePtr = ISMComponentMap.Find(&ISMComponent);
 	if (VisualDescHandlePtr)
 	{
 		RemoveVisualDesc(*VisualDescHandlePtr);
diff --git a/Engine/Plugins/Runtime/MassGameplay/Source/MassRepresentation/Public/MassRepresentationTypes.h b/Engine/Plugins/Runtime/MassGameplay/Source/MassRepresentation/Public/MassRepresentationTypes.h
--- a/Engine/Plugins/Runtime/MassGameplay/Source/MassRepresentation/Public/MassRepresentationTypes.h
+++ b/Engine/Plugins/Runtime/MassGameplay/Source/MassRepresentation/Public/MassRepresentationTypes.h
@@ -7,14 +7,14 @@
 #include "Misc/MTAccessDetector.h"
 #include "InstanceDataTypes.h"
 #include "Experimental/Containers/RobinHoodHashTable.h"
-
+#include "UObject/ObjectKey.h"
+#include "Components/InstancedStaticMeshComponent.h"
 #include "MassRepresentationTypes.generated.h"
 
 class UMaterialInterface;
 class UStaticMesh;
 struct FMassLODSignificanceRange;
 class UMassVisualizationComponent;
-class UInstancedStaticMeshComponent;
 
 DECLARE_LOG_CATEGORY_EXTERN(LogMassRepresentation, Log, All);
 
@@ -23,6 +23,8 @@ namespace UE::Mass::ProcessorGroupNames
 	const FName Representation = FName(TEXT("Representation"));
 }
 
+using FISMCSharedDataKey = TObjectKey<UInstancedStaticMeshComponent>;
+
 UENUM()
 enum class EMassRepresentationType : uint8
 {
@@ -383,17 +385,17 @@ struct FMassISMCSharedDataMap
 		static constexpr bool bValueToCheck = true;
 	};
 
-	FMassISMCSharedData& GetAndMarkDirtyChecked(const uint32 Hash)
+	FMassISMCSharedData& GetAndMarkDirtyChecked(const FISMCSharedDataKey OwnerKey)
 	{
-		const int32 DataIndex = Map[Hash];
+		const int32 DataIndex = Map[OwnerKey];
 		DirtyData[DataIndex] = true;
 		return Data[DataIndex];
 	}
 
-	FMassISMCSharedData* GetAndMarkDirty(const uint32 Hash)
+	FMassISMCSharedData* GetAndMarkDirty(const FISMCSharedDataKey OwnerKey)
 	{
-		const int32* DataIndex = Map.Find(Hash);
-		if (ensureMsgf(DataIndex, TEXT("%hs Failed to find Hash %u"), __FUNCTION__, Hash))
+		const int32* DataIndex = Map.Find(OwnerKey);
+		if (ensureMsgf(DataIndex, TEXT("%hs Failed to find OwnerKey %u"), __FUNCTION__, *GetNameSafe(OwnerKey.ResolveObjectPtrEvenIfGarbage())))
 		{
 			DirtyData[*DataIndex] = true;
 			return &Data[*DataIndex];
@@ -402,28 +404,28 @@ struct FMassISMCSharedDataMap
 	}
 	
 	template<typename... TArgs>
-	FMassISMCSharedData& FindOrAdd(const uint32 Hash, TArgs&&... InNewInstanceArgs)
+	FMassISMCSharedData& FindOrAdd(const FISMCSharedDataKey OwnerKey, TArgs&&... InNewInstanceArgs)
 	{
-		const int32* DataIndex = Map.Find(Hash);
+		const int32* DataIndex = Map.Find(OwnerKey);
 		if (DataIndex == nullptr)
 		{
-			return Add(Hash, Forward<TArgs>(InNewInstanceArgs)...);
+			return Add(OwnerKey, Forward<TArgs>(InNewInstanceArgs)...);
 		}
 		check(Data.IsValidIndex(*DataIndex));
 		return Data[*DataIndex];
 	}
 
-	FMassISMCSharedData* Find(const uint32 Hash)
+	FMassISMCSharedData* Find(const FISMCSharedDataKey OwnerKey)
 	{
-		int32* DataIndex = Map.Find(Hash);
+		int32* DataIndex = Map.Find(OwnerKey);
 		return (DataIndex == nullptr || *DataIndex == INDEX_NONE) ? (FMassISMCSharedData*)nullptr : &Data[*DataIndex];
 	}
 
 	template<typename... TArgs>
-	FMassISMCSharedData& Add(const uint32 Hash, TArgs&&... InNewInstanceArgs)
+	FMassISMCSharedData& Add(const FISMCSharedDataKey OwnerKey, TArgs&&... InNewInstanceArgs)
 	{
 		const int32 DataIndex = FreeIndices.Num() ? FreeIndices.Pop() : Data.Num();
-		Map.Add(Hash, DataIndex);
+		Map.Add(OwnerKey, DataIndex);
 
 		if (DataIndex == Data.Num())
 		{
@@ -439,10 +441,10 @@ struct FMassISMCSharedDataMap
 		}
 	}
 
-	void Remove(const uint32 Hash)
+	void Remove(const FISMCSharedDataKey OwnerKey)
 	{
 		int32 DataIndex = INDEX_NONE;
-		if (ensure(Map.RemoveAndCopyValue(Hash, DataIndex)))
+		if (ensure(Map.RemoveAndCopyValue(OwnerKey, DataIndex)))
 		{
 			DirtyData[DataIndex] = false;
 			Data[DataIndex].Reset();
@@ -477,6 +479,11 @@ struct FMassISMCSharedDataMap
 		return DirtyData[DataIndex];
 	}
 
+	bool IsEmpty() const
+	{
+		return NumValid() == 0;
+	}
+
 	void Reset()
 	{
 		*this = FMassISMCSharedDataMap();
@@ -489,12 +496,28 @@ struct FMassISMCSharedDataMap
 
 protected:
 	TArray<FMassISMCSharedData> Data;
-	/** Mapping from Hash of data represented by FMassISMCSharedData to an index to Data */
-	TMap<uint32, int32> Map;
+	/** Mapping from Owner (as FObjectKey) of data represented by FMassISMCSharedData to an index to Data */
+	TMap<FISMCSharedDataKey, int32> Map;
 	/** Indicates whether corresponding Data entry has any instance work assigned to it (instance addition or removal) */
 	TBitArray<> DirtyData;
 	/** Indices to Data that are available for reuse */
 	TArray<int32> FreeIndices;
+
+public:
+	UE_DEPRECATED(5.5, "Deprecated. Using hashes in Mass Visualization is being phased out. Use FISMCSharedDataKey instead.")
+	FMassISMCSharedData& GetAndMarkDirtyChecked(const uint32 Hash);
+	UE_DEPRECATED(5.5, "Deprecated. Using hashes in Mass Visualization is being phased out. Use FISMCSharedDataKey instead.")
+	FMassISMCSharedData* GetAndMarkDirty(const uint32 Hash);	
+	template<typename... TArgs>
+	UE_DEPRECATED(5.5, "Deprecated. Using hashes in Mass Visualization is being phased out. Use FISMCSharedDataKey instead.")
+	FMassISMCSharedData& FindOrAdd(const uint32 Hash, TArgs&&... InNewInstanceArgs);
+	UE_DEPRECATED(5.5, "Deprecated. Using hashes in Mass Visualization is being phased out. Use FISMCSharedDataKey instead.")
+	FMassISMCSharedData* Find(const uint32 Hash);
+	template<typename... TArgs>
+	UE_DEPRECATED(5.5, "Deprecated. Using hashes in Mass Visualization is being phased out. Use FISMCSharedDataKey instead.")
+	FMassISMCSharedData& Add(const uint32 Hash, TArgs&&... InNewInstanceArgs);
+	UE_DEPRECATED(5.5, "Deprecated. Using hashes in Mass Visualization is being phased out. Use FISMCSharedDataKey instead.")
+	void Remove(const uint32 Hash);
 };
 
 
@@ -504,12 +527,12 @@ struct MASSREPRESENTATION_API FMassLODSignificanceRange
 	GENERATED_BODY()
 public:
 
-	void AddBatchedTransform(const FMassEntityHandle EntityHandle, const FTransform& Transform, const FTransform& PrevTransform, const TArray<uint32>& ExcludeStaticMeshRefs);
+	void AddBatchedTransform(const FMassEntityHandle EntityHandle, const FTransform& Transform, const FTransform& PrevTransform, TConstArrayView<FISMCSharedDataKey> ExcludeStaticMeshRefs);
 
 	// Adds the specified struct reinterpreted as custom floats to our custom data. Individual members of the specified struct should always fit into a float.
 	// When adding any custom data, the custom data must be added for every instance.
 	template<typename InCustomDataType>
-	void AddBatchedCustomData(InCustomDataType InCustomData, const TArray<uint32>& ExcludeStaticMeshRefs, int32 NumFloatsToPad = 0)
+	void AddBatchedCustomData(InCustomDataType InCustomData, const TArray<FISMCSharedDataKey>& ExcludeStaticMeshRefs, int32 NumFloatsToPad = 0)
 	{
 		check(ISMCSharedDataPtr);
 		static_assert((sizeof(InCustomDataType) % sizeof(float)) == 0, "AddBatchedCustomData: InCustomDataType should have a total size multiple of sizeof(float), and have members that fit in a float's boundaries");
@@ -529,22 +552,21 @@ public:
 		}
 	}
 
-	void AddBatchedCustomDataFloats(const TArray<float>& CustomFloats, const TArray<uint32>& ExcludeStaticMeshRefs);
+	void AddBatchedCustomDataFloats(const TArray<float>& CustomFloats, const TArray<FISMCSharedDataKey>& ExcludeStaticMeshRefs);
 
 	/** Single-instance version of AddBatchedCustomData when called to add entities (as opposed to modify existing ones).*/
 	void AddInstance(const FMassEntityHandle EntityHandle, const FTransform& Transform);
 
 	void RemoveInstance(const FMassEntityHandle EntityHandle);
 
-	void WriteCustomDataFloatsAtStartIndex(int32 StaticMeshIndex, const TArrayView<float>& CustomFloats, const int32 FloatsPerInstance, const int32 StartIndex, const TArray<uint32>& ExcludeStaticMeshRefs);
+	void WriteCustomDataFloatsAtStartIndex(int32 StaticMeshIndex, const TArrayView<float>& CustomFloats, const int32 FloatsPerInstance, const int32 StartIndex, const TArray<FISMCSharedDataKey>& ExcludeStaticMeshRefs);
 
 	/** LOD Significance range */
 	float MinSignificance;
 	float MaxSignificance;
 
 	/** The component handling these instances */
-	UPROPERTY(VisibleAnywhere, Category = "Mass/Debug")
-	TArray<uint32> StaticMeshRefs;
+	TArray<FISMCSharedDataKey> StaticMeshRefs;
 
 	FMassISMCSharedDataMap* ISMCSharedDataPtr = nullptr;
 
@@ -601,7 +623,7 @@ public:
 	{
 		if (FMassLODSignificanceRange* Range = GetLODSignificanceRange(LODSignificance))
 		{
-			Range->AddBatchedTransform(EntityHandle, Transform, PrevTransform, TArray<uint32>());
+			Range->AddBatchedTransform(EntityHandle, Transform, PrevTransform, {});
 			if(PrevLODSignificance >= 0.0f)
 			{
 				FMassLODSignificanceRange* PrevRange = GetLODSignificanceRange(PrevLODSignificance);
@@ -629,7 +651,7 @@ public:
 	{
 		if (FMassLODSignificanceRange* Range = GetLODSignificanceRange(LODSignificance))
 		{
-			Range->AddBatchedCustomData(InCustomData, TArray<uint32>(), NumFloatsToPad);
+			Range->AddBatchedCustomData(InCustomData, {}, NumFloatsToPad);
 			if(PrevLODSignificance >= 0.0f)
 			{
 				FMassLODSignificanceRange* PrevRange = GetLODSignificanceRange(PrevLODSignificance);
@@ -646,7 +668,7 @@ public:
 	{
 		if (FMassLODSignificanceRange* Range = GetLODSignificanceRange(LODSignificance))
 		{
-			Range->AddBatchedCustomDataFloats(CustomFloats, TArray<uint32>());
+			Range->AddBatchedCustomDataFloats(CustomFloats, {});
 			if(PrevLODSignificance >= 0.0f)
 			{
 				FMassLODSignificanceRange* PrevRange = GetLODSignificanceRange(PrevLODSignificance);
@@ -663,7 +685,7 @@ public:
 	{
 		if (FMassLODSignificanceRange* Range = GetLODSignificanceRange(LODSignificance))
 		{
-			Range->WriteCustomDataFloatsAtStartIndex(StaticMeshIndex, CustomFloats, FloatsPerInstance, FloatStartIndex, TArray<uint32>());
+			Range->WriteCustomDataFloatsAtStartIndex(StaticMeshIndex, CustomFloats, FloatsPerInstance, FloatStartIndex, {});
 			if(PrevLODSignificance >= 0.0f)
 			{
 				FMassLODSignificanceRange* PrevRange = GetLODSignificanceRange(PrevLODSignificance);
@@ -695,13 +717,13 @@ public:
 protected:
 
 	/** Destroy the visual instance */
-	void ClearVisualInstance(FMassISMCSharedDataMap& ISMCSharedData);
+	void ClearVisualInstance(UInstancedStaticMeshComponent& ISMComponent);
 
 	/** Information about this static mesh which will represent all instances */
 	UPROPERTY(VisibleAnywhere, Category = "Mass/Debug")
 	FStaticMeshInstanceVisualizationDesc Desc;
 
-	/** The component handling these instances */
+	/** The components handling these instances */
 	UPROPERTY(VisibleAnywhere, Category = "Mass/Debug")
 	TArray<TObjectPtr<UInstancedStaticMeshComponent>> InstancedStaticMeshComponents;
 
@@ -713,10 +735,12 @@ protected:
 	//-----------------------------------------------------------------------------
 	// DEPRECATED
 	//-----------------------------------------------------------------------------
+	UE_DEPRECATED(5.5, "Deprecated in flavor of the function taking the ISMComponent parameter. This version is not defunct")
+	void ClearVisualInstance(FMassISMCSharedDataMap& ISMCSharedData) {}
 public:
-	UE_DEPRECATED(5.4, "Deprecated in favor of new version taking FMassEntityHandle parameter instead of int32 to identify the entity. This deprecated function is now defunct.")
+	UE_DEPRECATED(5.4, "Deprecated in flavor of new version taking FMassEntityHandle parameter instead of int32 to identify the entity. This deprecated function is now defunct.")
 	void AddBatchedTransform(const int32 InstanceId, const FTransform& Transform, const FTransform& PrevTransform, const float LODSignificance, const float PrevLODSignificance = -1.0f) {}
-	UE_DEPRECATED(5.4, "Deprecated in favor of new version taking FMassEntityHandle parameter instead of int32 to identify the entity. This deprecated function is now defunct.")
+	UE_DEPRECATED(5.4, "Deprecated in flavor of new version taking FMassEntityHandle parameter instead of int32 to identify the entity. This deprecated function is now defunct.")
 	void RemoveInstance(const int32 InstanceId, const float LODSignificance) {}
 };
 
@@ -773,3 +797,21 @@ typedef FMassInstancedStaticMeshInfoArrayViewAccessDetector FMassInstancedStatic
 typedef TArrayView<FMassInstancedStaticMeshInfo> FMassInstancedStaticMeshInfoArrayView;
 
 #endif // ENABLE_MT_DETECTOR
+
+//-----------------------------------------------------------------------------
+// DEPRECATED
+//-----------------------------------------------------------------------------
+
+template<typename... TArgs>
+FMassISMCSharedData& FMassISMCSharedDataMap::FindOrAdd(const uint32 Hash, TArgs&&... InNewInstanceArgs)
+{
+	static FMassISMCSharedData Dummy;
+	return Dummy;
+}
+
+template<typename... TArgs>
+FMassISMCSharedData& FMassISMCSharedDataMap::Add(const uint32 Hash, TArgs&&... InNewInstanceArgs)
+{
+	static FMassISMCSharedData Dummy;
+	return Dummy;
+}
diff --git a/Engine/Plugins/Runtime/MassGameplay/Source/MassRepresentation/Public/MassVisualizationComponent.h b/Engine/Plugins/Runtime/MassGameplay/Source/MassRepresentation/Public/MassVisualizationComponent.h
--- a/Engine/Plugins/Runtime/MassGameplay/Source/MassRepresentation/Public/MassVisualizationComponent.h
+++ b/Engine/Plugins/Runtime/MassGameplay/Source/MassRepresentation/Public/MassVisualizationComponent.h
@@ -112,7 +112,7 @@ protected:
 	 * @param ForcedStaticMeshRefKeys if not empty will be used when adding individual FMassStaticMeshInstanceVisualizationMeshDesc
 	 *	instances to LOD significance ranges.
 	 */	
-	void BuildLODSignificanceForInfo(FMassInstancedStaticMeshInfo& Info, TConstArrayView<uint32> ForcedStaticMeshRefKeys = TConstArrayView<uint32>());
+	void BuildLODSignificanceForInfo(FMassInstancedStaticMeshInfo& Info, TConstArrayView<UInstancedStaticMeshComponent*> StaticMeshRefKeys);
 
 	/** Either adds an element to InstancedStaticMeshInfos or reuses an existing entry based on InstancedStaticMeshInfosFreeIndices*/
 	FStaticMeshInstanceVisualizationDescHandle AddInstancedStaticMeshInfo(const FStaticMeshInstanceVisualizationDesc& Desc);
@@ -125,18 +125,25 @@ protected:
 	/** Indices to InstancedStaticMeshInfos that have been released and can be reused */
 	TArray<FStaticMeshInstanceVisualizationDescHandle> InstancedStaticMeshInfosFreeIndices;
 
-	/** Mapping from ISMComponent object path hash to corresponding VisualDescHandle */
-	TMap<uint32, FStaticMeshInstanceVisualizationDescHandle> ISMComponentMap;
+	/** Mapping from ISMComponent (indicated by FISMCSharedDataKey) to corresponding VisualDescHandle */
+	TMap<FISMCSharedDataKey, FStaticMeshInstanceVisualizationDescHandle> ISMComponentMap;
 
 	FMassISMCSharedDataMap ISMCSharedData;
 
+	/** 
+	 * Mapping FMassStaticMeshInstanceVisualizationMeshDesc hash to FMassISMCSharedData entries for all FMassStaticMeshInstanceVisualizationMeshDesc
+	 * that didn't come with ISMC explicitly provided. Used only for initialization.
+	 * Note that FMassStaticMeshInstanceVisualizationMeshDesc that were added with ISMComponents provided directly
+	 * (via AddVisualDescWithISMComponents call) will never make it to this map.
+	 */
+	TMap<uint32, FISMCSharedDataKey> MeshDescToISMCMap;
+
 	/** Indicies to InstancedStaticMeshInfos that need their SMComponent constructed */
 	TArray<FStaticMeshInstanceVisualizationDescHandle> InstancedSMComponentsRequiringConstructing;
 
-	UE_DEPRECATED(5.4, "This flavor of BuildLODSignificanceForInfo is no longer supported. Use the other one instead.")
-	void BuildLODSignificanceForInfo(FMassInstancedStaticMeshInfo& Info, const uint32 ForcedStaticMeshRefKey);
+	UE_DEPRECATED(5.4, "This flavor of BuildLODSignificanceForInfo is no longer supported and is defunct.")
+	void BuildLODSignificanceForInfo(FMassInstancedStaticMeshInfo& Info, const uint32 ForcedStaticMeshRefKey){}
 
-#if WITH_MASSGAMEPLAY_DEBUG
-	TMap<uint32, TArray<FString>> DebugHashToPathMap;
-#endif // WITH_MASSGAMEPLAY_DEBUG
+	UE_DEPRECATED(5.5, "This flavor of BuildLODSignificanceForInfo is no longer supported and is defunct.")
+	void BuildLODSignificanceForInfo(FMassInstancedStaticMeshInfo& Info) {}
 };
