Fixed mistakenly inverted IsValid()

diff --git a/Engine/Plugins/Runtime/MassGameplay/Source/MassSpawner/Private/MassEntityTemplate.cpp b/Engine/Plugins/Runtime/MassGameplay/Source/MassSpawner/Private/MassEntityTemplate.cpp
--- a/Engine/Plugins/Runtime/MassGameplay/Source/MassSpawner/Private/MassEntityTemplate.cpp
+++ b/Engine/Plugins/Runtime/MassGameplay/Source/MassSpawner/Private/MassEntityTemplate.cpp
@@ -12,8 +12,8 @@
 //----------------------------------------------------------------------//
 FString FMassEntityTemplateID::ToString() const
 {
-	return IsValid() ? FString::Printf(TEXT("[Invalid]"))
-		: FString::Printf(TEXT("[%s:%d]"), *ConfigGuid.ToString(EGuidFormats::DigitsLower), FlavorHash);
+	return IsValid() ? FString::Printf(TEXT("[%s:%d]"), *ConfigGuid.ToString(EGuidFormats::DigitsLower), FlavorHash) 
+			: FString::Printf(TEXT("[Invalid]"));
 }
 
 //----------------------------------------------------------------------//
