Removed repeated include header files in MassSceneComponentLocationTranslator.cpp

diff --git a/Engine/Plugins/Runtime/MassGameplay/Source/MassActors/Private/Translators/MassSceneComponentLocationTranslator.cpp b/Engine/Plugins/Runtime/MassGameplay/Source/MassActors/Private/Translators/MassSceneComponentLocationTranslator.cpp
--- a/Engine/Plugins/Runtime/MassGameplay/Source/MassActors/Private/Translators/MassSceneComponentLocationTranslator.cpp
+++ b/Engine/Plugins/Runtime/MassGameplay/Source/MassActors/Private/Translators/MassSceneComponentLocationTranslator.cpp
@@ -3,8 +3,6 @@
 #include "Translators/MassSceneComponentLocationTranslator.h"
 #include "MassCommonTypes.h"
 #include "Components/SceneComponent.h"
-#include "MassEntityManager.h"
-#include "MassCommonTypes.h"
 #include "MassExecutionContext.h"
 
 //----------------------------------------------------------------------//
@@ -80,4 +78,4 @@ void UMassSceneComponentLocationToActorTranslator::Execute(FMassEntityManager& E
 				}
 			}
 		});
-}
\ No newline at end of file
+}
