Issue 460106: fixed PVS Studio warnings in MassEQSBlueprintLibrary.cpp

diff --git a/Engine/Plugins/Runtime/MassGameplay/Source/MassEQS/Private/MassEQSBlueprintLibrary.cpp b/Engine/Plugins/Runtime/MassGameplay/Source/MassEQS/Private/MassEQSBlueprintLibrary.cpp
--- a/Engine/Plugins/Runtime/MassGameplay/Source/MassEQS/Private/MassEQSBlueprintLibrary.cpp
+++ b/Engine/Plugins/Runtime/MassGameplay/Source/MassEQS/Private/MassEQSBlueprintLibrary.cpp
@@ -92,19 +92,22 @@ TArray<FMassEnvQueryEntityInfoBlueprintWrapper> UMassEQSBlueprintLibrary::GetEnv
 
 	TArray<FMassEnvQueryEntityInfoBlueprintWrapper> ResultInfo = {};
 
-	const FEnvQueryResult* QueryResult = QueryInstance->GetQueryResult();
-	const TSubclassOf<UEnvQueryItemType> ItemType = QueryResult->ItemType;
-	const EEnvQueryRunMode::Type RunMode = QueryInstance->GetRunMode();
-
-	if (QueryResult && (QueryResult->GetRawStatus() == EEnvQueryStatus::Success) && (ItemType.Get() != nullptr) && ItemType->IsChildOf(UEnvQueryItemType_MassEntityHandle::StaticClass()))
+	if (const FEnvQueryResult* QueryResult = QueryInstance->GetQueryResult())
 	{
-		if (RunMode != EEnvQueryRunMode::AllMatching)
-		{
-			ResultInfo.Add(GetItemAsEntityInfoBPWrapper(QueryResult, 0));
-		}
-		else
+		const TSubclassOf<UEnvQueryItemType> ItemType = QueryResult->ItemType;
+		const EEnvQueryRunMode::Type RunMode = QueryInstance->GetRunMode();
+
+		if ((QueryResult->GetRawStatus() == EEnvQueryStatus::Success) 
+			&& ItemType && ItemType->IsChildOf(UEnvQueryItemType_MassEntityHandle::StaticClass()))
 		{
-			GetAllAsEntityInfoBPWrappers(QueryResult, ResultInfo);
+			if (RunMode != EEnvQueryRunMode::AllMatching)
+			{
+				ResultInfo.Add(GetItemAsEntityInfoBPWrapper(QueryResult, 0));
+			}
+			else
+			{
+				GetAllAsEntityInfoBPWrappers(QueryResult, ResultInfo);
+			}
 		}
 	}
 
