Fixed MassEQS's MT Access Detector issues when trying to pop requests.
Also done a minor code massaging where I noticed things to improve.

diff --git a/Engine/Plugins/Runtime/MassGameplay/Source/MassEQS/Private/MassEQSSubsystem.cpp b/Engine/Plugins/Runtime/MassGameplay/Source/MassEQS/Private/MassEQSSubsystem.cpp
--- a/Engine/Plugins/Runtime/MassGameplay/Source/MassEQS/Private/MassEQSSubsystem.cpp
+++ b/Engine/Plugins/Runtime/MassGameplay/Source/MassEQS/Private/MassEQSSubsystem.cpp
@@ -15,10 +15,10 @@
 // RequestQueue
 //----------------------------------------------------------------------//
 
-FMassEQSRequestHandle UMassEQSSubsystem::PushRequest(const FEnvQueryInstance& QueryInstance, int32 RequestQueueIndex, TUniquePtr<FMassEQSRequestData>&& Request)
+FMassEQSRequestHandle UMassEQSSubsystem::PushRequest(const FEnvQueryInstance& QueryInstance, const int32 RequestQueueIndex, TUniquePtr<FMassEQSRequestData>&& Request)
 {
-	UE_MT_SCOPED_WRITE_ACCESS(RequestAccessDetector);
-
+	// It's a "read" since the detector only cares about the RequestQueues array, not its elements
+	UE_MT_SCOPED_READ_ACCESS(RequestAccessDetector);
 	check(RequestQueueIndex >= 0);
 
 	FMassEQSRequestHandle RequestHandle = HandleManager.GetNextHandle();
@@ -34,10 +34,10 @@ FMassEQSRequestHandle UMassEQSSubsystem::PushRequest(const FEnvQueryInstance& Qu
 }
 
 
-TUniquePtr<FMassEQSRequestData> UMassEQSSubsystem::PopRequest(int32 RequestQueueIndex)
+TUniquePtr<FMassEQSRequestData> UMassEQSSubsystem::PopRequest(const int32 RequestQueueIndex)
 {
-	UE_MT_SCOPED_WRITE_ACCESS(RequestAccessDetector);
-
+	// It's a "read" since the detector only cares about the RequestQueues array, not its elements
+	UE_MT_SCOPED_READ_ACCESS(RequestAccessDetector);
 	check(RequestQueueIndex >= 0);
 
 	TUniquePtr<FMassEQSRequestData> OutRequest = nullptr;
@@ -120,6 +120,7 @@ int32 UMassEQSSubsystem::GetRequestQueueIndex(TSubclassOf<UEnvQueryNode> Request
 		return *IndexPtr;
 	}
 
+	UE_MT_SCOPED_WRITE_ACCESS(RequestAccessDetector);
 	RequestQueues.Emplace();
 	int32 NewIndex = RequestQueues.Num() - 1;
 	
diff --git a/Engine/Plugins/Runtime/MassGameplay/Source/MassEQS/Private/MassEQSTypes.cpp b/Engine/Plugins/Runtime/MassGameplay/Source/MassEQS/Private/MassEQSTypes.cpp
--- a/Engine/Plugins/Runtime/MassGameplay/Source/MassEQS/Private/MassEQSTypes.cpp
+++ b/Engine/Plugins/Runtime/MassGameplay/Source/MassEQS/Private/MassEQSTypes.cpp
@@ -1,26 +1,33 @@
 // Copyright Epic Games, Inc. All Rights Reserved.
 
 #include "MassEQSTypes.h"
-
 #include "EnvironmentQuery/EnvQueryTypes.h"
 #include "MassEQSSubsystem.h"
 #include "MassEQS.h"
 
 
-// Creating a static representation for invalid handles
-FMassEQSRequestHandle FMassEQSRequestHandle::Invalid = FMassEQSRequestHandle(INDEX_NONE, 0);
-
+//-----------------------------------------------------------------------------
+// FMassEnvQueryEntityInfo
+//-----------------------------------------------------------------------------
 FMassEnvQueryEntityInfo::FMassEnvQueryEntityInfo(int32 Index, int32 SerialNumber, const FTransform& Transform)
 	: EntityHandle(Index, SerialNumber)
 	, CachedTransform(Transform)
 {
 }
 
+//-----------------------------------------------------------------------------
+// FMassEQSRequestHandle
+//-----------------------------------------------------------------------------
+FMassEQSRequestHandle FMassEQSRequestHandle::Invalid = FMassEQSRequestHandle(INDEX_NONE, 0);
+
 FString FMassEQSRequestHandle::ToString() const
 {
 	return FString::Printf(TEXT("[%d,%u]"), Index, SerialNumber);
 }
 
+//-----------------------------------------------------------------------------
+// IMassEQSRequestInterface::FMassEQSRequestHandler
+//-----------------------------------------------------------------------------
 void IMassEQSRequestInterface::FMassEQSRequestHandler::SendOrRecieveRequest(FEnvQueryInstance& QueryInstance, const IMassEQSRequestInterface& MassEQSRequestInterface)
 {
 #if WITH_EDITOR
diff --git a/Engine/Plugins/Runtime/MassGameplay/Source/MassEQS/Private/Processors/MassEnvQueryProcessorBase.cpp b/Engine/Plugins/Runtime/MassGameplay/Source/MassEQS/Private/Processors/MassEnvQueryProcessorBase.cpp
--- a/Engine/Plugins/Runtime/MassGameplay/Source/MassEQS/Private/Processors/MassEnvQueryProcessorBase.cpp
+++ b/Engine/Plugins/Runtime/MassGameplay/Source/MassEQS/Private/Processors/MassEnvQueryProcessorBase.cpp
@@ -7,11 +7,14 @@
 void UMassEnvQueryProcessorBase::Initialize(UObject& Owner)
 {
 	Super::Initialize(Owner);
-	UMassEQSSubsystem* MassEQSSubsystem = Owner.GetWorld()->GetSubsystem<UMassEQSSubsystem>();
-	check(MassEQSSubsystem)
 
 	if (CorrespondingRequestClass)
 	{
+		UWorld* World = Owner.GetWorld();
+		check(World);
+		UMassEQSSubsystem* MassEQSSubsystem = World->GetSubsystem<UMassEQSSubsystem>();
+		check(MassEQSSubsystem)
+	
 		CachedRequestQueryIndex = MassEQSSubsystem->GetRequestQueueIndex(CorrespondingRequestClass);
 	}
 }
\ No newline at end of file
diff --git a/Engine/Plugins/Runtime/MassGameplay/Source/MassEQS/Public/MassEQSSubsystem.h b/Engine/Plugins/Runtime/MassGameplay/Source/MassEQS/Public/MassEQSSubsystem.h
--- a/Engine/Plugins/Runtime/MassGameplay/Source/MassEQS/Public/MassEQSSubsystem.h
+++ b/Engine/Plugins/Runtime/MassGameplay/Source/MassEQS/Public/MassEQSSubsystem.h
@@ -24,10 +24,10 @@ class UMassEQSSubsystem : public UMassSubsystemBase
 
 public:
 	/** Push new Request into RequestQueue */
-	FMassEQSRequestHandle PushRequest(const FEnvQueryInstance& QueryInstance, int32 RequestQueueIndex, TUniquePtr<FMassEQSRequestData>&& Request);
+	FMassEQSRequestHandle PushRequest(const FEnvQueryInstance& QueryInstance, const int32 RequestQueueIndex, TUniquePtr<FMassEQSRequestData>&& Request);
 
 	/** Dequeue next Request off the RequestQueue, removing it from the queue, and transferring ownership of UniquePtr */
-	TUniquePtr<FMassEQSRequestData> PopRequest(int32 RequestQueueIndex);
+	TUniquePtr<FMassEQSRequestData> PopRequest(const int32 RequestQueueIndex);
 
 	/** Send results from a finished/processed request to be stored in this subsystem */
 	void SubmitResults(FMassEQSRequestHandle RequestHandle, TUniquePtr<FMassEQSRequestData>&& Result);
@@ -49,21 +49,25 @@ public:
 
 	virtual void Initialize(FSubsystemCollectionBase& Collection) override;
 	virtual void Deinitialize() override;
-protected:
 
+protected:
 	/**
 	 * Mapping from Class of requester, to index into RequestQueues.
 	 * Each Class that extends UMassEQSGenerator or UMassEQSTest 
 	 * can be a requester, this will store a Queue for each of these types.
 	 * @see PostInitialize()
 	 */
-	TMap<UClass*, int32> RequestQueueLookup = {};
+	TMap<TSubclassOf<UEnvQueryNode>, int32> RequestQueueLookup = {};
 	/**
 	 * Array of Queues holding each Request. 
 	 * TQueue does not support copy/assignment, which TMap requires.
 	 * This is the reason the TQueues are stored in a separate Array.
 	 */
 	TArray<FRequestQueue> RequestQueues = {};
+	/** 
+	 * The detector only guards the extension of RequestQueues array. Modifying specific elements is not guarded since 
+	 * those are a thread safe queues.
+	 */
 	UE_MT_DECLARE_RW_ACCESS_DETECTOR(RequestAccessDetector);
 
 	/** Holds results for finished Query */
