Added GetTypeHash to a couple of Mass shared fragments and TMassLODTickRateController

diff --git a/Engine/Plugins/Runtime/MassGameplay/Source/MassLOD/Public/MassLODTickRateController.h b/Engine/Plugins/Runtime/MassGameplay/Source/MassLOD/Public/MassLODTickRateController.h
--- a/Engine/Plugins/Runtime/MassGameplay/Source/MassLOD/Public/MassLODTickRateController.h
+++ b/Engine/Plugins/Runtime/MassGameplay/Source/MassLOD/Public/MassLODTickRateController.h
@@ -50,6 +50,11 @@ public:
 	template <typename TLODFragment, typename TVariableTickRateFragment>
 	bool UpdateTickRateFromLOD(FMassExecutionContext& Context, TConstArrayView<TLODFragment> LODList, TArrayView<TVariableTickRateFragment> TickRateList, const double Time);
 
+	friend FORCEINLINE uint32 GetTypeHash(const TMassLODTickRateController<TVariableTickChunkFragment, FLODLogic>& SharedFragmentInstance)
+	{
+		return HashCombineFast(GetTypeHash(SharedFragmentInstance.TickRates), uint32(SharedFragmentInstance.bShouldSpreadFirstUpdate));
+	}
+
 protected:
 
 	/** Tick rate for each LOD */
diff --git a/Engine/Plugins/Runtime/MassGameplay/Source/MassLOD/Public/MassSimulationLOD.h b/Engine/Plugins/Runtime/MassGameplay/Source/MassLOD/Public/MassSimulationLOD.h
--- a/Engine/Plugins/Runtime/MassGameplay/Source/MassLOD/Public/MassSimulationLOD.h
+++ b/Engine/Plugins/Runtime/MassGameplay/Source/MassLOD/Public/MassSimulationLOD.h
@@ -159,6 +159,15 @@ struct MASSLOD_API FMassSimulationVariableTickSharedFragment : public FMassShare
 		}
 		return true;
 	}
+
+	/**
+	 * Custom hash function for GetOrCreateConstSharedFragment to use. This hashes the parameters of LODTickRateController
+	 * which are the values indicating the uniqueness of a given shared fragment
+	 */
+	friend FORCEINLINE uint32 GetTypeHash(const FMassSimulationVariableTickSharedFragment& SharedFragmentInstance)
+	{
+		return GetTypeHash(SharedFragmentInstance.LODTickRateController);
+	}
 };
 
 UCLASS(meta = (DisplayName = "Simulation LOD"))
