Fixed an assert with read / write access issues with a fragment in Mass Simulation LOD Processor.

diff --git a/Engine/Plugins/Runtime/MassGameplay/Source/MassLOD/Private/MassSimulationLOD.cpp b/Engine/Plugins/Runtime/MassGameplay/Source/MassLOD/Private/MassSimulationLOD.cpp
--- a/Engine/Plugins/Runtime/MassGameplay/Source/MassLOD/Private/MassSimulationLOD.cpp
+++ b/Engine/Plugins/Runtime/MassGameplay/Source/MassLOD/Private/MassSimulationLOD.cpp
@@ -177,7 +177,7 @@ void UMassSimulationLODProcessor::Execute(FMassEntityManager& EntityManager, FMa
 		check(World);
 		EntityQuerySetLODTag.ForEachEntityChunk(EntityManager, Context, [](FMassExecutionContext& Context)
 		{
-			TConstArrayView<FMassSimulationLODFragment> SimulationLODFragments = Context.GetMutableFragmentView<FMassSimulationLODFragment>();
+			TConstArrayView<FMassSimulationLODFragment> SimulationLODFragments = Context.GetFragmentView<FMassSimulationLODFragment>();
 			const int32 NumEntities = Context.GetNumEntities();
 			for (int32 Index = 0; Index < NumEntities; ++Index)
 			{
