Massaged Mass actor spawning code to take into consideration that spawning might fail (used to not take that into account)

diff --git a/Engine/Plugins/Runtime/MassGameplay/Source/MassActors/Public/MassActorSpawnerSubsystem.h b/Engine/Plugins/Runtime/MassGameplay/Source/MassActors/Public/MassActorSpawnerSubsystem.h
--- a/Engine/Plugins/Runtime/MassGameplay/Source/MassActors/Public/MassActorSpawnerSubsystem.h
+++ b/Engine/Plugins/Runtime/MassGameplay/Source/MassActors/Public/MassActorSpawnerSubsystem.h
@@ -77,7 +77,10 @@ public:
 	/** Delegate that will be called just before the spawning an actor, giving the chance to the processor to prepare it */
 	FMassActorPreSpawnDelegate ActorPreSpawnDelegate;
 
-	/** Delegate that will be called once the spawning is done */
+	/** 
+	 * Delegate that will be called once the spawning is done. Note that it will be called regardless of whether 
+	 * the spawning was successful. 
+	 */
 	FMassActorPostSpawnDelegate ActorPostSpawnDelegate;
 
 	/** The current status of the spawn request */
@@ -165,7 +168,7 @@ public:
 	 * @param SpawnRequestHandle [IN/OUT] the spawn request handle to remove
 	 * @return true if successfully removed the request
 	 */
-	 bool RemoveActorSpawnRequest(FMassActorSpawnRequestHandle& SpawnRequestHandle);
+	bool RemoveActorSpawnRequest(FMassActorSpawnRequestHandle& SpawnRequestHandle);
 
 	/**
 	 * Returns the stored spawn request from the handle, useful to update the transform
diff --git a/Engine/Plugins/Runtime/MassGameplay/Source/MassRepresentation/Private/MassRepresentationActorManagement.cpp b/Engine/Plugins/Runtime/MassGameplay/Source/MassRepresentation/Private/MassRepresentationActorManagement.cpp
--- a/Engine/Plugins/Runtime/MassGameplay/Source/MassRepresentation/Private/MassRepresentationActorManagement.cpp
+++ b/Engine/Plugins/Runtime/MassGameplay/Source/MassRepresentation/Private/MassRepresentationActorManagement.cpp
@@ -86,21 +86,27 @@ EMassActorSpawnRequestAction UMassRepresentationActorManagement::OnPostActorSpaw
 	, FConstStructView SpawnRequest, TSharedRef<FMassEntityManager> EntityManager) const
 {
 	const FMassActorSpawnRequest& MassActorSpawnRequest = SpawnRequest.Get<const FMassActorSpawnRequest>();
-	checkf(MassActorSpawnRequest.SpawnedActor, TEXT("Expecting valid spawned actor"));
-
-	// Might be already done if the actor has a MassAgentComponent via the callback OnMassAgentComponentEntityAssociated on the MassRepresentationSubsystem
-	FMassActorFragment& ActorInfo = EntityManager->GetFragmentDataChecked<FMassActorFragment>(MassActorSpawnRequest.MassAgent);
-	if (ActorInfo.IsValid())
-	{
-		// If already set, make sure it is pointing to the same actor.
-		checkf(ActorInfo.Get() == MassActorSpawnRequest.SpawnedActor, TEXT("Expecting the pointer to the spawned actor in the actor fragment"));
-	}
-	else
+	
+	if (MassActorSpawnRequest.SpawnedActor)
 	{
-		ActorInfo.SetAndUpdateHandleMap(MassActorSpawnRequest.MassAgent, MassActorSpawnRequest.SpawnedActor, true/*bIsOwnedByMass*/);
+		// Might be already done if the actor has a MassAgentComponent via the callback OnMassAgentComponentEntityAssociated on the MassRepresentationSubsystem
+		FMassActorFragment& ActorInfo = EntityManager->GetFragmentDataChecked<FMassActorFragment>(MassActorSpawnRequest.MassAgent);
+		if (ActorInfo.IsValid())
+		{
+			// If already set, make sure it is pointing to the same actor.
+			checkf(ActorInfo.Get() == MassActorSpawnRequest.SpawnedActor, TEXT("Expecting the pointer to the spawned actor in the actor fragment"));
+		}
+		else
+		{
+			ActorInfo.SetAndUpdateHandleMap(MassActorSpawnRequest.MassAgent, MassActorSpawnRequest.SpawnedActor, true/*bIsOwnedByMass*/);
+		}
+
+		// by default we keep the spawn request since there's code in other places that will consume this request by 
+		// calling RemoveActorSpawnRequest (like UMassRepresentationSubsystem::GetOrSpawnActorFromTemplate)
+		return EMassActorSpawnRequestAction::Keep;
 	}
 
-	return EMassActorSpawnRequestAction::Keep;
+	return EMassActorSpawnRequestAction::Remove;
 }
 
 void UMassRepresentationActorManagement::ReleaseAnyActorOrCancelAnySpawning(FMassEntityManager& EntityManager, const FMassEntityHandle MassAgent)
