PR #11982: Export FMassSpawnedEntityType::GetEntityConfig methods

diff --git a/Engine/Plugins/Runtime/MassGameplay/Source/MassSpawner/Public/MassSpawnerTypes.h b/Engine/Plugins/Runtime/MassGameplay/Source/MassSpawner/Public/MassSpawnerTypes.h
--- a/Engine/Plugins/Runtime/MassGameplay/Source/MassSpawner/Public/MassSpawnerTypes.h
+++ b/Engine/Plugins/Runtime/MassGameplay/Source/MassSpawner/Public/MassSpawnerTypes.h
@@ -66,8 +66,8 @@ struct FMassSpawnedEntityType
 		EntityConfig.ResetWeakPtr();
 	}
 
-	const UMassEntityConfigAsset* GetEntityConfig() const;
-	UMassEntityConfigAsset* GetEntityConfig();
+	MASSSPAWNER_API const UMassEntityConfigAsset* GetEntityConfig() const;
+	MASSSPAWNER_API UMassEntityConfigAsset* GetEntityConfig();
 
 private:
 	UPROPERTY(Transient)
