[Mass] fixed some processor execution flags to not run in pure Editor without associated world

diff --git a/Engine/Plugins/Runtime/MassGameplay/Source/MassActors/Private/Translators/MassSceneComponentLocationTranslator.cpp b/Engine/Plugins/Runtime/MassGameplay/Source/MassActors/Private/Translators/MassSceneComponentLocationTranslator.cpp
--- a/Engine/Plugins/Runtime/MassGameplay/Source/MassActors/Private/Translators/MassSceneComponentLocationTranslator.cpp
+++ b/Engine/Plugins/Runtime/MassGameplay/Source/MassActors/Private/Translators/MassSceneComponentLocationTranslator.cpp
@@ -11,7 +11,7 @@
 UMassSceneComponentLocationToMassTranslator::UMassSceneComponentLocationToMassTranslator()
 	: EntityQuery(*this)
 {
-	ExecutionFlags = (int32)EProcessorExecutionFlags::All;
+	ExecutionFlags = (int32)EProcessorExecutionFlags::AllNetModes;
 	ExecutionOrder.ExecuteInGroup = UE::Mass::ProcessorGroupNames::SyncWorldToMass;
 	RequiredTags.Add<FMassSceneComponentLocationCopyToMassTag>();
 }
@@ -47,7 +47,7 @@ void UMassSceneComponentLocationToMassTranslator::Execute(FMassEntityManager& En
 UMassSceneComponentLocationToActorTranslator::UMassSceneComponentLocationToActorTranslator()
 	: EntityQuery(*this)
 {
-	ExecutionFlags = (int32)EProcessorExecutionFlags::All;
+	ExecutionFlags = (int32)EProcessorExecutionFlags::AllNetModes;
 	ExecutionOrder.ExecuteInGroup = UE::Mass::ProcessorGroupNames::UpdateWorldFromMass;
 	ExecutionOrder.ExecuteAfter.Add(UE::Mass::ProcessorGroupNames::Movement);
 	RequiredTags.Add<FMassSceneComponentLocationCopyToActorTag>();
diff --git a/Engine/Plugins/Runtime/MassGameplay/Source/MassRepresentation/Private/MassStationaryISMRepresentationFragmentDestructor.cpp b/Engine/Plugins/Runtime/MassGameplay/Source/MassRepresentation/Private/MassStationaryISMRepresentationFragmentDestructor.cpp
--- a/Engine/Plugins/Runtime/MassGameplay/Source/MassRepresentation/Private/MassStationaryISMRepresentationFragmentDestructor.cpp
+++ b/Engine/Plugins/Runtime/MassGameplay/Source/MassRepresentation/Private/MassStationaryISMRepresentationFragmentDestructor.cpp
@@ -13,7 +13,7 @@ UMassStationaryISMRepresentationFragmentDestructor::UMassStationaryISMRepresenta
 {
 	ObservedType = FMassRepresentationFragment::StaticStruct();
 	Operation = EMassObservedOperation::Remove;
-	ExecutionFlags = (int32)EProcessorExecutionFlags::All;
+	ExecutionFlags = (int32)EProcessorExecutionFlags::AllWorldModes;
 	bRequiresGameThreadExecution = true; // not sure about this
 }
 
