[InstancedActors] An attempt to handle a weirdly failed spawning case and gather more information while at it. The core reason is unknown and the gathered information should be helpful in finding the core of the problem.

diff --git a/Engine/Plugins/Runtime/MassGameplay/Source/MassActors/Private/MassActorSubsystem.cpp b/Engine/Plugins/Runtime/MassGameplay/Source/MassActors/Private/MassActorSubsystem.cpp
--- a/Engine/Plugins/Runtime/MassGameplay/Source/MassActors/Private/MassActorSubsystem.cpp
+++ b/Engine/Plugins/Runtime/MassGameplay/Source/MassActors/Private/MassActorSubsystem.cpp
@@ -17,11 +17,14 @@ void FMassActorFragment::SetAndUpdateHandleMap(const FMassEntityHandle MassAgent
 {
 	SetNoHandleMapUpdate(MassAgent, InActor, bInIsOwnedByMass);
 
-	UWorld* World = InActor->GetWorld();
-	check(World);
-	if (UMassActorSubsystem* MassActorSubsystem = UWorld::GetSubsystem<UMassActorSubsystem>(World))
+	if (InActor)
 	{
-		MassActorSubsystem->SetHandleForActor(InActor, MassAgent);
+		UWorld* World = InActor->GetWorld();
+		check(World);
+		if (UMassActorSubsystem* MassActorSubsystem = UWorld::GetSubsystem<UMassActorSubsystem>(World))
+		{
+			MassActorSubsystem->SetHandleForActor(InActor, MassAgent);
+		}
 	}
 }
 
