Fixed a MassRepresentation crash caused by a Visualization Trait configured to use StaticMeshInstance for visualization and not having any actual meshes configured.

To avoid duplicating code and to maintain consistency added FStaticMeshInstanceVisualizationDesc::IsValid that checks if there are any non-null meshes configured.

diff --git a/Engine/Plugins/Runtime/MassGameplay/Source/MassRepresentation/Private/MassRepresentationTypes.cpp b/Engine/Plugins/Runtime/MassGameplay/Source/MassRepresentation/Private/MassRepresentationTypes.cpp
--- a/Engine/Plugins/Runtime/MassGameplay/Source/MassRepresentation/Private/MassRepresentationTypes.cpp
+++ b/Engine/Plugins/Runtime/MassGameplay/Source/MassRepresentation/Private/MassRepresentationTypes.cpp
@@ -107,6 +107,21 @@ FMassStaticMeshInstanceVisualizationMeshDesc::FMassStaticMeshInstanceVisualizati
 	ISMComponentClass = UInstancedStaticMeshComponent::StaticClass();
 }
 
+//-----------------------------------------------------------------------------
+// FStaticMeshInstanceVisualizationDesc
+//-----------------------------------------------------------------------------
+bool FStaticMeshInstanceVisualizationDesc::IsValid() const
+{
+	for (const FMassStaticMeshInstanceVisualizationMeshDesc& MeshDesc : Meshes)
+	{
+		if (MeshDesc.Mesh && MeshDesc.ISMComponentClass)
+		{
+			return true;
+		}
+	}
+	return false;
+}
+
 //-----------------------------------------------------------------------------
 // FMassISMCSharedDataMap - DEPRECATED
 //-----------------------------------------------------------------------------
diff --git a/Engine/Plugins/Runtime/MassGameplay/Source/MassRepresentation/Private/MassVisualizationComponent.cpp b/Engine/Plugins/Runtime/MassGameplay/Source/MassRepresentation/Private/MassVisualizationComponent.cpp
--- a/Engine/Plugins/Runtime/MassGameplay/Source/MassRepresentation/Private/MassVisualizationComponent.cpp
+++ b/Engine/Plugins/Runtime/MassGameplay/Source/MassRepresentation/Private/MassVisualizationComponent.cpp
@@ -70,24 +70,21 @@ FStaticMeshInstanceVisualizationDescHandle UMassVisualizationComponent::FindOrAd
 	FStaticMeshInstanceVisualizationDescHandle VisualDescHandle(InstancedStaticMeshInfos.IndexOfByPredicate([&Desc](const FMassInstancedStaticMeshInfo& Info) { return Info.GetDesc() == Desc; }));
 	if (!VisualDescHandle.IsValid())
 	{
-		bool bValidDescription = false;
-
-		for (const FMassStaticMeshInstanceVisualizationMeshDesc& MeshDesc : Desc.Meshes)
+		if (Desc.IsValid())
 		{
-			if (MeshDesc.Mesh && MeshDesc.ISMComponentClass)
+			for (const FMassStaticMeshInstanceVisualizationMeshDesc& MeshDesc : Desc.Meshes)
 			{
-				// if we've already encountered MeshDesc in the past MeshDescToISMCMap already contains information
-				// about actual ISMC used to represent it, and at the same time indicates the ISMCSharedData data
-				// tied to it. Regardless we need to process all MeshDesc instances here so that we have all the 
-				// data ready when InstancedSMComponentsRequiringConstructing gets processed next time
-				// UMassVisualizationComponent::ConstructStaticMeshComponents gets called.
-				MeshDescToISMCMap.FindOrAdd(GetTypeHash(MeshDesc), FISMCSharedDataKey());
-				bValidDescription = true;
+				if (MeshDesc.Mesh && MeshDesc.ISMComponentClass)
+				{
+					// if we've already encountered MeshDesc in the past MeshDescToISMCMap already contains information
+					// about actual ISMC used to represent it, and at the same time indicates the ISMCSharedData data
+					// tied to it. Regardless we need to process all MeshDesc instances here so that we have all the 
+					// data ready when InstancedSMComponentsRequiringConstructing gets processed next time
+					// UMassVisualizationComponent::ConstructStaticMeshComponents gets called.
+					MeshDescToISMCMap.FindOrAdd(GetTypeHash(MeshDesc), FISMCSharedDataKey());
+				}
 			}
-		}
 
-		if (bValidDescription)
-		{
 			VisualDescHandle = AddInstancedStaticMeshInfo(Desc);
 			check(VisualDescHandle.IsValid());
 
diff --git a/Engine/Plugins/Runtime/MassGameplay/Source/MassRepresentation/Private/MassVisualizationTrait.cpp b/Engine/Plugins/Runtime/MassGameplay/Source/MassRepresentation/Private/MassVisualizationTrait.cpp
--- a/Engine/Plugins/Runtime/MassGameplay/Source/MassRepresentation/Private/MassVisualizationTrait.cpp
+++ b/Engine/Plugins/Runtime/MassGameplay/Source/MassRepresentation/Private/MassVisualizationTrait.cpp
@@ -12,6 +12,13 @@
 #include "MassEntityUtils.h"
 #include "MassVisualizationLODProcessor.h"
 #include "MassRepresentationProcessor.h"
+#if WITH_EDITOR
+#include "Logging/MessageLog.h"
+#include "Editor.h"
+#include "UObject/UnrealType.h"
+#endif
+
+#define LOCTEXT_NAMESPACE "Mass"
 
 
 UMassVisualizationTrait::UMassVisualizationTrait()
@@ -78,18 +85,39 @@ void UMassVisualizationTrait::BuildTemplate(FMassEntityTemplateBuildContext& Bui
 	{
 		UE_LOG(LogMassRepresentation, Error, TEXT("Expecting a valid class for the representation actor management"));
 	}
-	FConstSharedStruct ParamsFragment = EntityManager.GetOrCreateConstSharedFragment(Params);
-	ParamsFragment.Get<const FMassRepresentationParameters>().ComputeCachedValues();
-	BuildContext.AddConstSharedFragment(ParamsFragment);
 
 	FMassRepresentationFragment& RepresentationFragment = BuildContext.AddFragment_GetRef<FMassRepresentationFragment>();
-	if (bRegisterStaticMeshDesc)
-	{
-		RepresentationFragment.StaticMeshDescHandle = RepresentationSubsystem->FindOrAddStaticMeshDesc(StaticMeshInstanceDesc);
-	}
 	RepresentationFragment.HighResTemplateActorIndex = HighResTemplateActor.Get() ? RepresentationSubsystem->FindOrAddTemplateActor(HighResTemplateActor.Get()) : INDEX_NONE;
 	RepresentationFragment.LowResTemplateActorIndex = LowResTemplateActor.Get() ? RepresentationSubsystem->FindOrAddTemplateActor(LowResTemplateActor.Get()) : INDEX_NONE;
 
+	bool bStaticMeshDescriptionValid = StaticMeshInstanceDesc.IsValid();
+	if (bStaticMeshDescriptionValid)
+	{
+		if (bRegisterStaticMeshDesc)
+		{
+			RepresentationFragment.StaticMeshDescHandle = RepresentationSubsystem->FindOrAddStaticMeshDesc(StaticMeshInstanceDesc);
+			ensureMsgf(RepresentationFragment.StaticMeshDescHandle.IsValid()
+				, TEXT("Expected to get a valid StaticMeshDescHandle since we already checked that StaticMeshInstanceDesc is valid"));
+			// if the unexpected happens and StaticMeshDescHandle is not valid we're going to treat it as if StaticMeshInstanceDesc
+			// was not valid in the first place and handle it accordingly in a moment
+			bStaticMeshDescriptionValid = RepresentationFragment.StaticMeshDescHandle.IsValid();
+		}
+	}
+
+	FConstSharedStruct ParamsFragment;
+	if (bStaticMeshDescriptionValid)
+	{
+		ParamsFragment = EntityManager.GetOrCreateConstSharedFragment(Params);
+	}
+	else
+	{
+		FMassRepresentationParameters ParamsCopy = Params;
+		SanitizeParams(ParamsCopy, bStaticMeshDescriptionValid);
+		ParamsFragment = EntityManager.GetOrCreateConstSharedFragment(ParamsCopy);
+	}
+	ParamsFragment.Get<const FMassRepresentationParameters>().ComputeCachedValues();
+	BuildContext.AddConstSharedFragment(ParamsFragment);
+
 	FConstSharedStruct LODParamsFragment = EntityManager.GetOrCreateConstSharedFragment(LODParams);
 	BuildContext.AddConstSharedFragment(LODParamsFragment);
 
@@ -105,4 +133,82 @@ void UMassVisualizationTrait::BuildTemplate(FMassEntityTemplateBuildContext& Bui
 	BuildContext.AddTag<FMassVisualizationProcessorTag>();
 }
 
+void UMassVisualizationTrait::SanitizeParams(FMassRepresentationParameters& InOutParams, const bool bStaticMeshDeterminedInvalid) const
+{
+	if (bStaticMeshDeterminedInvalid || (StaticMeshInstanceDesc.IsValid() == false))
+	{
+		for (int32 LODIndex = 0; LODIndex < EMassLOD::Max; ++LODIndex)
+		{
+			if (InOutParams.LODRepresentation[LODIndex] == EMassRepresentationType::StaticMeshInstance)
+			{
+				InOutParams.LODRepresentation[LODIndex] = EMassRepresentationType::None;
+			}
+		}
+	}
+}
+
+void UMassVisualizationTrait::Serialize(FArchive& Ar)
+{
+	Super::Serialize(Ar);
+
+#if WITH_EDITOR
+	if (GEditor && (Ar.IsLoading() || Ar.IsSaving()))
+	{
+		ValidateParams();
+	}
+#endif // WITH_EDITOR
+}
+
+#if WITH_EDITOR
+void UMassVisualizationTrait::ValidateParams() const
+{
+	// the SM config provided is not valid. We need to check if EMassRepresentationType::StaticMeshInstance
+	// is being used as any of the LODRepresentations. If so then we need to clear those out and report an error
+	if (StaticMeshInstanceDesc.IsValid() == false)
+	{
+		for (int32 LODIndex = 0; LODIndex < EMassLOD::Max; ++LODIndex)
+		{
+			if (Params.LODRepresentation[LODIndex] == EMassRepresentationType::StaticMeshInstance)
+			{
+#if WITH_UNREAL_DEVELOPER_TOOLS && WITH_EDITOR
+				if (GEditor)
+				{
+					static const FText ErrorMessage(LOCTEXT("VisualizationTraitMissingSM", "Trait using StaticMeshInstance representation type but no meshes are configured."));
+					static const FText InfoMessage(LOCTEXT("SeeLogForDetails", "See the log for details."));
+
+					FMessageLog EditorErrors("MassEntity");
+					EditorErrors.Error(ErrorMessage);
+					EditorErrors.Notify(ErrorMessage);
+					EditorErrors.Info(InfoMessage);
+				}
+#endif // WITH_UNREAL_DEVELOPER_TOOLS
+
+				UE_LOG(LogMassRepresentation, Error, TEXT("Trait %s is using StaticMeshInstance representation type for "
+					"LODRepresentation[%s] while the trait's StaticMeshInstanceDesc is not valid (has no Meshes). Entities "
+					"won't be visible at this LOD level.")
+					, *GetPathName(), *UEnum::GetValueAsString(EMassLOD::Type(LODIndex)));
+			}
+		}
+	}
+}
+
+void UMassVisualizationTrait::PostEditChangeProperty(FPropertyChangedEvent& PropertyChangedEvent)
+{
+	static const FName ParamsName = GET_MEMBER_NAME_CHECKED(UMassVisualizationTrait, Params);
+	static const FName StaticMeshDescriptionName = GET_MEMBER_NAME_CHECKED(UMassVisualizationTrait, StaticMeshInstanceDesc);
+
+	Super::PostEditChangeProperty(PropertyChangedEvent);
+
+	if (PropertyChangedEvent.MemberProperty)
+	{
+		const FName PropName = PropertyChangedEvent.MemberProperty->GetFName();
+		if (PropName == ParamsName || PropName == StaticMeshDescriptionName)
+		{
+			ValidateParams();
+		}
+	}
+}
+
+#endif // WITH_EDITOR
 
+#undef LOCTEXT_NAMESPACE 
diff --git a/Engine/Plugins/Runtime/MassGameplay/Source/MassRepresentation/Public/MassRepresentationTypes.h b/Engine/Plugins/Runtime/MassGameplay/Source/MassRepresentation/Public/MassRepresentationTypes.h
--- a/Engine/Plugins/Runtime/MassGameplay/Source/MassRepresentation/Public/MassRepresentationTypes.h
+++ b/Engine/Plugins/Runtime/MassGameplay/Source/MassRepresentation/Public/MassRepresentationTypes.h
@@ -143,6 +143,9 @@ struct FStaticMeshInstanceVisualizationDesc : public FTableRowBase
 	{
 		new(this)FStaticMeshInstanceVisualizationDesc();
 	}
+
+	/** @return whether any of descriptions in Meshes is valid. This implies that empty Meshes will be treated as not valid.*/
+	bool IsValid() const;
 };
 
 /** Handle for FStaticMeshInstanceVisualizationDesc's registered with UMassRepresentationSubsystem */
diff --git a/Engine/Plugins/Runtime/MassGameplay/Source/MassRepresentation/Public/MassVisualizationTrait.h b/Engine/Plugins/Runtime/MassGameplay/Source/MassRepresentation/Public/MassVisualizationTrait.h
--- a/Engine/Plugins/Runtime/MassGameplay/Source/MassRepresentation/Public/MassVisualizationTrait.h
+++ b/Engine/Plugins/Runtime/MassGameplay/Source/MassRepresentation/Public/MassVisualizationTrait.h
@@ -21,8 +21,6 @@ class MASSREPRESENTATION_API UMassVisualizationTrait : public UMassEntityTraitBa
 public:
 	UMassVisualizationTrait();
 
-	virtual void BuildTemplate(FMassEntityTemplateBuildContext& BuildContext, const UWorld& World) const override;
-
 	/** Instanced static mesh information for this agent */
 	UPROPERTY(EditAnywhere, Category = "Mass|Visual")
 	mutable FStaticMeshInstanceVisualizationDesc StaticMeshInstanceDesc;
@@ -65,4 +63,20 @@ protected:
 	 * data that will never be used.
 	 */
 	bool bRegisterStaticMeshDesc = true;
+
+	virtual void BuildTemplate(FMassEntityTemplateBuildContext& BuildContext, const UWorld& World) const override;
+
+	/** 
+	 * Tests whether StaticMeshInstanceDesc is valid and if not cleans up InOutParamss of EMassRepresentationType::StaticMeshInstance
+	 * occurrences. 
+	 * @param bStaticMeshDeterminedInvalid if StaticMeshInstanceDesc has already been determined invalid then bStaticMeshDeterminedInvalid
+	 *	can be set to `true` to skip the redundant check.
+	 */
+	void SanitizeParams(FMassRepresentationParameters& InOutParams, const bool bStaticMeshDeterminedInvalid = false) const;
+
+	virtual void Serialize(FArchive& Ar) override;
+#if WITH_EDITOR
+	virtual void PostEditChangeProperty(struct FPropertyChangedEvent& PropertyChangedEvent) override;
+	void ValidateParams() const;
+#endif // WITH_EDITOR
 };
