Modified MassVisualizationTrait::ValidateParams to never report issues with CDOs - some are initially misconfigured on purpose.

diff --git a/Engine/Plugins/Runtime/MassGameplay/Source/MassRepresentation/Private/MassVisualizationTrait.cpp b/Engine/Plugins/Runtime/MassGameplay/Source/MassRepresentation/Private/MassVisualizationTrait.cpp
--- a/Engine/Plugins/Runtime/MassGameplay/Source/MassRepresentation/Private/MassVisualizationTrait.cpp
+++ b/Engine/Plugins/Runtime/MassGameplay/Source/MassRepresentation/Private/MassVisualizationTrait.cpp
@@ -162,6 +162,12 @@ void UMassVisualizationTrait::Serialize(FArchive& Ar)
 #if WITH_EDITOR
 void UMassVisualizationTrait::ValidateParams() const
 {
+	// if this test is called on any of the CDOs we don't care, we're never going to utilize those in practice.
+	if (HasAnyFlags(RF_ClassDefaultObject))
+	{
+		return;
+	}
+
 	// the SM config provided is not valid. We need to check if EMassRepresentationType::StaticMeshInstance
 	// is being used as any of the LODRepresentations. If so then we need to clear those out and report an error
 	if (StaticMeshInstanceDesc.IsValid() == false)
