[Mass] Added missing const shared fragments getters to FMassEntityTemplateData

diff --git a/Engine/Plugins/Runtime/MassGameplay/Source/MassSpawner/Public/MassEntityTemplate.h b/Engine/Plugins/Runtime/MassGameplay/Source/MassSpawner/Public/MassEntityTemplate.h
--- a/Engine/Plugins/Runtime/MassGameplay/Source/MassSpawner/Public/MassEntityTemplate.h
+++ b/Engine/Plugins/Runtime/MassGameplay/Source/MassSpawner/Public/MassEntityTemplate.h
@@ -281,6 +281,17 @@ struct MASSSPAWNER_API FMassEntityTemplateData
 		return Composition.SharedFragments.Contains(ScriptStruct);
 	}
 
+	template<typename T>
+	bool HasConstSharedFragment() const
+	{
+		return Composition.ConstSharedFragments.Contains<T>();
+	}
+
+	bool HasConstSharedFragment(const UScriptStruct& ScriptStruct) const
+	{
+		return Composition.ConstSharedFragments.Contains(ScriptStruct);
+	}
+
 	void Sort()
 	{
 		SharedFragmentValues.Sort();
