Trait changes:
* allowing for "null" traits to indicate procedurally added elements (i.e. not added by traits)
* allowing multiple traits to add the same fragments. Kept reporting the fact as a Warning, but only when developer tools are available.
* added checking for multiple instanced of a given trait class - the redundant instances are reported and ignored
* made results of trait's ValidateTemplate result matter by including it in the validation result (used to be kinda ignored)
  - this change exposed a missing dependency handling in UMassStateTreeTrait
* made messages put into MassEntity's MessageLog page more explicit and clickable
* made changes to existing traits to take advantage of the changes.

diff --git a/Engine/Plugins/AI/MassAI/Source/MassAIBehavior/Private/MassStateTreeTrait.cpp b/Engine/Plugins/AI/MassAI/Source/MassAIBehavior/Private/MassStateTreeTrait.cpp
--- a/Engine/Plugins/AI/MassAI/Source/MassAIBehavior/Private/MassStateTreeTrait.cpp
+++ b/Engine/Plugins/AI/MassAI/Source/MassAIBehavior/Private/MassStateTreeTrait.cpp
@@ -42,22 +42,23 @@ void UMassStateTreeTrait::BuildTemplate(FMassEntityTemplateBuildContext& BuildCo
 	BuildContext.AddFragment<FMassStateTreeInstanceFragment>();
 }
 
-void UMassStateTreeTrait::ValidateTemplate(FMassEntityTemplateBuildContext& BuildContext, const UWorld& World) const
+bool UMassStateTreeTrait::ValidateTemplate(FMassEntityTemplateBuildContext& BuildContext, const UWorld& World) const
 {
 	UMassStateTreeSubsystem* MassStateTreeSubsystem = World.GetSubsystem<UMassStateTreeSubsystem>();
 	if (!MassStateTreeSubsystem)
 	{
 		UE_VLOG(&World, LogMassBehavior, Error, TEXT("Failed to get Mass StateTree Subsystem."));
-		return;
+		return false;
 	}
 
 	if (!StateTree)
 	{
 		UE_VLOG(MassStateTreeSubsystem, LogMassBehavior, Error, TEXT("StateTree asset is not set or unavailable."));
-		return;
+		return false;
 	}
 
 	// Make sure all the required subsystems can be found.
+	bool bIssuesFound = false;
 	for (const FStateTreeExternalDataDesc& ItemDesc : StateTree->GetExternalDataDescs())
 	{
 		if (ensure(ItemDesc.Struct) && ItemDesc.Requirement == EStateTreeExternalDataRequirement::Required)
@@ -67,21 +68,33 @@ void UMassStateTreeTrait::ValidateTemplate(FMassEntityTemplateBuildContext& Buil
 				const TSubclassOf<UWorldSubsystem> SubClass = Cast<UClass>(const_cast<UStruct*>(ToRawPtr(ItemDesc.Struct)));
 				USubsystem* Subsystem = World.GetSubsystemBase(SubClass);
 				UE_CVLOG(!Subsystem, MassStateTreeSubsystem, LogMassBehavior, Error, TEXT("StateTree %s: Could not find required subsystem %s"), *GetNameSafe(StateTree), *GetNameSafe(ItemDesc.Struct));
+				bIssuesFound = bIssuesFound || !Subsystem;
 			}
 			else if (ItemDesc.Struct->IsChildOf(FMassFragment::StaticStruct()))
 			{
 				const bool bContainsFragment = BuildContext.HasFragment(*CastChecked<UScriptStruct>(ItemDesc.Struct));
 				UE_CVLOG(!bContainsFragment, MassStateTreeSubsystem, LogMassBehavior, Error, TEXT("StateTree %s: Could not find required fragment %s"), *GetNameSafe(StateTree), *GetNameSafe(ItemDesc.Struct));
+				bIssuesFound = bIssuesFound || !bContainsFragment;
 			}
 			else if (ItemDesc.Struct->IsChildOf(FMassSharedFragment::StaticStruct()))
 			{
 				const bool bContainsFragment = BuildContext.HasSharedFragment(*CastChecked<UScriptStruct>(ItemDesc.Struct));
 				UE_CVLOG(!bContainsFragment, MassStateTreeSubsystem, LogMassBehavior, Error, TEXT("StateTree %s: Could not find required shared fragment %s"), *GetNameSafe(StateTree), *GetNameSafe(ItemDesc.Struct));
+				bIssuesFound = bIssuesFound || !bContainsFragment;
+			}
+			else if (ItemDesc.Struct->IsChildOf(FMassConstSharedFragment::StaticStruct()))
+			{
+				const bool bContainsFragment = BuildContext.HasConstSharedFragment(*CastChecked<UScriptStruct>(ItemDesc.Struct));
+				UE_CVLOG(!bContainsFragment, MassStateTreeSubsystem, LogMassBehavior, Error, TEXT("StateTree %s: Could not find required const shared fragment %s"), *GetNameSafe(StateTree), *GetNameSafe(ItemDesc.Struct));
+				bIssuesFound = bIssuesFound || !bContainsFragment;
 			}
 			else
 			{
 				UE_VLOG(MassStateTreeSubsystem, LogMassBehavior, Error, TEXT("StateTree %s: Unsupported requirement %s"), *GetNameSafe(StateTree), *GetNameSafe(ItemDesc.Struct));
+				bIssuesFound = true;
 			}
 		}
 	}
+
+	return !bIssuesFound;
 }
diff --git a/Engine/Plugins/AI/MassAI/Source/MassAIBehavior/Public/MassStateTreeTrait.h b/Engine/Plugins/AI/MassAI/Source/MassAIBehavior/Public/MassStateTreeTrait.h
--- a/Engine/Plugins/AI/MassAI/Source/MassAIBehavior/Public/MassStateTreeTrait.h
+++ b/Engine/Plugins/AI/MassAI/Source/MassAIBehavior/Public/MassStateTreeTrait.h
@@ -19,7 +19,7 @@ protected:
 
 	virtual void BuildTemplate(FMassEntityTemplateBuildContext& BuildContext, const UWorld& World) const override;
 
-	virtual void ValidateTemplate(FMassEntityTemplateBuildContext& BuildContext, const UWorld& World) const override;
+	virtual bool ValidateTemplate(FMassEntityTemplateBuildContext& BuildContext, const UWorld& World) const override;
 
 	UPROPERTY(Category="StateTree", EditAnywhere, meta=(RequiredAssetDataTags="Schema=/Script/MassAIBehavior.MassStateTreeSchema"))
 	TObjectPtr<UStateTree> StateTree;
diff --git a/Engine/Plugins/Runtime/InstancedActors/Source/InstancedActors/Private/InstancedActorsVisualizationTrait.cpp b/Engine/Plugins/Runtime/InstancedActors/Source/InstancedActors/Private/InstancedActorsVisualizationTrait.cpp
--- a/Engine/Plugins/Runtime/InstancedActors/Source/InstancedActors/Private/InstancedActorsVisualizationTrait.cpp
+++ b/Engine/Plugins/Runtime/InstancedActors/Source/InstancedActors/Private/InstancedActorsVisualizationTrait.cpp
@@ -120,14 +120,17 @@ void UInstancedActorsVisualizationTrait::BuildTemplate(FMassEntityTemplateBuildC
 	// @todo Implement version of AddVisualDescWithISMComponent that supports multiple ISMCs and use that here
 	if (ensure(InstanceData.IsValid()))
 	{
-		FMassRepresentationFragment& RepresentationFragment = BuildContext.GetFragmentChecked<FMassRepresentationFragment>();
-		RepresentationFragment.StaticMeshDescHandle = InstanceData->GetDefaultVisualizationChecked().MassStaticMeshDescHandle;
-
-		if (RepresentationFragment.LowResTemplateActorIndex == INDEX_NONE)
+		FMassRepresentationFragment* RepresentationFragment = BuildContext.GetFragment<FMassRepresentationFragment>();
+		if (ensureMsgf(RepresentationFragment, TEXT("Configuration error, we always expect to have a FMassRepresentationFragment instance at this point")))
 		{
-			// if there's no "low res actor" we reuse the high-res one, otherwise we risk the visualization actor getting 
-			// removed when switching from EMassLOD::High down to EMassLOD::Medium
-			RepresentationFragment.LowResTemplateActorIndex = RepresentationFragment.HighResTemplateActorIndex;
+			RepresentationFragment->StaticMeshDescHandle = InstanceData->GetDefaultVisualizationChecked().MassStaticMeshDescHandle;
+
+			if (RepresentationFragment->LowResTemplateActorIndex == INDEX_NONE)
+			{
+				// if there's no "low res actor" we reuse the high-res one, otherwise we risk the visualization actor getting 
+				// removed when switching from EMassLOD::High down to EMassLOD::Medium
+				RepresentationFragment->LowResTemplateActorIndex = RepresentationFragment->HighResTemplateActorIndex;
+			}
 		}
 	}
 }
diff --git a/Engine/Plugins/Runtime/MassGameplay/Source/MassRepresentation/Private/MassVisualizationTrait.cpp b/Engine/Plugins/Runtime/MassGameplay/Source/MassRepresentation/Private/MassVisualizationTrait.cpp
--- a/Engine/Plugins/Runtime/MassGameplay/Source/MassRepresentation/Private/MassVisualizationTrait.cpp
+++ b/Engine/Plugins/Runtime/MassGameplay/Source/MassRepresentation/Private/MassVisualizationTrait.cpp
@@ -159,43 +159,45 @@ void UMassVisualizationTrait::Serialize(FArchive& Ar)
 #endif // WITH_EDITOR
 }
 
+bool UMassVisualizationTrait::ValidateTemplate(FMassEntityTemplateBuildContext& BuildContext, const UWorld& World) const
+{
+	Super::ValidateTemplate(BuildContext, World);
+
 #if WITH_EDITOR
-void UMassVisualizationTrait::ValidateParams() const
+	return ValidateParams();
+#else
+	return true;
+#endif // WITH_EDITOR
+}
+
+#if WITH_EDITOR
+bool UMassVisualizationTrait::ValidateParams() const
 {
-	// if this test is called on any of the CDOs we don't care, we're never going to utilize those in practice.
-	if (HasAnyFlags(RF_ClassDefaultObject))
-	{
-		return;
-	}
+	bool bIssuesFound = false;
 
-	// the SM config provided is not valid. We need to check if EMassRepresentationType::StaticMeshInstance
-	// is being used as any of the LODRepresentations. If so then we need to clear those out and report an error
-	if (StaticMeshInstanceDesc.IsValid() == false)
+	// if this test is called on any of the CDOs we don't care, we're never going to utilize those in practice.
+	if (HasAnyFlags(RF_ClassDefaultObject) == false)
 	{
-		for (int32 LODIndex = 0; LODIndex < EMassLOD::Max; ++LODIndex)
+		// the SM config provided is not valid. We need to check if EMassRepresentationType::StaticMeshInstance
+		// is being used as any of the LODRepresentations. If so then we need to clear those out and report an error
+		if (StaticMeshInstanceDesc.IsValid() == false)
 		{
-			if (Params.LODRepresentation[LODIndex] == EMassRepresentationType::StaticMeshInstance)
+			for (int32 LODIndex = 0; LODIndex < EMassLOD::Max; ++LODIndex)
 			{
-#if WITH_UNREAL_DEVELOPER_TOOLS && WITH_EDITOR
-				if (GEditor)
+				if (Params.LODRepresentation[LODIndex] == EMassRepresentationType::StaticMeshInstance)
 				{
-					static const FText ErrorMessage(LOCTEXT("VisualizationTraitMissingSM", "Trait using StaticMeshInstance representation type but no meshes are configured."));
-					static const FText InfoMessage(LOCTEXT("SeeLogForDetails", "See the log for details."));
+					bIssuesFound = true;
 
-					FMessageLog EditorErrors("MassEntity");
-					EditorErrors.Error(ErrorMessage);
-					EditorErrors.Notify(ErrorMessage);
-					EditorErrors.Info(InfoMessage);
+					UE_LOG(LogMassRepresentation, Error, TEXT("Trait %s is using StaticMeshInstance representation type for "
+						"LODRepresentation[%s] while the trait's StaticMeshInstanceDesc is not valid (has no Meshes). Entities "
+						"won't be visible at this LOD level.")
+						, *GetPathName(), *UEnum::GetValueAsString(EMassLOD::Type(LODIndex)));
 				}
-#endif // WITH_UNREAL_DEVELOPER_TOOLS
-
-				UE_LOG(LogMassRepresentation, Error, TEXT("Trait %s is using StaticMeshInstance representation type for "
-					"LODRepresentation[%s] while the trait's StaticMeshInstanceDesc is not valid (has no Meshes). Entities "
-					"won't be visible at this LOD level.")
-					, *GetPathName(), *UEnum::GetValueAsString(EMassLOD::Type(LODIndex)));
 			}
 		}
 	}
+
+	return !bIssuesFound;
 }
 
 void UMassVisualizationTrait::PostEditChangeProperty(FPropertyChangedEvent& PropertyChangedEvent)
diff --git a/Engine/Plugins/Runtime/MassGameplay/Source/MassRepresentation/Public/MassVisualizationTrait.h b/Engine/Plugins/Runtime/MassGameplay/Source/MassRepresentation/Public/MassVisualizationTrait.h
--- a/Engine/Plugins/Runtime/MassGameplay/Source/MassRepresentation/Public/MassVisualizationTrait.h
+++ b/Engine/Plugins/Runtime/MassGameplay/Source/MassRepresentation/Public/MassVisualizationTrait.h
@@ -65,6 +65,7 @@ protected:
 	bool bRegisterStaticMeshDesc = true;
 
 	virtual void BuildTemplate(FMassEntityTemplateBuildContext& BuildContext, const UWorld& World) const override;
+	virtual bool ValidateTemplate(FMassEntityTemplateBuildContext& BuildContext, const UWorld& World) const override;
 
 	/** 
 	 * Tests whether StaticMeshInstanceDesc is valid and if not cleans up InOutParamss of EMassRepresentationType::StaticMeshInstance
@@ -72,11 +73,11 @@ protected:
 	 * @param bStaticMeshDeterminedInvalid if StaticMeshInstanceDesc has already been determined invalid then bStaticMeshDeterminedInvalid
 	 *	can be set to `true` to skip the redundant check.
 	 */
-	void SanitizeParams(FMassRepresentationParameters& InOutParams, const bool bStaticMeshDeterminedInvalid = false) const;
+	virtual void SanitizeParams(FMassRepresentationParameters& InOutParams, const bool bStaticMeshDeterminedInvalid = false) const;
 
 	virtual void Serialize(FArchive& Ar) override;
 #if WITH_EDITOR
 	virtual void PostEditChangeProperty(struct FPropertyChangedEvent& PropertyChangedEvent) override;
-	void ValidateParams() const;
+	virtual bool ValidateParams() const;
 #endif // WITH_EDITOR
 };
diff --git a/Engine/Plugins/Runtime/MassGameplay/Source/MassSpawner/Private/MassEntityTemplateRegistry.cpp b/Engine/Plugins/Runtime/MassGameplay/Source/MassSpawner/Private/MassEntityTemplateRegistry.cpp
--- a/Engine/Plugins/Runtime/MassGameplay/Source/MassSpawner/Private/MassEntityTemplateRegistry.cpp
+++ b/Engine/Plugins/Runtime/MassGameplay/Source/MassSpawner/Private/MassEntityTemplateRegistry.cpp
@@ -13,6 +13,8 @@
 
 #if WITH_EDITOR
 #include "Editor.h"
+#include "Misc/UObjectToken.h"
+#include "Framework/Docking/TabManager.h"
 #endif
 
 #define LOCTEXT_NAMESPACE "Mass"
@@ -88,100 +90,165 @@ void FMassEntityTemplateRegistry::DestroyTemplate(FMassEntityTemplateID Template
 //----------------------------------------------------------------------//
 bool FMassEntityTemplateBuildContext::BuildFromTraits(TConstArrayView<UMassEntityTraitBase*> Traits, const UWorld& World)
 {
-	TraitAddedTypes.Reset();
-	TraitsDependencies.Reset();
+	ensureMsgf(bBuildInProgress == false, TEXT("Unexpected occurrence - it suggests FMassEntityTemplateBuildContext::BuildFromTraits "
+		"has been called as a consequence of some UMassEntityTraitBase::BuildTemplate call. Check the callstack."));
 
+	bBuildInProgress = true;
 	for (const UMassEntityTraitBase* Trait : Traits)
 	{
 		check(Trait);
-		BuildingTrait = Trait;
-		BuildingTrait->BuildTemplate(*this, World);
+		if (SetTraitBeingProcessed(Trait))
+		{
+			Trait->BuildTemplate(*this, World);
+		}
+	}
+	bBuildInProgress = false;
+
+	const bool bTemplateValid = ValidateBuildContext(World);
+	
+	ResetBuildTimeData();
+
+	return bTemplateValid;
+}
+
+bool FMassEntityTemplateBuildContext::SetTraitBeingProcessed(const UMassEntityTraitBase* Trait)
+{
+	if (Trait == nullptr || TraitsProcessed.Contains(Trait) == false)
+	{
+		TraitsData.Add({Trait});
+		return true;
 	}
 
-	BuildingTrait = nullptr;
+	UE_LOG(LogMass, Warning, TEXT("Attempting to add %s to FMassEntityTemplateBuildContext while this or another instance of the trait class has already been added.")
+		, *GetNameSafe(Trait));
 
-	return ValidateBuildContext(World);
+	IgnoredTraits.Add(Trait);
+	return false;
 }
 
 bool FMassEntityTemplateBuildContext::ValidateBuildContext(const UWorld& World)
 {
-	// Group same types(key) together
-	TraitAddedTypes.KeySort( [](const UStruct& LHS, const UStruct& RHS) { return LHS.GetName() < RHS.GetName(); } );
-
-	// Loop through all the registered fragments and make sure only one trait registered them.
-	const UStruct* CurrentStruct = nullptr;
-	const UMassEntityTraitBase* CurrentTrait = nullptr;
-	bool bHeaderOutputted = false;
-	bool bFragmentHasMultipleOwners = false;
-	for (const auto& Pair : TraitAddedTypes)
+#define WITH_MESSAGES (WITH_UNREAL_DEVELOPER_TOOLS && WITH_EDITOR)
+#if WITH_MESSAGES
+	TArray<TSharedRef<FTokenizedMessage>> Messages;
+#define IF_MESSAGES(Message) if (GEditor) { Message }
+#else
+#define IF_MESSAGES(_)
+#endif // WITH_MESSAGES
+
+	int32 ErrorCount = 0;
+	int32 WarningCount = 0;
+
+
+	TMap<const UStruct*, const UMassEntityTraitBase*> TypesAlreadyAdded;
+
+	// these are non-critical warnings, we want to report these to the users as a potential configuration issue,
+	// but it won't affect the final entity template composition (for example adding the same fragment is fine since 
+	// the entity template handles that gracefully).
+	for (const FTraitData& TraitData : TraitsData)
 	{
-		if (CurrentStruct != Pair.Key)
+		for (const UStruct* TypeAdded : TraitData.TypesAdded)
 		{
-			CurrentStruct = Pair.Key;
-			CurrentTrait = Pair.Value;
-			check(CurrentTrait);
-			CurrentTrait->ValidateTemplate(*this, World);
-			bHeaderOutputted = false;
-		}
-		else
-		{
-			if (!bHeaderOutputted)
+			const UMassEntityTraitBase*& SourceTrait = TypesAlreadyAdded.FindOrAdd(TypeAdded);
+			if (SourceTrait != nullptr)
+			{
+				// we report this only if it wasn't added twice by the same trait, the one we're processing right now
+				UE_CLOG(SourceTrait != TraitData.Trait
+					, LogMass, Warning, TEXT("%s: Fragment %s already added by %s")
+					, *GetNameSafe(TraitData.Trait), *GetNameSafe(TypeAdded), *SourceTrait->GetName());
+				++WarningCount;
+
+				IF_MESSAGES(
+					Messages.Add_GetRef(FTokenizedMessage::Create(EMessageSeverity::Warning))
+						->AddToken(FUObjectToken::Create(TraitData.Trait))
+						->AddToken(FTextToken::Create(LOCTEXT("MassEntityTraitFragmentDuplicationWarning1", "trying to add fragment of type")))
+						->AddToken(FUObjectToken::Create(TypeAdded))
+						->AddToken(FTextToken::Create(LOCTEXT("MassEntityTraitFragmentDuplicationWarning2", "while it has already been added by")))
+						->AddToken(FUObjectToken::Create(SourceTrait));
+				);
+			}
+			else
 			{
-				UE_LOG(LogMass, Warning, TEXT("%s: Fragment(%s) was added multiple time and should only be added by one trait. Fragment was added by:")
-					, CurrentTrait ? *GetNameSafe(CurrentTrait->GetOuter()) : TEXT("None")
-					, CurrentStruct ? *CurrentStruct->GetName() : TEXT("null"));
-				UE_LOG(LogMass, Warning, TEXT("\t\t%s"), CurrentTrait ? *CurrentTrait->GetClass()->GetName() : TEXT("null"));
-				bHeaderOutputted = true;
+				SourceTrait = TraitData.Trait;
 			}
-			UE_LOG(LogMass, Warning, TEXT("\t\t%s"), *Pair.Value->GetClass()->GetName());
-			bFragmentHasMultipleOwners = true;
 		}
- 	}
+	}
 
-	// Loop through all the traits dependencies and check if they have been added
-	CurrentTrait = nullptr;
-	bHeaderOutputted = false;
-	bool bMissingFragmentDependencies = false;
-	for (const auto& Dependency : TraitsDependencies)
+	// these are critical, we're going to fail the validation if anything here fails
+	for (const FTraitData& TraitData : TraitsData)
 	{
-		if (CurrentTrait != Dependency.Get<1>())
+		for (const UStruct* TypeRequired : TraitData.TypesRequired)
 		{
-			CurrentTrait = Dependency.Get<1>();
-			bHeaderOutputted = false;
-		}
-		if (!TraitAddedTypes.Contains(Dependency.Get<0>()))
-		{
-			if (!bHeaderOutputted)
+			if (TypesAlreadyAdded.Contains(TypeRequired) == false)
 			{
-				check(CurrentTrait);
-				UE_LOG(LogMass, Error, TEXT("%s: Trait(%s) has missing dependency:"), *GetNameSafe(CurrentTrait->GetOuter())
-					, *CurrentTrait->GetClass()->GetName());
-				bHeaderOutputted = true;
+				UE_LOG(LogMass, Error, TEXT("%s: Missing required fragment of type %s")
+					, *GetNameSafe(TraitData.Trait), *GetNameSafe(TypeRequired));
+				++ErrorCount;
+				IF_MESSAGES(
+					Messages.Add_GetRef(FTokenizedMessage::Create(EMessageSeverity::Error))
+						->AddToken(FUObjectToken::Create(TraitData.Trait))
+						->AddToken(FTextToken::Create(LOCTEXT("MassEntityTraitMissingDependencies", "unsatisfied dependency, missing")))
+						->AddToken(FUObjectToken::Create(TypeRequired));
+				);
 			}
-			UE_LOG(LogMass, Error, TEXT("\t\t%s"), *Dependency.Get<0>()->GetName());
-			bMissingFragmentDependencies = true;
 		}
 	}
 
-#if WITH_UNREAL_DEVELOPER_TOOLS && WITH_EDITOR
-	if (GEditor && (bFragmentHasMultipleOwners || bMissingFragmentDependencies))
+	for (const FTraitData& TraitData : TraitsData)
 	{
-		FMessageLog EditorErrors("MassEntity");
-		if (bFragmentHasMultipleOwners)
+		if (TraitData.Trait && TraitData.Trait->ValidateTemplate(*this, World) == false)
 		{
-			EditorErrors.Warning(LOCTEXT("MassEntityTraitsFragmentOwnershipError", "Some fragments are added by multiple traits and can only be added by one!"));
-			EditorErrors.Notify(LOCTEXT("MassEntityTraitsFragmentOwnershipError", "Some fragments are added by multiple traits and can only be added by one!"));
+			++ErrorCount;
+			IF_MESSAGES(
+				Messages.Add_GetRef(FTokenizedMessage::Create(EMessageSeverity::Error))
+					->AddToken(FUObjectToken::Create(TraitData.Trait))
+					->AddToken(FTextToken::Create(LOCTEXT("MassEntityTraitFailedValidation", "trait-specific validation failed")));
+			);
 		}
-		if (bMissingFragmentDependencies)
-		{
-			EditorErrors.Error(LOCTEXT("MassEntityTraitsMissingFragment", "Some traits are requiring the presence of fragments which are missing!"));
-			EditorErrors.Notify(LOCTEXT("MassEntityTraitsMissingFragment", "Some traits are requiring the presence of fragments which are missing!"));
-		}
-		EditorErrors.Info(FText::FromString(TEXT("See the log for details")));
 	}
-#endif // WITH_UNREAL_DEVELOPER_TOOLS
 
-	return !bFragmentHasMultipleOwners && !bMissingFragmentDependencies;
+	for (const UMassEntityTraitBase* IgnoredTrait : IgnoredTraits)
+	{
+		IF_MESSAGES(
+			Messages.Add_GetRef(FTokenizedMessage::Create(EMessageSeverity::Warning))
+				->AddToken(FUObjectToken::Create(IgnoredTrait))
+				->AddToken(FTextToken::Create(LOCTEXT("MassEntityTraitIgnoredTrait", "trait was ignored. Check if it's not a duplicate.")));
+		);
+		++WarningCount;
+	}
+	
+	// @todo add dependencies on trait classes? might be hard if traits are unrelated, like requiring UMassLODCollectorTrait 
+	// or UMassDistanceLODCollectorTrait - both supply alternative implementations of a given functionality, but are unrelated.
+	// Could be done with a complex requirements system (similar to entity queries - "all of X", "any of Y", etc) - probably 
+	// not worth it since we don't even have a use case for it right now.
+
+#if WITH_MESSAGES
+	if (GEditor && Messages.Num())
+	{
+		TSharedRef<FTokenizedMessage> SummaryMessage = Messages.Add_GetRef(
+			FTokenizedMessage::Create(ErrorCount ? EMessageSeverity::Error : EMessageSeverity::Warning))
+			->AddToken(FTextToken::Create(FText::FormatOrdered(LOCTEXT("MassEntityTraitResult", "Mass Entity Template validation:\n{0} errors and {1} warnings found"), ErrorCount, WarningCount)));
+
+		Messages.Add_GetRef(FTokenizedMessage::Create(EMessageSeverity::Info))
+			->AddToken(FActionToken::Create(LOCTEXT("MassSeeLogForDetails", "See the log for more details.")
+				, LOCTEXT("MassSeeLogForDetailsTooltip", "Open the Output Log tab.")
+				, FOnActionTokenExecuted::CreateLambda([]()
+					{
+						FGlobalTabmanager::Get()->TryInvokeTab(FName("OutputLog"));
+					}))
+				);
+
+		FMessageLog EditorErrors("MassEntity");
+		EditorErrors.AddMessages(Messages);
+		EditorErrors.Notify(SummaryMessage->ToText());
+	}
+#endif // WITH_MESSAGES
+
+#undef IF_MESSAGES
+#undef WITH_MESSAGES
+
+	// only the Errors render the template invalid, Warnings just warn about stuff not being set up quite right, but we can recover.
+	return (ErrorCount == 0);
 }
 
 //-----------------------------------------------------------------------------
diff --git a/Engine/Plugins/Runtime/MassGameplay/Source/MassSpawner/Public/MassEntityTemplateRegistry.h b/Engine/Plugins/Runtime/MassGameplay/Source/MassSpawner/Public/MassEntityTemplateRegistry.h
--- a/Engine/Plugins/Runtime/MassGameplay/Source/MassSpawner/Public/MassEntityTemplateRegistry.h
+++ b/Engine/Plugins/Runtime/MassGameplay/Source/MassSpawner/Public/MassEntityTemplateRegistry.h
@@ -9,6 +9,8 @@
 #include "MassTranslator.h"
 #include "MassEntityTemplateRegistry.generated.h"
 
+#define ENSURE_SUPPORTED_TRAIT_OPERATION() ensureMsgf(bBuildInProgress == false, TEXT("This method is not expected to be called as "\
+	"part of trait's BuildTemplate call. Traits are not supposed to add elements based on other traits due to arbitrary trait ordering."));
 
 class UWorld;
 class UMassEntityTraitBase;
@@ -92,24 +94,21 @@ struct FMassEntityTemplateBuildContext
 	}
 
 	template<typename T>
-	T& GetFragmentChecked()
+	T* GetFragment()
 	{
-		check(TraitAddedTypes.Find(T::StaticStruct()) != nullptr);
-		T* FragmentInstance = TemplateData.GetMutableFragment<T>();
-		check(FragmentInstance);
-		return *FragmentInstance;
+		return TemplateData.GetMutableFragment<T>();
 	}
 
 	template<typename T>
 	bool HasFragment() const
 	{
-		ensureMsgf(!BuildingTrait, TEXT("This method is not expected to be called within the build from trait call."));
+		ENSURE_SUPPORTED_TRAIT_OPERATION();
 		return TemplateData.HasFragment<T>();
 	}
 	
 	bool HasFragment(const UScriptStruct& ScriptStruct) const
 	{
-		ensureMsgf(!BuildingTrait, TEXT("This method is not expected to be called within the build from trait call."));
+		ENSURE_SUPPORTED_TRAIT_OPERATION();
 		return TemplateData.HasFragment(ScriptStruct);
 	}
 
@@ -122,23 +121,36 @@ struct FMassEntityTemplateBuildContext
 	template<typename T>
 	bool HasChunkFragment() const
 	{
-		ensureMsgf(!BuildingTrait, TEXT("This method is not expected to be called within the build from trait call."));
+		ENSURE_SUPPORTED_TRAIT_OPERATION();
 		return TemplateData.HasChunkFragment<T>();
 	}
 
 	template<typename T>
 	bool HasSharedFragment() const
 	{
-		ensureMsgf(!BuildingTrait, TEXT("This method is not expected to be called within the build from trait call."));
+		ENSURE_SUPPORTED_TRAIT_OPERATION();
 		return TemplateData.HasSharedFragment<T>();
 	}
 
 	bool HasSharedFragment(const UScriptStruct& ScriptStruct) const
 	{
-		ensureMsgf(!BuildingTrait, TEXT("This method is not expected to be called within the build from trait call."));
+		ENSURE_SUPPORTED_TRAIT_OPERATION();
 		return TemplateData.HasSharedFragment(ScriptStruct);
 	}
 
+	template<typename T>
+	bool HasConstSharedFragment() const
+	{
+		ENSURE_SUPPORTED_TRAIT_OPERATION();
+		return TemplateData.HasConstSharedFragment<T>();
+	}
+
+	bool HasConstSharedFragment(const UScriptStruct& ScriptStruct) const
+	{
+		ENSURE_SUPPORTED_TRAIT_OPERATION();
+		return TemplateData.HasConstSharedFragment(ScriptStruct);
+	}
+
 	//----------------------------------------------------------------------//
 	// Translators
 	//----------------------------------------------------------------------//
@@ -166,7 +178,7 @@ struct FMassEntityTemplateBuildContext
 
 	void AddDependency(const UStruct* Dependency)
 	{
-		TraitsDependencies.Add( {Dependency, BuildingTrait} );
+		TraitsData.Last().TypesRequired.Add(Dependency);
 	}
 
 	//----------------------------------------------------------------------//
@@ -198,15 +210,37 @@ protected:
 
 	void TypeAdded(const UStruct& Type)
 	{
-		if (ensureMsgf(BuildingTrait, TEXT("Expected to be called within the BuildTemplateFromTrait method")))
-		{
-			TraitAddedTypes.Add(&Type, BuildingTrait);
-		}
+		checkf(TraitsData.Num(), TEXT("Adding elements to the build context before BuildFromTraits or SetTraitBeingProcessed was called is unsupported"));
+		TraitsData.Last().TypesAdded.Add(&Type);
 	}
 
-	const UMassEntityTraitBase* BuildingTrait = nullptr;
-	TMultiMap<const UStruct*, const UMassEntityTraitBase*> TraitAddedTypes;
-	TArray< TTuple<const UStruct*, const UMassEntityTraitBase*> > TraitsDependencies;
+	/** 
+	 * Return true if the given trait can be used. The function will fail if a trait instance of the given class has already 
+	 * been processed. The function will also fail the very same trait instance is used multiple times.
+	 * Note that it's ok for Trait to be nullptr to indicate the subsequent additions to the build context are procedural
+	 * in nature and are not associated with any traits. In that case it's ok to have multiple SetTraitBeingProcessed(nullptr)
+	 * calls.
+	 */
+	bool SetTraitBeingProcessed(const UMassEntityTraitBase* Trait);
+
+	void ResetBuildTimeData()
+	{
+		TraitsData.Reset();
+		TraitsProcessed.Reset();
+		IgnoredTraits.Reset();
+		bBuildInProgress = false;
+	}
+
+	struct FTraitData
+	{
+		const UMassEntityTraitBase* Trait = nullptr;
+		TArray<const UStruct*> TypesAdded;
+		TArray<const UStruct*> TypesRequired;
+	};
+	TArray<FTraitData> TraitsData;
+	TSet<const UMassEntityTraitBase*> TraitsProcessed;
+	TSet<const UMassEntityTraitBase*> IgnoredTraits;
+	bool bBuildInProgress = false;
 
 	FMassEntityTemplateData& TemplateData;
 	FMassEntityTemplateID TemplateID;
@@ -275,3 +309,5 @@ class MASSSPAWNER_API UDEPRECATED_MassEntityTemplateRegistry : public UObject
 {
 	GENERATED_BODY()
 };
+
+#undef ENSURE_SUPPORTED_TRAIT_OPERATION
diff --git a/Engine/Plugins/Runtime/MassGameplay/Source/MassSpawner/Public/MassEntityTraitBase.h b/Engine/Plugins/Runtime/MassGameplay/Source/MassSpawner/Public/MassEntityTraitBase.h
--- a/Engine/Plugins/Runtime/MassGameplay/Source/MassSpawner/Public/MassEntityTraitBase.h
+++ b/Engine/Plugins/Runtime/MassGameplay/Source/MassSpawner/Public/MassEntityTraitBase.h
@@ -26,6 +26,14 @@ public:
 
 	virtual void DestroyTemplate() const {}
 
-	/** Called when all Traits have had BuildTemplate() called. */
-	virtual void ValidateTemplate(FMassEntityTemplateBuildContext& BuildContext, const UWorld& World) const {};
+	/**
+	 * Called once all traits have been processed and fragment requirements have been checked. Override this function
+	 * to perform additional Trait's configuration validation. Returning `false` will indicate that the trait instance
+	 * is not happy with the validation results - this result will be treated as an error.
+	 * @return whether the validation was successful
+	 */
+	virtual bool ValidateTemplate(FMassEntityTemplateBuildContext& BuildContext, const UWorld& World) const
+	{
+		return true;
+	}
 };
