[Mass] added some dependencies to MassSpawner and MassGameplayEditor modules in preparation for the trait validation CL

diff --git a/Engine/Plugins/Runtime/MassGameplay/MassGameplay.uplugin b/Engine/Plugins/Runtime/MassGameplay/MassGameplay.uplugin
--- a/Engine/Plugins/Runtime/MassGameplay/MassGameplay.uplugin
+++ b/Engine/Plugins/Runtime/MassGameplay/MassGameplay.uplugin
@@ -99,6 +99,10 @@
 		{
 			"Name": "StateTree",
 			"Enabled": true
+		},
+		{
+			"Name": "DataValidation",
+			"Enabled": true
 		}
 	]
 }
\ No newline at end of file
diff --git a/Engine/Plugins/Runtime/MassGameplay/Source/MassGameplayEditor/MassGameplayEditor.Build.cs b/Engine/Plugins/Runtime/MassGameplay/Source/MassGameplayEditor/MassGameplayEditor.Build.cs
--- a/Engine/Plugins/Runtime/MassGameplay/Source/MassGameplayEditor/MassGameplayEditor.Build.cs
+++ b/Engine/Plugins/Runtime/MassGameplay/Source/MassGameplayEditor/MassGameplayEditor.Build.cs
@@ -46,6 +46,7 @@ namespace UnrealBuildTool.Rules
 				"PropertyEditor",
 				"AIGraph",
 				"ToolMenus",
+				"DataValidation",
 			}
 			);
 		}
diff --git a/Engine/Plugins/Runtime/MassGameplay/Source/MassSpawner/MassSpawner.Build.cs b/Engine/Plugins/Runtime/MassGameplay/Source/MassSpawner/MassSpawner.Build.cs
--- a/Engine/Plugins/Runtime/MassGameplay/Source/MassSpawner/MassSpawner.Build.cs
+++ b/Engine/Plugins/Runtime/MassGameplay/Source/MassSpawner/MassSpawner.Build.cs
@@ -31,6 +31,11 @@ namespace UnrealBuildTool.Rules
 						"Slate"
 					}
 				);
+
+				PublicDependencyModuleNames.Add("MassEntityEditor");
+				// here for communication with MassTraitRepository
+				PrivateIncludePathModuleNames.Add("MassGameplayEditor");
+				DynamicallyLoadedModuleNames.Add("MassGameplayEditor");
 			}
 		}
 	}
