[Mass] Improved Mass Trait validation workflow:
* issues are clearly and precisesly described in the MessageLog, MassEntity page
* messages contain specific elemtn type names missing and which traits provide said types
* there's a convenient "fix" option for every "missing element" issue
* traits can now declare additional requirements during validation - this is intended to allow asset-specific dependencies. StateTreeTrait is a good example here and is attached to the review
* added an editor-time function for procedurally adding traits to entity configs. This functionality is used to fix entity config assets in the editor

In addition:
* added a bunch of BuildContext.IsInspectingData calls to multiple traits to support new Trait validation and fixing.

diff --git a/Engine/Plugins/AI/MassAI/Source/MassAIBehavior/Private/MassStateTreeTrait.cpp b/Engine/Plugins/AI/MassAI/Source/MassAIBehavior/Private/MassStateTreeTrait.cpp
--- a/Engine/Plugins/AI/MassAI/Source/MassAIBehavior/Private/MassStateTreeTrait.cpp
+++ b/Engine/Plugins/AI/MassAI/Source/MassAIBehavior/Private/MassStateTreeTrait.cpp
@@ -16,18 +16,18 @@ void UMassStateTreeTrait::BuildTemplate(FMassEntityTemplateBuildContext& BuildCo
 	FMassEntityManager& EntityManager = UE::Mass::Utils::GetEntityManagerChecked(World);
 
 	UMassStateTreeSubsystem* MassStateTreeSubsystem = World.GetSubsystem<UMassStateTreeSubsystem>();
-	if (!MassStateTreeSubsystem)
+	if (!MassStateTreeSubsystem && !BuildContext.IsInspectingData())
 	{
 		UE_VLOG(&World, LogMassBehavior, Error, TEXT("Failed to get Mass StateTree Subsystem."));
 		return;
 	}
 
-	if (!StateTree)
+	if (!StateTree && !BuildContext.IsInspectingData())
 	{
 		UE_VLOG(MassStateTreeSubsystem, LogMassBehavior, Error, TEXT("StateTree asset is not set or unavailable."));
 		return;
 	}
-	if (!StateTree->IsReadyToRun())
+	if (StateTree != nullptr && !BuildContext.IsInspectingData() && !StateTree->IsReadyToRun())
 	{
 		UE_VLOG(MassStateTreeSubsystem, LogMassBehavior, Error, TEXT("StateTree asset is ready to run."));
 		return;
@@ -42,7 +42,7 @@ void UMassStateTreeTrait::BuildTemplate(FMassEntityTemplateBuildContext& BuildCo
 	BuildContext.AddFragment<FMassStateTreeInstanceFragment>();
 }
 
-bool UMassStateTreeTrait::ValidateTemplate(FMassEntityTemplateBuildContext& BuildContext, const UWorld& World) const
+bool UMassStateTreeTrait::ValidateTemplate(const FMassEntityTemplateBuildContext& BuildContext, const UWorld& World, FAdditionalTraitRequirements& OutTraitRequirements) const
 {
 	UMassStateTreeSubsystem* MassStateTreeSubsystem = World.GetSubsystem<UMassStateTreeSubsystem>();
 	if (!MassStateTreeSubsystem)
@@ -65,32 +65,47 @@ bool UMassStateTreeTrait::ValidateTemplate(FMassEntityTemplateBuildContext& Buil
 		{
 			if (ItemDesc.Struct->IsChildOf(UWorldSubsystem::StaticClass()))
 			{
-				const TSubclassOf<UWorldSubsystem> SubClass = Cast<UClass>(const_cast<UStruct*>(ToRawPtr(ItemDesc.Struct)));
-				USubsystem* Subsystem = World.GetSubsystemBase(SubClass);
-				UE_CVLOG(!Subsystem, MassStateTreeSubsystem, LogMassBehavior, Error, TEXT("StateTree %s: Could not find required subsystem %s"), *GetNameSafe(StateTree), *GetNameSafe(ItemDesc.Struct));
-				bIssuesFound = bIssuesFound || !Subsystem;
+				if (BuildContext.IsInspectingData() == false)
+				{
+					const TSubclassOf<UWorldSubsystem> SubClass = Cast<UClass>(const_cast<UStruct*>(ToRawPtr(ItemDesc.Struct)));
+					USubsystem* Subsystem = World.GetSubsystemBase(SubClass);
+					UE_CVLOG_ALWAYS_UELOG(!Subsystem, MassStateTreeSubsystem, LogMassBehavior, Error, TEXT("StateTree %s: Could not find required subsystem %s"), *GetNameSafe(StateTree), *GetNameSafe(ItemDesc.Struct));
+					bIssuesFound = bIssuesFound || !Subsystem;
+				}
 			}
 			else if (ItemDesc.Struct->IsChildOf(FMassFragment::StaticStruct()))
 			{
-				const bool bContainsFragment = BuildContext.HasFragment(*CastChecked<UScriptStruct>(ItemDesc.Struct));
-				UE_CVLOG(!bContainsFragment, MassStateTreeSubsystem, LogMassBehavior, Error, TEXT("StateTree %s: Could not find required fragment %s"), *GetNameSafe(StateTree), *GetNameSafe(ItemDesc.Struct));
-				bIssuesFound = bIssuesFound || !bContainsFragment;
+				const UScriptStruct& FragmentType = *CastChecked<UScriptStruct>(ItemDesc.Struct);
+				if (BuildContext.HasFragment(FragmentType) == false)
+				{
+					OutTraitRequirements.Add(&FragmentType);
+					UE_VLOG_ALWAYS_UELOG(MassStateTreeSubsystem, LogMassBehavior, Error, TEXT("StateTree %s: Could not find required fragment %s"), *GetNameSafe(StateTree), *GetNameSafe(ItemDesc.Struct));
+					bIssuesFound = true;
+				}
 			}
 			else if (ItemDesc.Struct->IsChildOf(FMassSharedFragment::StaticStruct()))
 			{
-				const bool bContainsFragment = BuildContext.HasSharedFragment(*CastChecked<UScriptStruct>(ItemDesc.Struct));
-				UE_CVLOG(!bContainsFragment, MassStateTreeSubsystem, LogMassBehavior, Error, TEXT("StateTree %s: Could not find required shared fragment %s"), *GetNameSafe(StateTree), *GetNameSafe(ItemDesc.Struct));
-				bIssuesFound = bIssuesFound || !bContainsFragment;
+				const UScriptStruct& FragmentType = *CastChecked<UScriptStruct>(ItemDesc.Struct);
+				if (BuildContext.HasSharedFragment(FragmentType) == false)
+				{
+					OutTraitRequirements.Add(&FragmentType);
+					UE_VLOG_ALWAYS_UELOG(MassStateTreeSubsystem, LogMassBehavior, Error, TEXT("StateTree %s: Could not find required shared fragment %s"), *GetNameSafe(StateTree), *GetNameSafe(ItemDesc.Struct));
+					bIssuesFound = true;
+				}
 			}
 			else if (ItemDesc.Struct->IsChildOf(FMassConstSharedFragment::StaticStruct()))
 			{
-				const bool bContainsFragment = BuildContext.HasConstSharedFragment(*CastChecked<UScriptStruct>(ItemDesc.Struct));
-				UE_CVLOG(!bContainsFragment, MassStateTreeSubsystem, LogMassBehavior, Error, TEXT("StateTree %s: Could not find required const shared fragment %s"), *GetNameSafe(StateTree), *GetNameSafe(ItemDesc.Struct));
-				bIssuesFound = bIssuesFound || !bContainsFragment;
+				const UScriptStruct& FragmentType = *CastChecked<UScriptStruct>(ItemDesc.Struct);
+				if (BuildContext.HasConstSharedFragment(FragmentType) == false)
+				{
+					OutTraitRequirements.Add(&FragmentType);
+					UE_VLOG_ALWAYS_UELOG(MassStateTreeSubsystem, LogMassBehavior, Error, TEXT("StateTree %s: Could not find required const shared fragment %s"), *GetNameSafe(StateTree), *GetNameSafe(ItemDesc.Struct));
+					bIssuesFound = true;
+				}
 			}
 			else
 			{
-				UE_VLOG(MassStateTreeSubsystem, LogMassBehavior, Error, TEXT("StateTree %s: Unsupported requirement %s"), *GetNameSafe(StateTree), *GetNameSafe(ItemDesc.Struct));
+				UE_VLOG_ALWAYS_UELOG(MassStateTreeSubsystem, LogMassBehavior, Error, TEXT("StateTree %s: Unsupported requirement %s"), *GetNameSafe(StateTree), *GetNameSafe(ItemDesc.Struct));
 				bIssuesFound = true;
 			}
 		}
diff --git a/Engine/Plugins/AI/MassAI/Source/MassAIBehavior/Public/MassStateTreeTrait.h b/Engine/Plugins/AI/MassAI/Source/MassAIBehavior/Public/MassStateTreeTrait.h
--- a/Engine/Plugins/AI/MassAI/Source/MassAIBehavior/Public/MassStateTreeTrait.h
+++ b/Engine/Plugins/AI/MassAI/Source/MassAIBehavior/Public/MassStateTreeTrait.h
@@ -18,8 +18,7 @@ class MASSAIBEHAVIOR_API UMassStateTreeTrait : public UMassEntityTraitBase
 protected:
 
 	virtual void BuildTemplate(FMassEntityTemplateBuildContext& BuildContext, const UWorld& World) const override;
-
-	virtual bool ValidateTemplate(FMassEntityTemplateBuildContext& BuildContext, const UWorld& World) const override;
+	virtual bool ValidateTemplate(const FMassEntityTemplateBuildContext& BuildContext, const UWorld& World, FAdditionalTraitRequirements& OutTraitRequirements) const override;
 
 	UPROPERTY(Category="StateTree", EditAnywhere, meta=(RequiredAssetDataTags="Schema=/Script/MassAIBehavior.MassStateTreeSchema"))
 	TObjectPtr<UStateTree> StateTree;
diff --git a/Engine/Plugins/AI/MassCrowd/Source/MassCrowd/Private/MassCrowdServerRepresentationTrait.cpp b/Engine/Plugins/AI/MassCrowd/Source/MassCrowd/Private/MassCrowdServerRepresentationTrait.cpp
--- a/Engine/Plugins/AI/MassCrowd/Source/MassCrowd/Private/MassCrowdServerRepresentationTrait.cpp
+++ b/Engine/Plugins/AI/MassCrowd/Source/MassCrowd/Private/MassCrowdServerRepresentationTrait.cpp
@@ -26,7 +26,7 @@ UMassCrowdServerRepresentationTrait::UMassCrowdServerRepresentationTrait()
 void UMassCrowdServerRepresentationTrait::BuildTemplate(FMassEntityTemplateBuildContext& BuildContext, const UWorld& World) const
 {
 	// This should only be ran on NM_DedicatedServer network mode
-	if (!World.IsNetMode(NM_DedicatedServer))
+	if (!World.IsNetMode(NM_DedicatedServer) && !BuildContext.IsInspectingData())
 	{
 		return;
 	}
@@ -38,7 +38,7 @@ void UMassCrowdServerRepresentationTrait::BuildTemplate(FMassEntityTemplateBuild
 	FMassEntityManager& EntityManager = UE::Mass::Utils::GetEntityManagerChecked(World);
 
 	UMassCrowdRepresentationSubsystem* RepresentationSubsystem = World.GetSubsystem<UMassCrowdRepresentationSubsystem>();
-	check(RepresentationSubsystem);
+	check(RepresentationSubsystem || BuildContext.IsInspectingData());
 
 	FMassRepresentationSubsystemSharedFragment SubsystemSharedFragment;
 	SubsystemSharedFragment.RepresentationSubsystem = RepresentationSubsystem;
@@ -52,7 +52,10 @@ void UMassCrowdServerRepresentationTrait::BuildTemplate(FMassEntityTemplateBuild
 
 	FMassRepresentationFragment& RepresentationFragment = BuildContext.AddFragment_GetRef<FMassRepresentationFragment>();
 	RepresentationFragment.StaticMeshDescHandle = FStaticMeshInstanceVisualizationDescHandle();
-	RepresentationFragment.HighResTemplateActorIndex = TemplateActor.Get() ? RepresentationSubsystem->FindOrAddTemplateActor(TemplateActor.Get()) : INDEX_NONE;
+	if (!BuildContext.IsInspectingData())
+	{
+		RepresentationFragment.HighResTemplateActorIndex = TemplateActor.Get() ? RepresentationSubsystem->FindOrAddTemplateActor(TemplateActor.Get()) : INDEX_NONE;
+	}
 	RepresentationFragment.LowResTemplateActorIndex = INDEX_NONE;
 
 	BuildContext.AddFragment<FMassRepresentationLODFragment>();
diff --git a/Engine/Plugins/Runtime/MassGameplay/Source/MassActors/Private/MassAgentTraits.cpp b/Engine/Plugins/Runtime/MassGameplay/Source/MassActors/Private/MassAgentTraits.cpp
--- a/Engine/Plugins/Runtime/MassGameplay/Source/MassActors/Private/MassAgentTraits.cpp
+++ b/Engine/Plugins/Runtime/MassGameplay/Source/MassActors/Private/MassAgentTraits.cpp
@@ -43,7 +43,7 @@ void UMassAgentCapsuleCollisionSyncTrait::BuildTemplate(FMassEntityTemplateBuild
 {
 	BuildContext.AddFragment<FCapsuleComponentWrapperFragment>();
 	BuildContext.AddFragment<FAgentRadiusFragment>();
-	if (bSyncTransform)
+	if (bSyncTransform || BuildContext.IsInspectingData())
 	{
 		BuildContext.AddFragment<FTransformFragment>();
 	}
@@ -66,14 +66,14 @@ void UMassAgentCapsuleCollisionSyncTrait::BuildTemplate(FMassEntityTemplateBuild
 			}
 		});
 
-	if (bSyncTransform)
+	if (bSyncTransform || BuildContext.IsInspectingData())
 	{
-		if (EnumHasAnyFlags(SyncDirection, EMassTranslationDirection::ActorToMass))
+		if (EnumHasAnyFlags(SyncDirection, EMassTranslationDirection::ActorToMass) || BuildContext.IsInspectingData())
 		{
 			BuildContext.AddTranslator<UMassCapsuleTransformToMassTranslator>();
 		}
 
-		if (EnumHasAnyFlags(SyncDirection, EMassTranslationDirection::MassToActor))
+		if (EnumHasAnyFlags(SyncDirection, EMassTranslationDirection::MassToActor) || BuildContext.IsInspectingData())
 		{
 			BuildContext.AddTranslator<UMassTransformToActorCapsuleTranslator>();
 		}
@@ -112,12 +112,12 @@ void UMassAgentMovementSyncTrait::BuildTemplate(FMassEntityTemplateBuildContext&
 			}
 		});
 
-	if (EnumHasAnyFlags(SyncDirection, EMassTranslationDirection::ActorToMass))
+	if (EnumHasAnyFlags(SyncDirection, EMassTranslationDirection::ActorToMass) || BuildContext.IsInspectingData())
 	{
 		BuildContext.AddTranslator<UMassCharacterMovementToMassTranslator>();
 	}
 
-	if (EnumHasAnyFlags(SyncDirection, EMassTranslationDirection::MassToActor))
+	if (EnumHasAnyFlags(SyncDirection, EMassTranslationDirection::MassToActor) || BuildContext.IsInspectingData())
 	{
 		BuildContext.AddTranslator<UMassCharacterMovementToActorTranslator>();
 	}
@@ -130,12 +130,14 @@ void UMassAgentOrientationSyncTrait::BuildTemplate(FMassEntityTemplateBuildConte
 {
 	BuildContext.RequireFragment<FCharacterMovementComponentWrapperFragment>();
 
-	if (EnumHasAnyFlags(SyncDirection, EMassTranslationDirection::ActorToMass))
+	if (EnumHasAnyFlags(SyncDirection, EMassTranslationDirection::ActorToMass) 
+		|| BuildContext.IsInspectingData())
 	{
 		BuildContext.AddTranslator<UMassCharacterOrientationToMassTranslator>();
 	}
 
-	if (EnumHasAnyFlags(SyncDirection, EMassTranslationDirection::MassToActor))
+	if (EnumHasAnyFlags(SyncDirection, EMassTranslationDirection::MassToActor)
+		|| BuildContext.IsInspectingData())
 	{
 		BuildContext.AddTranslator<UMassCharacterOrientationToActorTranslator>();
 	}
@@ -178,12 +180,14 @@ void UMassAgentFeetLocationSyncTrait::BuildTemplate(FMassEntityTemplateBuildCont
 			}
 		});
 
-	if (EnumHasAnyFlags(SyncDirection, EMassTranslationDirection::ActorToMass))
+	if (EnumHasAnyFlags(SyncDirection, EMassTranslationDirection::ActorToMass)
+		|| BuildContext.IsInspectingData())
 	{
 		BuildContext.AddTranslator<UMassSceneComponentLocationToMassTranslator>();
 	}
 
-	if (EnumHasAnyFlags(SyncDirection, EMassTranslationDirection::MassToActor))
+	if (EnumHasAnyFlags(SyncDirection, EMassTranslationDirection::MassToActor)
+		|| BuildContext.IsInspectingData())
 	{
 		BuildContext.AddTranslator<UMassSceneComponentLocationToActorTranslator>();
 	}
diff --git a/Engine/Plugins/Runtime/MassGameplay/Source/MassGameplayDebug/Private/MassDebugVisualizationTrait.cpp b/Engine/Plugins/Runtime/MassGameplay/Source/MassGameplayDebug/Private/MassDebugVisualizationTrait.cpp
--- a/Engine/Plugins/Runtime/MassGameplay/Source/MassGameplayDebug/Private/MassDebugVisualizationTrait.cpp
+++ b/Engine/Plugins/Runtime/MassGameplay/Source/MassGameplayDebug/Private/MassDebugVisualizationTrait.cpp
@@ -15,19 +15,22 @@ void UMassDebugVisualizationTrait::BuildTemplate(FMassEntityTemplateBuildContext
 	const UStaticMesh* const DebugMesh = nullptr;
 #endif
 	
-	if (DebugMesh)
+	if (DebugMesh || BuildContext.IsInspectingData())
 	{
 #if WITH_EDITORONLY_DATA
 		FSimDebugVisFragment& DebugVisFragment = BuildContext.AddFragment_GetRef<FSimDebugVisFragment>();
-		UMassDebuggerSubsystem* Debugger = World.GetSubsystem<UMassDebuggerSubsystem>();
-		if (ensure(Debugger))
+		if (!BuildContext.IsInspectingData())
 		{
-			UMassDebugVisualizationComponent* DebugVisComponent = Debugger->GetVisualizationComponent();
-			if (ensure(DebugVisComponent))
+			UMassDebuggerSubsystem* Debugger = World.GetSubsystem<UMassDebuggerSubsystem>();
+			if (ensure(Debugger))
 			{
-				DebugVisFragment.VisualType = DebugVisComponent->AddDebugVisType(DebugShape);
+				UMassDebugVisualizationComponent* DebugVisComponent = Debugger->GetVisualizationComponent();
+				if (ensure(DebugVisComponent))
+				{
+					DebugVisFragment.VisualType = DebugVisComponent->AddDebugVisType(DebugShape);
+				}
+				// @todo this path requires a fragment destructor that will remove the mesh from the debugger.
 			}
-			// @todo this path requires a fragment destructor that will remove the mesh from the debugger.
 		}
 #endif // WITH_EDITORONLY_DATA
 	}
diff --git a/Engine/Plugins/Runtime/MassGameplay/Source/MassGameplayEditor/Private/MassTraitRepository.cpp b/Engine/Plugins/Runtime/MassGameplay/Source/MassGameplayEditor/Private/MassTraitRepository.cpp
new file mode 100644
--- /dev/null
+++ b/Engine/Plugins/Runtime/MassGameplay/Source/MassGameplayEditor/Private/MassTraitRepository.cpp
@@ -0,0 +1,409 @@
+// Copyright Epic Games, Inc. All Rights Reserved.
+ 
+#include "MassTraitRepository.h"
+#include "MassEntityTraitBase.h"
+#include "Engine/Engine.h"
+#include "Engine/World.h"
+#include "Misc/CoreDelegates.h"
+#include "UObject/UObjectIterator.h"
+#include "Editor.h"
+#include "MassEntityTemplateRegistry.h"
+#include "Misc/UObjectToken.h"
+#include "Framework/Docking/TabManager.h"
+#include "Logging/MessageLog.h"
+#include "MassAssortedFragmentsTrait.h"
+#include "DataValidationFixers.h"
+#include "MassEntityConfigAsset.h"
+#include "MassEntityEditor.h"
+
+#define LOCTEXT_NAMESPACE "Mass"
+
+
+namespace UE::Mass
+{
+	namespace Editor
+	{
+		TConstArrayView<FName> GetTraitsNameAddingElements(const FName ElementName)
+		{
+			if (UMassTraitRepository* TraitRepo = GEditor->GetEditorSubsystem<UMassTraitRepository>())
+			{
+				return TraitRepo->GetTraitsNameAddingElements(ElementName);
+			}
+			return TConstArrayView<FName>();
+		}
+	}
+
+	namespace Private
+	{
+		/** a helper function that wraps up code for fixing the WeakConfig by adding a trait of class WeakTraitClass */
+		FFixResult AddTraitToConfigFix(const TWeakObjectPtr<UMassEntityConfigAsset>& WeakConfig, const TWeakObjectPtr<UClass>& WeakTraitClass)
+		{
+			if (UMassEntityConfigAsset* ConfigAsset = WeakConfig.Get())
+			{
+				if (TSubclassOf<UMassEntityTraitBase> TraitClass = WeakTraitClass.Get())
+				{
+					if (ConfigAsset->AddTrait(TraitClass) != nullptr)
+					{
+						return FFixResult::Success();
+					}
+					return FFixResult::Failure(LOCTEXT("FailedToCreateTrait", "Failed to create an instance of the trait."));
+				}
+				return FFixResult::Failure(LOCTEXT("TraitClassNoLongerAvailable", "Trait class no longer available."));
+			}
+			return FFixResult::Failure(LOCTEXT("ConfigAssetNoLongerAvailable", "Config asset no longer available."));
+		}
+	}
+};
+
+//-----------------------------------------------------------------------------
+// FMassTraitInspectionContext
+//-----------------------------------------------------------------------------
+FMassTraitInspectionContext::FMassTraitInspectionContext()
+	: EntityTemplate()
+	, BuildContext(EntityTemplate)
+{
+	BuildContext.EnableDataInvestigationMode();
+}
+
+FMassTraitInspectionContext::FInvestigationContext::FInvestigationContext(FMassEntityTemplateData& InTemplate)
+	: Super(InTemplate)
+{
+}
+
+void FMassTraitInspectionContext::FInvestigationContext::SetTrait(const UMassEntityTraitBase& Trait)
+{
+	SetTraitBeingProcessed(&Trait);
+}
+
+//-----------------------------------------------------------------------------
+// UMassTraitRepository
+//-----------------------------------------------------------------------------
+TWeakObjectPtr<UWorld> UMassTraitRepository::GlobalInvestigationWorld;
+
+void UMassTraitRepository::Initialize(FSubsystemCollectionBase& Collection)
+{
+	Super::Initialize(Collection);
+
+	OnNewTraitTypeHandle = UMassEntityTraitBase::GetOnNewTraitTypeEvent().AddUObject(this, &UMassTraitRepository::OnNewTraitType);
+
+#if WITH_MASSENTITY_DEBUG
+	FMassDebugger::OnDebugEvent.AddUObject(this, &UMassTraitRepository::OnDebugEvent);
+#endif // WITH_MASSENTITY_DEBUG
+}
+
+// @todo we have an opportunity here to make it really flexible. Every message name could be associated with a 
+// dedicated TFunction (via some map) that would handle the given message type. This way users could extend or override
+// the way certain events are handled. 
+void UMassTraitRepository::OnDebugEvent(const FName EventName, FConstStructView Payload, EMassDebugMessageSeverity SeverityOverride)
+{
+	static const FName MissingTraitMessageName = FMassMissingTraitMessage::StaticStruct()->GetFName();
+	static const FName DuplicateElementsMessageName = FMassDuplicateElementsMessage::StaticStruct()->GetFName();
+
+#define OVERRIDABLE_SEVERITY(DefaultSeverity) UE::Mass::Debug::MassSeverityToMessageSeverity(DefaultSeverity, SeverityOverride)
+
+	if (EventName == MissingTraitMessageName)
+	{
+		if (const FMassMissingTraitMessage* MissingTraitMessage = Payload.GetPtr<const FMassMissingTraitMessage>())
+		{
+			FMessageLog MessageLog(UE::Mass::Editor::MessageLogPageName);
+
+			TArray<TSharedRef<FTokenizedMessage>> Messages;
+
+			const FName MissingElement = GetFNameSafe(MissingTraitMessage->MissingType);
+			const FName TraitClassName = MissingTraitMessage->RequestingTrait ? GetFNameSafe(MissingTraitMessage->RequestingTrait->GetClass()) : FName();
+
+			TSharedRef<FTokenizedMessage> IntroMessage = Messages.Add_GetRef(FTokenizedMessage::Create(OVERRIDABLE_SEVERITY(EMessageSeverity::Error)))
+				->AddToken(FAssetNameToken::Create(GetPathNameSafe(MissingTraitMessage->RequestingTrait)
+					, FText::FormatOrdered(LOCTEXT("MissingElementSuggestionHeader"
+						, "Trait {0}")
+						, FText::FromName(TraitClassName))
+				));
+
+			UMassTraitRepository* TraitRepo = GEditor->GetEditorSubsystem<UMassTraitRepository>();
+			TConstArrayView<FName> SuggestedTraitNames = TraitRepo
+				? TraitRepo->GetTraitsNameAddingElements(MissingElement)
+				: TConstArrayView<FName>();
+
+			if (SuggestedTraitNames.Num())
+			{
+				// The FixController will coordinate IFixer instances and the relevant FFixToken to
+				// ensure that only one of them can be applied. Once any of the fixes is applied
+				// the rest will become inactive (the FFixToken tokens will become grayed out and non-clickable).
+				// @todo at the moment FMutuallyExclusiveFixSet doesn't care whether fixing was successful. 
+				//		Should be relatively easy to address but needs to be coordinated with the author
+				TSharedRef<UE::DataValidation::FMutuallyExclusiveFixSet> FixController = MakeShareable(new UE::DataValidation::FMutuallyExclusiveFixSet());
+
+				IntroMessage->AddText(FText::FormatOrdered(LOCTEXT("MissingElementSuggestionOptions"
+					, "has unsatisfied dependency of {0}. The following actions can address it:")
+					, FText::FromName(MissingElement)));
+
+				for (const FName& SuggestedTraitName : SuggestedTraitNames)
+				{
+					bool bFixable = false;
+					if (MissingTraitMessage->RequestingTrait)
+					{
+						TWeakObjectPtr<UClass> WeakTraitClass = GetTraitClass(SuggestedTraitName);
+						UMassEntityConfigAsset* EntityConfigAsset = Cast<UMassEntityConfigAsset>(MissingTraitMessage->RequestingTrait->GetOuter());
+
+						if (EntityConfigAsset && WeakTraitClass.IsValid())
+						{
+							bFixable = true;
+
+							TWeakObjectPtr<UMassEntityConfigAsset> WeakConfig = EntityConfigAsset;
+							
+							// capturing FixController to make sure it exists as long as the fixes are alive. The lambda will
+							// be destroyed once the Fixer tokens get destroyed, for example during MessageLog page clearing.
+							TFunction<FFixResult()> ApplyFix = [WeakConfig, WeakTraitClass, _ = FixController]()
+								{
+									return UE::Mass::Private::AddTraitToConfigFix(WeakConfig, WeakTraitClass);
+								};
+
+							const TSharedRef<UE::DataValidation::IFixer> Fixer = UE::DataValidation::MakeFix(MoveTemp(ApplyFix));
+
+							FixController->Add(FText::FormatOrdered(LOCTEXT("AddMissingTrait", "Add {0} trait to {1} entity config")
+										, FText::FromName(SuggestedTraitName)
+										, FText::FromName(EntityConfigAsset->GetFName()))
+									, Fixer);
+						}
+					}
+					
+					if (bFixable == false)
+					{
+						// unfixable (since we're unable to determine the UMassEntityConfigAsset outer), so just report
+						Messages.Add_GetRef(FTokenizedMessage::Create(OVERRIDABLE_SEVERITY(EMessageSeverity::Info)))
+							->AddText(FText::FormatOrdered(LOCTEXT("MissingElementSuggestionUnfixable", "\t{0}")
+								, FText::FromName(SuggestedTraitName)));
+					}
+				}
+
+				// for every IFixer instance we created in the loop above the call bellow will create a FFixToken
+				// related to that specific "fix". We attach all the tokens to the initial "here are your options" message
+				FixController->CreateTokens([IntroMessage](TSharedRef<FFixToken> FixToken)
+					{
+						IntroMessage->AddToken(FixToken);
+					}
+				);
+			}
+			else
+			{
+				IntroMessage->AddText(FText::FormatOrdered(LOCTEXT("MissingElementSuggestionNoOptions"
+					, "has unsatisfied dependency of {0}. There are no registered Traits that provide the type. Try using {1}.")
+					, FText::FromName(MissingElement)
+					, FText::FromName(UMassAssortedFragmentsTrait::StaticClass()->GetFName())));
+			}
+
+			MessageLog.AddMessages(Messages);
+		}
+	}
+	else if (EventName == UE::Mass::Debug::TraitFailedValidation)
+	{
+		if (const FMassGenericDebugEvent* GenericEvent = Payload.GetPtr<const FMassGenericDebugEvent>())
+		{
+			const UMassEntityTraitBase* Trait = Cast<UMassEntityTraitBase>(GenericEvent->Context);
+			FMessageLog MessageLog(UE::Mass::Editor::MessageLogPageName);
+			MessageLog.AddMessage(FTokenizedMessage::Create(OVERRIDABLE_SEVERITY(EMessageSeverity::Error)))
+				->AddToken(FUObjectToken::Create(Trait))
+				->AddToken(FTextToken::Create(LOCTEXT("TraitFailedValidation", "trait-specific validation failed")));
+		}
+	}
+	else if (EventName == UE::Mass::Debug::TraitIgnored)
+	{
+		if (const FMassGenericDebugEvent* GenericEvent = Payload.GetPtr<const FMassGenericDebugEvent>())
+		{
+			const UMassEntityTraitBase* Trait = Cast<UMassEntityTraitBase>(GenericEvent->Context);
+			FMessageLog MessageLog(UE::Mass::Editor::MessageLogPageName);
+			MessageLog.AddMessage(FTokenizedMessage::Create(OVERRIDABLE_SEVERITY(EMessageSeverity::Warning)))
+				->AddToken(FUObjectToken::Create(Trait))
+				->AddToken(FTextToken::Create(LOCTEXT("TraitIgnoredTrait", "trait was ignored. Check if it's not a duplicate.")));
+		}
+	}
+	else if (EventName == DuplicateElementsMessageName)
+	{
+		if (const FMassDuplicateElementsMessage* DuplicateElementsMessage = Payload.GetPtr<const FMassDuplicateElementsMessage>())
+		{
+			FMessageLog MessageLog(UE::Mass::Editor::MessageLogPageName);
+			MessageLog.AddMessage(FTokenizedMessage::Create(OVERRIDABLE_SEVERITY(EMessageSeverity::Warning)))
+				->AddToken(FUObjectToken::Create(DuplicateElementsMessage->DuplicatingTrait))
+				->AddToken(FTextToken::Create(FText::FormatOrdered(
+					LOCTEXT("TraitFragmentDuplicationWarning", "trying to add fragment of type {0} while it has already been added by")
+					, FText::FromName(GetFNameSafe(DuplicateElementsMessage->Element))
+					)))
+				->AddToken(FUObjectToken::Create(DuplicateElementsMessage->OriginalTrait))
+				->AddToken(FTextToken::Create(LOCTEXT("TraitFragmentDuplicationWarningCheckConflicts", "Check your entity config for conflicting traits")));
+		}
+	}
+
+#undef OVERRIDABLE_SEVERITY
+}
+
+TConstArrayView<FName> UMassTraitRepository::GetTraitsNameAddingElements(const FName ElementName)
+{
+	if (bIsRepositoryInitialized == false)
+	{
+		InitRepository();
+	}
+
+	if (TArray<FName>* FoundTraits = ElementTypeToTraitMap.Find(ElementName))
+	{
+		return MakeArrayView(*FoundTraits);
+	}
+
+	return TConstArrayView<FName>();
+}
+
+void UMassTraitRepository::InitRepository()
+{
+	if (bIsRepositoryInitialized == true)
+	{
+		return;
+	}
+
+	UWorld::InitializationValues IVS;
+	IVS.InitializeScenes(false)
+		.AllowAudioPlayback(false)
+		.RequiresHitProxies(false)
+		.CreatePhysicsScene(false)
+		.CreateNavigation(false)
+		.CreateAISystem(false)
+		.ShouldSimulatePhysics(false)
+		.EnableTraceCollision(false)
+		.SetTransactional(false)
+		.CreateFXSystem(false);
+
+	InvestigationWorld = UWorld::CreateWorld(EWorldType::Inactive
+		, /*bInformEngineOfWorld=*/false
+		, /*WorldName=*/TEXT("MassTraitRepository_InvestigationWorld")
+		, /*Package=*/nullptr
+		, /*bAddToRoot=*/false
+		, /*InFeatureLevel=*/ERHIFeatureLevel::Num
+		, &IVS
+		, /*bInSkipInitWorld=*/true);
+	GlobalInvestigationWorld = InvestigationWorld;
+	check(InvestigationWorld);
+	InvestigationWorld->InitWorld(IVS);
+
+	// Marking as "initialized" so that the OnNewTraitType calls below do their job as expected. 
+	bIsRepositoryInitialized = true;
+
+	// Using RF_NoFlags to include CDOs (filtered out by default), since that's all we care about.
+	for (TObjectIterator<UMassEntityTraitBase> ClassIterator(/*AdditionalExclusionFlags=*/RF_NoFlags); ClassIterator; ++ClassIterator)
+	{
+		if (ClassIterator && ClassIterator->HasAnyFlags(RF_ClassDefaultObject))
+		{
+			OnNewTraitType(**ClassIterator);
+		}
+	}
+}
+
+void UMassTraitRepository::Deinitialize()
+{
+	FCoreDelegates::OnPostEngineInit.RemoveAll(this);
+	UMassEntityTraitBase::GetOnNewTraitTypeEvent().Remove(OnNewTraitTypeHandle);
+
+	if (IsValid(InvestigationWorld))
+	{
+		InvestigationWorld->DestroyWorld(/*bInformEngineOfWorld=*/false);
+	}
+
+	Super::Deinitialize();
+}
+
+void UMassTraitRepository::OnNewTraitType(UMassEntityTraitBase& Trait)
+{
+	if (bIsRepositoryInitialized == false)
+	{
+		// since the repository is not initialized yet we assume we don't need to collect the information about Trait just yet.
+		// Once InitRepository is called all existing Trait CDOs will be collected and processed.
+		return;
+	}
+
+	// simply ignore abstract classes, we don't care about these since the user will never be able to use them anyway.
+	if (Trait.GetClass()->HasAnyClassFlags(CLASS_Abstract))
+	{
+		return;
+	}
+
+	checkf(Trait.HasAnyFlags(RF_ClassDefaultObject), TEXT("Only CDOs are expected here."));
+	check(InvestigationWorld);
+
+	const FName TraitName = Trait.GetClass()->GetFName();
+	FTraitAndElements TraitData;
+
+	// first check if we have this one already. If so we need to remove and re-add in the type got updated
+	if (TraitClassNameToDataMap.RemoveAndCopyValue(TraitName, TraitData))
+	{
+		for (const FName TypeName : TraitData.ElementNames)
+		{
+			ElementTypeToTraitMap.FindChecked(TypeName).RemoveSingleSwap(TraitName, EAllowShrinking::No);
+		}
+	}
+
+	TraitData.TraitClass = Trait.GetClass();
+	TraitData.ElementNames.Reset();
+
+	FMassTraitInspectionContext InvestigationContext;
+	InvestigationContext.BuildContext.SetTrait(Trait);
+	Trait.BuildTemplate(InvestigationContext.BuildContext, *InvestigationWorld);
+
+	const FMassArchetypeCompositionDescriptor& Composition = InvestigationContext.EntityTemplate.GetCompositionDescriptor();
+
+	for (FMassFragmentBitSet::FIndexIterator It = Composition.Fragments.GetIndexIterator(); It; ++It)
+	{
+		const FName StructName = Composition.Fragments.DebugGetStructTypeName(*It);
+		TraitData.ElementNames.Add(StructName);
+	}
+
+	for (FMassTagBitSet::FIndexIterator It = Composition.Tags.GetIndexIterator(); It; ++It)
+	{
+		const FName StructName = Composition.Tags.DebugGetStructTypeName(*It);
+		TraitData.ElementNames.Add(StructName);
+	}
+
+	for (FMassChunkFragmentBitSet::FIndexIterator It = Composition.ChunkFragments.GetIndexIterator(); It; ++It)
+	{
+		const FName StructName = Composition.ChunkFragments.DebugGetStructTypeName(*It);
+		TraitData.ElementNames.Add(StructName);
+	}
+
+	for (FMassSharedFragmentBitSet::FIndexIterator It = Composition.SharedFragments.GetIndexIterator(); It; ++It)
+	{
+		const FName StructName = Composition.SharedFragments.DebugGetStructTypeName(*It);
+		TraitData.ElementNames.Add(StructName);
+	}
+
+	for (FMassConstSharedFragmentBitSet::FIndexIterator It = Composition.ConstSharedFragments.GetIndexIterator(); It; ++It)
+	{
+		const FName StructName = Composition.ConstSharedFragments.DebugGetStructTypeName(*It);
+		TraitData.ElementNames.Add(StructName);
+	}
+
+	for (const FName ElementTypeName : TraitData.ElementNames)
+	{
+		ElementTypeToTraitMap.FindOrAdd(ElementTypeName).Add(TraitName);
+	}
+	TraitClassNameToDataMap.Add(TraitName, MoveTemp(TraitData));
+}
+
+UWorld* UMassTraitRepository::GetInvestigationWorld()
+{
+	return GlobalInvestigationWorld.Get();
+}
+
+TWeakObjectPtr<UClass> UMassTraitRepository::GetTraitClass(const FName TraitClassName) const
+{
+	if (const FTraitAndElements* Data = TraitClassNameToDataMap.Find(TraitClassName))
+	{
+		return Data->TraitClass;
+	}
+	return TWeakObjectPtr<UClass>();
+}
+
+//-----------------------------------------------------------------------------
+// UMassDebugEntitySubsystem
+//-----------------------------------------------------------------------------
+bool UMassDebugEntitySubsystem::ShouldCreateSubsystem(UObject* Outer) const
+{
+	return Outer != nullptr && Outer == UMassTraitRepository::GetInvestigationWorld();
+}
+
+#undef LOCTEXT_NAMESPACE 
\ No newline at end of file
diff --git a/Engine/Plugins/Runtime/MassGameplay/Source/MassGameplayEditor/Public/MassTraitRepository.h b/Engine/Plugins/Runtime/MassGameplay/Source/MassGameplayEditor/Public/MassTraitRepository.h
new file mode 100644
--- /dev/null
+++ b/Engine/Plugins/Runtime/MassGameplay/Source/MassGameplayEditor/Public/MassTraitRepository.h
@@ -0,0 +1,97 @@
+// Copyright Epic Games, Inc. All Rights Reserved.
+
+#pragma once
+
+#include "MassEntityTemplateRegistry.h"
+#include "EditorSubsystem.h"
+#include "MassEntitySubsystem.h"
+#include "StructUtils/StructView.h"
+#include "MassDebugger.h"
+#include "MassTraitRepository.generated.h"
+
+class UMassEntityTraitBase;
+class UWorld;
+
+namespace UE::Mass::Editor
+{
+	extern MASSGAMEPLAYEDITOR_API TConstArrayView<FName> GetTraitsNameAddingElements(const FName ElementName);
+}
+
+struct FMassTraitInspectionContext
+{
+	struct FInvestigationContext : public FMassEntityTemplateBuildContext
+	{
+		using Super = FMassEntityTemplateBuildContext;
+		explicit FInvestigationContext(FMassEntityTemplateData& InTemplate);
+		void SetTrait(const UMassEntityTraitBase& Trait);
+	};
+
+	FMassTraitInspectionContext();
+
+	FMassEntityTemplateData EntityTemplate;
+	FInvestigationContext BuildContext;
+};
+
+/** 
+ * Subsystem to store information about Mass traits so that we can make helpful suggestions to users like which
+ * traits supply a fragment they need.
+ * It also serves to post trait validation information and fix options to the MessageLog.
+ */
+UCLASS()
+class MASSGAMEPLAYEDITOR_API UMassTraitRepository : public UEditorSubsystem
+{
+	GENERATED_BODY()
+	
+public:
+	static UWorld* GetInvestigationWorld();
+
+	TConstArrayView<FName> GetTraitsNameAddingElements(const FName ElementName);
+	TWeakObjectPtr<UClass> GetTraitClass(const FName TraitClassName) const;
+
+protected:
+	virtual void Initialize(FSubsystemCollectionBase& Collection) override;
+	virtual void Deinitialize() override;
+
+	/** 
+	 * Note that the function does nothing until InitRepository is called. This is done to avoid collecting trait
+	 * data until it's actually necessary.
+	 */
+	void OnNewTraitType(UMassEntityTraitBase& Trait);
+
+	void OnDebugEvent(const FName EventName, FConstStructView Payload, EMassDebugMessageSeverity SeverityOverride);
+
+	/** 
+	 * The method gathers all existing trait classes and processed them. We don't do that on subsystem's init since
+	 * very often the data won't be needed during the given editor run. We call the function lazily the first time 
+	 * the data is needed.
+	 */
+	void InitRepository();
+
+	FDelegateHandle OnNewTraitTypeHandle;
+	struct FTraitAndElements
+	{
+		TWeakObjectPtr<UClass> TraitClass;
+		TArray<FName> ElementNames;
+	};
+	TMap<FName, FTraitAndElements> TraitClassNameToDataMap;
+	TMap<FName, TArray<FName>> ElementTypeToTraitMap;
+
+
+	/** the World we use to host all the subsystems required to process traits */
+	UPROPERTY()
+	TObjectPtr<UWorld> InvestigationWorld;
+
+	/** We initialize the repository's data lazily and this property indicates whether it has been already done. */
+	bool bIsRepositoryInitialized;
+
+	static TWeakObjectPtr<UWorld> GlobalInvestigationWorld;
+};
+
+
+UCLASS()
+class UMassDebugEntitySubsystem : public UMassEntitySubsystem
+{
+	GENERATED_BODY()
+	
+	virtual bool ShouldCreateSubsystem(UObject* Outer) const override;
+};
diff --git a/Engine/Plugins/Runtime/MassGameplay/Source/MassLOD/Private/MassLODTrait.cpp b/Engine/Plugins/Runtime/MassGameplay/Source/MassLOD/Private/MassLODTrait.cpp
--- a/Engine/Plugins/Runtime/MassGameplay/Source/MassLOD/Private/MassLODTrait.cpp
+++ b/Engine/Plugins/Runtime/MassGameplay/Source/MassLOD/Private/MassLODTrait.cpp
@@ -39,7 +39,7 @@ void UMassSimulationLODTrait::BuildTemplate(FMassEntityTemplateBuildContext& Bui
 	FMassSimulationLODFragment& LODFragment = BuildContext.AddFragment_GetRef<FMassSimulationLODFragment>();
 
 	// Start all simulation LOD in the Off 
-	if (Params.bSetLODTags || bEnableVariableTicking)
+	if (Params.bSetLODTags || bEnableVariableTicking || BuildContext.IsInspectingData())
 	{
 		LODFragment.LOD = EMassLOD::Off;
 		BuildContext.AddTag<FMassOffLODTag>();
@@ -55,7 +55,7 @@ void UMassSimulationLODTrait::BuildTemplate(FMassEntityTemplateBuildContext& Bui
 	BuildContext.AddSharedFragment(SharedFragment);
 
 	// Variable ticking from simulation LOD
-	if (bEnableVariableTicking)
+	if (bEnableVariableTicking || BuildContext.IsInspectingData())
 	{
 		BuildContext.AddFragment<FMassSimulationVariableTickFragment>();
 		BuildContext.AddChunkFragment<FMassSimulationVariableTickChunkFragment>();
diff --git a/Engine/Plugins/Runtime/MassGameplay/Source/MassReplication/Private/MassReplicationTrait.cpp b/Engine/Plugins/Runtime/MassGameplay/Source/MassReplication/Private/MassReplicationTrait.cpp
--- a/Engine/Plugins/Runtime/MassGameplay/Source/MassReplication/Private/MassReplicationTrait.cpp
+++ b/Engine/Plugins/Runtime/MassGameplay/Source/MassReplication/Private/MassReplicationTrait.cpp
@@ -12,7 +12,7 @@
 
 void UMassReplicationTrait::BuildTemplate(FMassEntityTemplateBuildContext& BuildContext, const UWorld& World) const
 {
-	if (World.IsNetMode(NM_Standalone))
+	if (World.IsNetMode(NM_Standalone) && !BuildContext.IsInspectingData())
 	{
 		return;
 	}
@@ -29,12 +29,21 @@ void UMassReplicationTrait::BuildTemplate(FMassEntityTemplateBuildContext& Build
 	FMassEntityManager& EntityManager = UE::Mass::Utils::GetEntityManagerChecked(World);
 
 	UMassReplicationSubsystem* ReplicationSubsystem = UWorld::GetSubsystem<UMassReplicationSubsystem>(&World);
-	check(ReplicationSubsystem);
 
 	FConstSharedStruct ParamsFragment = EntityManager.GetOrCreateConstSharedFragment(Params);
 	BuildContext.AddConstSharedFragment(ParamsFragment);
 
 	uint32 ParamsHash = UE::StructUtils::GetStructCrc32(FConstStructView::Make(Params));
-	FSharedStruct SharedFragment = EntityManager.GetOrCreateSharedFragmentByHash<FMassReplicationSharedFragment>(ParamsHash, *ReplicationSubsystem, Params);
-	BuildContext.AddSharedFragment(SharedFragment);
+	if (LIKELY(!BuildContext.IsInspectingData()))
+	{
+		check(ReplicationSubsystem);
+		FSharedStruct SharedFragment = EntityManager.GetOrCreateSharedFragmentByHash<FMassReplicationSharedFragment>(ParamsHash, *ReplicationSubsystem, Params);
+		BuildContext.AddSharedFragment(SharedFragment);
+	}
+	else
+	{
+		// in the investigation mode we only care about the fragment type
+		FSharedStruct SharedFragment = EntityManager.GetOrCreateSharedFragmentByHash<FMassReplicationSharedFragment>(ParamsHash);
+		BuildContext.AddSharedFragment(SharedFragment);
+	}
 }
\ No newline at end of file
diff --git a/Engine/Plugins/Runtime/MassGameplay/Source/MassRepresentation/Private/MassDistanceVisualizationTrait.cpp b/Engine/Plugins/Runtime/MassGameplay/Source/MassRepresentation/Private/MassDistanceVisualizationTrait.cpp
--- a/Engine/Plugins/Runtime/MassGameplay/Source/MassRepresentation/Private/MassDistanceVisualizationTrait.cpp
+++ b/Engine/Plugins/Runtime/MassGameplay/Source/MassRepresentation/Private/MassDistanceVisualizationTrait.cpp
@@ -37,7 +37,8 @@ UMassDistanceVisualizationTrait::UMassDistanceVisualizationTrait()
 void UMassDistanceVisualizationTrait::BuildTemplate(FMassEntityTemplateBuildContext& BuildContext, const UWorld& World) const
 {
 	// This should not be ran on NM_Server network mode
-	if (World.IsNetMode(NM_DedicatedServer) && !bAllowServerSideVisualization)
+	if (World.IsNetMode(NM_DedicatedServer) && !bAllowServerSideVisualization
+		&& !BuildContext.IsInspectingData())
 	{
 		return;
 	}
@@ -49,7 +50,7 @@ void UMassDistanceVisualizationTrait::BuildTemplate(FMassEntityTemplateBuildCont
 	FMassEntityManager& EntityManager = UE::Mass::Utils::GetEntityManagerChecked(World);
 
 	UMassRepresentationSubsystem* RepresentationSubsystem = Cast<UMassRepresentationSubsystem>(World.GetSubsystemBase(RepresentationSubsystemClass));
-	if (RepresentationSubsystem == nullptr)
+	if (RepresentationSubsystem == nullptr && !BuildContext.IsInspectingData())
 	{
 		UE_LOG(LogMassRepresentation, Error, TEXT("Expecting a valid class for the representation subsystem"));
 		RepresentationSubsystem = UWorld::GetSubsystem<UMassRepresentationSubsystem>(&World);
@@ -62,7 +63,7 @@ void UMassDistanceVisualizationTrait::BuildTemplate(FMassEntityTemplateBuildCont
 	FSharedStruct SubsystemFragment = EntityManager.GetOrCreateSharedFragmentByHash<FMassRepresentationSubsystemSharedFragment>(SubsystemHash, SubsystemSharedFragment);
 	BuildContext.AddSharedFragment(SubsystemFragment);
 
-	if (!Params.RepresentationActorManagementClass)
+	if (!Params.RepresentationActorManagementClass && !BuildContext.IsInspectingData())
 	{
 		UE_LOG(LogMassRepresentation, Error, TEXT("Expecting a valid class for the representation actor management"));
 	}
@@ -71,12 +72,15 @@ void UMassDistanceVisualizationTrait::BuildTemplate(FMassEntityTemplateBuildCont
 	BuildContext.AddConstSharedFragment(ParamsFragment);
 
 	FMassRepresentationFragment& RepresentationFragment = BuildContext.AddFragment_GetRef<FMassRepresentationFragment>();
-	if (bRegisterStaticMeshDesc)
+	if (LIKELY(!BuildContext.IsInspectingData()))
 	{
-		RepresentationFragment.StaticMeshDescHandle = RepresentationSubsystem->FindOrAddStaticMeshDesc(StaticMeshInstanceDesc);
+		if (bRegisterStaticMeshDesc && !BuildContext.IsInspectingData())
+		{
+			RepresentationFragment.StaticMeshDescHandle = RepresentationSubsystem->FindOrAddStaticMeshDesc(StaticMeshInstanceDesc);
+		}
+		RepresentationFragment.HighResTemplateActorIndex = HighResTemplateActor.Get() ? RepresentationSubsystem->FindOrAddTemplateActor(HighResTemplateActor.Get()) : INDEX_NONE;
+		RepresentationFragment.LowResTemplateActorIndex = LowResTemplateActor.Get() ? RepresentationSubsystem->FindOrAddTemplateActor(LowResTemplateActor.Get()) : INDEX_NONE;
 	}
-	RepresentationFragment.HighResTemplateActorIndex = HighResTemplateActor.Get() ? RepresentationSubsystem->FindOrAddTemplateActor(HighResTemplateActor.Get()) : INDEX_NONE;
-	RepresentationFragment.LowResTemplateActorIndex = LowResTemplateActor.Get() ? RepresentationSubsystem->FindOrAddTemplateActor(LowResTemplateActor.Get()) : INDEX_NONE;
 
 	FConstSharedStruct LODParamsFragment = EntityManager.GetOrCreateConstSharedFragment(LODParams);
 	BuildContext.AddConstSharedFragment(LODParamsFragment);
diff --git a/Engine/Plugins/Runtime/MassGameplay/Source/MassRepresentation/Private/MassMovableVisualizationTrait.cpp b/Engine/Plugins/Runtime/MassGameplay/Source/MassRepresentation/Private/MassMovableVisualizationTrait.cpp
--- a/Engine/Plugins/Runtime/MassGameplay/Source/MassRepresentation/Private/MassMovableVisualizationTrait.cpp
+++ b/Engine/Plugins/Runtime/MassGameplay/Source/MassRepresentation/Private/MassMovableVisualizationTrait.cpp
@@ -6,7 +6,7 @@
 
 void UMassMovableVisualizationTrait::BuildTemplate(FMassEntityTemplateBuildContext& BuildContext, const UWorld& World) const
 {
-	if (World.IsNetMode(NM_DedicatedServer))
+	if (World.IsNetMode(NM_DedicatedServer) && !BuildContext.IsInspectingData())
 	{
 		return;
 	}
diff --git a/Engine/Plugins/Runtime/MassGameplay/Source/MassRepresentation/Private/MassVisualizationTrait.cpp b/Engine/Plugins/Runtime/MassGameplay/Source/MassRepresentation/Private/MassVisualizationTrait.cpp
--- a/Engine/Plugins/Runtime/MassGameplay/Source/MassRepresentation/Private/MassVisualizationTrait.cpp
+++ b/Engine/Plugins/Runtime/MassGameplay/Source/MassRepresentation/Private/MassVisualizationTrait.cpp
@@ -56,7 +56,8 @@ UMassVisualizationTrait::UMassVisualizationTrait()
 void UMassVisualizationTrait::BuildTemplate(FMassEntityTemplateBuildContext& BuildContext, const UWorld& World) const
 {
 	// This should not be ran on NM_Server network mode
-	if (World.IsNetMode(NM_DedicatedServer) && !bAllowServerSideVisualization)
+	if (World.IsNetMode(NM_DedicatedServer) && !bAllowServerSideVisualization 
+		&& !BuildContext.IsInspectingData())
 	{
 		return;
 	}
@@ -68,7 +69,7 @@ void UMassVisualizationTrait::BuildTemplate(FMassEntityTemplateBuildContext& Bui
 	FMassEntityManager& EntityManager = UE::Mass::Utils::GetEntityManagerChecked(World);
 
 	UMassRepresentationSubsystem* RepresentationSubsystem = Cast<UMassRepresentationSubsystem>(World.GetSubsystemBase(RepresentationSubsystemClass));
-	if (RepresentationSubsystem == nullptr)
+	if (RepresentationSubsystem == nullptr && !BuildContext.IsInspectingData())
 	{
 		UE_LOG(LogMassRepresentation, Error, TEXT("Expecting a valid class for the representation subsystem"));
 		RepresentationSubsystem = UWorld::GetSubsystem<UMassRepresentationSubsystem>(&World);
@@ -87,13 +88,16 @@ void UMassVisualizationTrait::BuildTemplate(FMassEntityTemplateBuildContext& Bui
 	}
 
 	FMassRepresentationFragment& RepresentationFragment = BuildContext.AddFragment_GetRef<FMassRepresentationFragment>();
-	RepresentationFragment.HighResTemplateActorIndex = HighResTemplateActor.Get() ? RepresentationSubsystem->FindOrAddTemplateActor(HighResTemplateActor.Get()) : INDEX_NONE;
-	RepresentationFragment.LowResTemplateActorIndex = LowResTemplateActor.Get() ? RepresentationSubsystem->FindOrAddTemplateActor(LowResTemplateActor.Get()) : INDEX_NONE;
+	if (LIKELY(BuildContext.IsInspectingData() == false))
+	{
+		RepresentationFragment.HighResTemplateActorIndex = HighResTemplateActor.Get() ? RepresentationSubsystem->FindOrAddTemplateActor(HighResTemplateActor.Get()) : INDEX_NONE;
+		RepresentationFragment.LowResTemplateActorIndex = LowResTemplateActor.Get() ? RepresentationSubsystem->FindOrAddTemplateActor(LowResTemplateActor.Get()) : INDEX_NONE;
+	}
 
 	bool bStaticMeshDescriptionValid = StaticMeshInstanceDesc.IsValid();
 	if (bStaticMeshDescriptionValid)
 	{
-		if (bRegisterStaticMeshDesc)
+		if (bRegisterStaticMeshDesc && !BuildContext.IsInspectingData())
 		{
 			RepresentationFragment.StaticMeshDescHandle = RepresentationSubsystem->FindOrAddStaticMeshDesc(StaticMeshInstanceDesc);
 			ensureMsgf(RepresentationFragment.StaticMeshDescHandle.IsValid()
@@ -159,9 +163,9 @@ void UMassVisualizationTrait::Serialize(FArchive& Ar)
 #endif // WITH_EDITOR
 }
 
-bool UMassVisualizationTrait::ValidateTemplate(FMassEntityTemplateBuildContext& BuildContext, const UWorld& World) const
+bool UMassVisualizationTrait::ValidateTemplate(const FMassEntityTemplateBuildContext& BuildContext, const UWorld& World, FAdditionalTraitRequirements& OutTraitRequirements) const
 {
-	Super::ValidateTemplate(BuildContext, World);
+	Super::ValidateTemplate(BuildContext, World, OutTraitRequirements);
 
 #if WITH_EDITOR
 	return ValidateParams();
diff --git a/Engine/Plugins/Runtime/MassGameplay/Source/MassRepresentation/Public/MassVisualizationTrait.h b/Engine/Plugins/Runtime/MassGameplay/Source/MassRepresentation/Public/MassVisualizationTrait.h
--- a/Engine/Plugins/Runtime/MassGameplay/Source/MassRepresentation/Public/MassVisualizationTrait.h
+++ b/Engine/Plugins/Runtime/MassGameplay/Source/MassRepresentation/Public/MassVisualizationTrait.h
@@ -65,7 +65,7 @@ protected:
 	bool bRegisterStaticMeshDesc = true;
 
 	virtual void BuildTemplate(FMassEntityTemplateBuildContext& BuildContext, const UWorld& World) const override;
-	virtual bool ValidateTemplate(FMassEntityTemplateBuildContext& BuildContext, const UWorld& World) const override;
+	virtual bool ValidateTemplate(const FMassEntityTemplateBuildContext& BuildContext, const UWorld& World, FAdditionalTraitRequirements& OutTraitRequirements) const override;
 
 	/** 
 	 * Tests whether StaticMeshInstanceDesc is valid and if not cleans up InOutParamss of EMassRepresentationType::StaticMeshInstance
diff --git a/Engine/Plugins/Runtime/MassGameplay/Source/MassSpawner/Private/MassEntityConfigAsset.cpp b/Engine/Plugins/Runtime/MassGameplay/Source/MassSpawner/Private/MassEntityConfigAsset.cpp
--- a/Engine/Plugins/Runtime/MassGameplay/Source/MassSpawner/Private/MassEntityConfigAsset.cpp
+++ b/Engine/Plugins/Runtime/MassGameplay/Source/MassSpawner/Private/MassEntityConfigAsset.cpp
@@ -13,6 +13,7 @@
 #include "Editor.h"
 #include "Framework/Notifications/NotificationManager.h"
 #include "Widgets/Notifications/SNotificationList.h"
+#include "MassEntityEditor.h"
 #endif // WITH_EDITOR
 
 #define LOCTEXT_NAMESPACE "Mass"
@@ -32,7 +33,7 @@ FMassEntityConfig::FMassEntityConfig(UObject& InOwner)
 	ConfigGuid = FGuid::NewGuid();
 }
 
-const UMassEntityTraitBase* FMassEntityConfig::FindTrait(TSubclassOf<UMassEntityTraitBase> TraitClass, const bool bExactMatch) const
+UMassEntityTraitBase* FMassEntityConfig::FindTraitInternal(TSubclassOf<UMassEntityTraitBase> TraitClass, const bool bExactMatch) const
 {
 	for (const TObjectPtr<UMassEntityTraitBase>& Trait : Traits)
 	{
@@ -42,7 +43,7 @@ const UMassEntityTraitBase* FMassEntityConfig::FindTrait(TSubclassOf<UMassEntity
 		}
 	}
 
-	return Parent ? Parent->FindTrait(TraitClass, bExactMatch) :  nullptr;
+	return Parent ? Parent->GetConfig().FindTraitInternal(TraitClass, bExactMatch) : nullptr;
 }
 
 const FMassEntityTemplate& FMassEntityConfig::GetOrCreateEntityTemplate(const UWorld& World) const
@@ -222,6 +223,11 @@ void FMassEntityConfig::PostDuplicate(const bool bDuplicateForPIE)
 		ConfigGuid = FGuid::NewGuid();
 	}
 }
+
+UMassEntityTraitBase* FMassEntityConfig::FindMutableTrait(TSubclassOf<UMassEntityTraitBase> TraitClass, const bool bExactMatch)
+{
+	return FindTraitInternal(TraitClass, bExactMatch);
+}
 #endif // WITH_EDITOR
 
 //-----------------------------------------------------------------------------
@@ -239,24 +245,33 @@ void UMassEntityConfigAsset::ValidateEntityConfig()
 {
 	if (UWorld* EditorWorld = GEditor->GetEditorWorldContext().World())
 	{
+		FMessageLog MessageLog(UE::Mass::Editor::MessageLogPageName);
+		MessageLog.NewPage(FText::FromName(UE::Mass::Editor::MessageLogPageName));
+
 		if (Config.ValidateEntityTemplate(*EditorWorld))
 		{
-			const FText InfoText = LOCTEXT("MassEntityConfigAssetNoErrorsDetected", "There were no errors nor warnings detected during validation of the EntityConfigAsset");
+			FMassEditorNotification Notification;
+			Notification.Message = FText::FormatOrdered(LOCTEXT("MassEntityConfigAssetNoErrorsDetected", "There were no errors detected during validation of {0}")
+				, FText::FromName(GetFName()));
+			Notification.Severity = EMessageSeverity::Info;
+			Notification.Show();
+		}
+	}
+}
 
-			FMessageLog EditorInfo("MassEntity");
-			EditorInfo.Info(InfoText);
+UMassEntityTraitBase* UMassEntityConfigAsset::AddTrait(TSubclassOf<UMassEntityTraitBase> TraitClass)
+{
+	check(TraitClass);
 
-			FNotificationInfo Info(InfoText);
-			Info.bFireAndForget = true;
-			Info.bUseThrobber = false;
-			Info.FadeOutDuration = 0.5f;
-			Info.ExpireDuration = 5.0f;
-			if (TSharedPtr<SNotificationItem> Notification = FSlateNotificationManager::Get().AddNotification(Info))
-			{
-				Notification->SetCompletionState(SNotificationItem::CS_Success);
-			}
-		}
+	UMassEntityTraitBase* TraitInstance = Config.FindMutableTrait(TraitClass, /*bExactMatch=*/true);
+	if (TraitInstance == nullptr)
+	{
+		TraitInstance = NewObject<UMassEntityTraitBase>(this, TraitClass);
+		check(TraitInstance);
+		Config.AddTrait(*TraitInstance);
+		Modify();
 	}
+	return TraitInstance;
 }
 #endif // WITH_EDITOR
 
diff --git a/Engine/Plugins/Runtime/MassGameplay/Source/MassSpawner/Private/MassEntityTemplateRegistry.cpp b/Engine/Plugins/Runtime/MassGameplay/Source/MassSpawner/Private/MassEntityTemplateRegistry.cpp
--- a/Engine/Plugins/Runtime/MassGameplay/Source/MassSpawner/Private/MassEntityTemplateRegistry.cpp
+++ b/Engine/Plugins/Runtime/MassGameplay/Source/MassSpawner/Private/MassEntityTemplateRegistry.cpp
@@ -5,20 +5,37 @@
 #include "MassEntityManager.h"
 #include "Engine/World.h"
 #include "VisualLogger/VisualLogger.h"
-#include "HAL/IConsoleManager.h"
 #include "MassSpawnerSubsystem.h"
 #include "MassEntityTypes.h"
 #include "MassEntityTraitBase.h"
-#include "Logging/MessageLog.h"
 
 #if WITH_EDITOR
 #include "Editor.h"
-#include "Misc/UObjectToken.h"
-#include "Framework/Docking/TabManager.h"
+#include "MassDebugger.h"
+#include "MassEntityEditor.h"
 #endif
 
 #define LOCTEXT_NAMESPACE "Mass"
 
+namespace UE::Mass::Debug
+{
+	const FName TraitFailedValidation(TEXT("TraitFailedValidation"));
+	const FName TraitIgnored(TEXT("TraitIgnored"));
+
+	static bool bReportDuplicatedFragmentsAsWarnings = false;
+	// anonymous namespace to force CVars's uniqueness - we use the same name in many places, sometimes withing same namespaces
+	namespace 
+	{
+		FAutoConsoleVariableRef CVars[] =
+		{
+			{ TEXT("mass.template.DuplicateElementsAsWarnings")
+				, bReportDuplicatedFragmentsAsWarnings
+				, TEXT("Whether to report a detection of a given element type being added by multiple traits as a Warning. Otherise we print the information out as `Info`")
+				, ECVF_Cheat}
+		};
+	}
+}
+
 //----------------------------------------------------------------------//
 // FMassEntityTemplateRegistry 
 //----------------------------------------------------------------------//
@@ -128,17 +145,27 @@ bool FMassEntityTemplateBuildContext::SetTraitBeingProcessed(const UMassEntityTr
 
 bool FMassEntityTemplateBuildContext::ValidateBuildContext(const UWorld& World)
 {
-#define WITH_MESSAGES (WITH_UNREAL_DEVELOPER_TOOLS && WITH_EDITOR)
-#if WITH_MESSAGES
-	TArray<TSharedRef<FTokenizedMessage>> Messages;
+#if WITH_UNREAL_DEVELOPER_TOOLS && WITH_EDITOR
 #define IF_MESSAGES(Message) if (GEditor) { Message }
 #else
 #define IF_MESSAGES(_)
-#endif // WITH_MESSAGES
+#endif
 
 	int32 ErrorCount = 0;
 	int32 WarningCount = 0;
 
+	// Doing the trait-specific validation first since it can add required elements to the build context
+	for (FTraitData& TraitData : TraitsData)
+	{
+		UMassEntityTraitBase::FAdditionalTraitRequirements TraitRequirementsWrapper(TraitData.TypesRequired);
+		if (LIKELY(TraitData.Trait) && TraitData.Trait->ValidateTemplate(*this, World, TraitRequirementsWrapper) == false)
+		{
+			++ErrorCount;
+			IF_MESSAGES(
+				FMassDebugger::DebugEvent(UE::Mass::Debug::TraitFailedValidation, FInstancedStruct::Make<FMassGenericDebugEvent>(TraitData.Trait));
+			);
+		}
+	}
 
 	TMap<const UStruct*, const UMassEntityTraitBase*> TypesAlreadyAdded;
 
@@ -152,19 +179,18 @@ bool FMassEntityTemplateBuildContext::ValidateBuildContext(const UWorld& World)
 			const UMassEntityTraitBase*& SourceTrait = TypesAlreadyAdded.FindOrAdd(TypeAdded);
 			if (SourceTrait != nullptr)
 			{
-				// we report this only if it wasn't added twice by the same trait, the one we're processing right now
-				UE_CLOG(SourceTrait != TraitData.Trait
-					, LogMass, Warning, TEXT("%s: Fragment %s already added by %s")
-					, *GetNameSafe(TraitData.Trait), *GetNameSafe(TypeAdded), *SourceTrait->GetName());
-				++WarningCount;
-
+				if (UE::Mass::Debug::bReportDuplicatedFragmentsAsWarnings)
+				{
+					// we report this only if it wasn't added twice by the same trait, the one we're processing right now
+					UE_CLOG(SourceTrait != TraitData.Trait
+						, LogMass, Warning, TEXT("%s: Fragment %s already added by %s. Check the entity config for conflicting traits")
+						, *GetNameSafe(TraitData.Trait), *GetNameSafe(TypeAdded), *SourceTrait->GetName());
+					++WarningCount;
+				}
 				IF_MESSAGES(
-					Messages.Add_GetRef(FTokenizedMessage::Create(EMessageSeverity::Warning))
-						->AddToken(FUObjectToken::Create(TraitData.Trait))
-						->AddToken(FTextToken::Create(LOCTEXT("MassEntityTraitFragmentDuplicationWarning1", "trying to add fragment of type")))
-						->AddToken(FUObjectToken::Create(TypeAdded))
-						->AddToken(FTextToken::Create(LOCTEXT("MassEntityTraitFragmentDuplicationWarning2", "while it has already been added by")))
-						->AddToken(FUObjectToken::Create(SourceTrait));
+					FMassDebugger::DebugEvent(FMassDuplicateElementsMessage::StaticStruct()->GetFName()
+						, FInstancedStruct::Make<FMassDuplicateElementsMessage>(TraitData.Trait, SourceTrait, TypeAdded)
+						, UE::Mass::Debug::bReportDuplicatedFragmentsAsWarnings ? EMassDebugMessageSeverity::Warning : EMassDebugMessageSeverity::Info);
 				);
 			}
 			else
@@ -185,34 +211,17 @@ bool FMassEntityTemplateBuildContext::ValidateBuildContext(const UWorld& World)
 					, *GetNameSafe(TraitData.Trait), *GetNameSafe(TypeRequired));
 				++ErrorCount;
 				IF_MESSAGES(
-					Messages.Add_GetRef(FTokenizedMessage::Create(EMessageSeverity::Error))
-						->AddToken(FUObjectToken::Create(TraitData.Trait))
-						->AddToken(FTextToken::Create(LOCTEXT("MassEntityTraitMissingDependencies", "unsatisfied dependency, missing")))
-						->AddToken(FUObjectToken::Create(TypeRequired));
-				);
+				{
+					FMassDebugger::DebugEvent<FMassMissingTraitMessage>(TraitData.Trait, TypeRequired);
+				});
 			}
 		}
 	}
 
-	for (const FTraitData& TraitData : TraitsData)
-	{
-		if (TraitData.Trait && TraitData.Trait->ValidateTemplate(*this, World) == false)
-		{
-			++ErrorCount;
-			IF_MESSAGES(
-				Messages.Add_GetRef(FTokenizedMessage::Create(EMessageSeverity::Error))
-					->AddToken(FUObjectToken::Create(TraitData.Trait))
-					->AddToken(FTextToken::Create(LOCTEXT("MassEntityTraitFailedValidation", "trait-specific validation failed")));
-			);
-		}
-	}
-
 	for (const UMassEntityTraitBase* IgnoredTrait : IgnoredTraits)
 	{
 		IF_MESSAGES(
-			Messages.Add_GetRef(FTokenizedMessage::Create(EMessageSeverity::Warning))
-				->AddToken(FUObjectToken::Create(IgnoredTrait))
-				->AddToken(FTextToken::Create(LOCTEXT("MassEntityTraitIgnoredTrait", "trait was ignored. Check if it's not a duplicate.")));
+			FMassDebugger::DebugEvent(UE::Mass::Debug::TraitIgnored, FInstancedStruct::Make<FMassGenericDebugEvent>(IgnoredTrait));
 		);
 		++WarningCount;
 	}
@@ -222,30 +231,18 @@ bool FMassEntityTemplateBuildContext::ValidateBuildContext(const UWorld& World)
 	// Could be done with a complex requirements system (similar to entity queries - "all of X", "any of Y", etc) - probably 
 	// not worth it since we don't even have a use case for it right now.
 
-#if WITH_MESSAGES
-	if (GEditor && Messages.Num())
+#if WITH_UNREAL_DEVELOPER_TOOLS && WITH_EDITOR
+	if (GEditor && (ErrorCount || WarningCount))
 	{
-		TSharedRef<FTokenizedMessage> SummaryMessage = Messages.Add_GetRef(
-			FTokenizedMessage::Create(ErrorCount ? EMessageSeverity::Error : EMessageSeverity::Warning))
-			->AddToken(FTextToken::Create(FText::FormatOrdered(LOCTEXT("MassEntityTraitResult", "Mass Entity Template validation:\n{0} errors and {1} warnings found"), ErrorCount, WarningCount)));
-
-		Messages.Add_GetRef(FTokenizedMessage::Create(EMessageSeverity::Info))
-			->AddToken(FActionToken::Create(LOCTEXT("MassSeeLogForDetails", "See the log for more details.")
-				, LOCTEXT("MassSeeLogForDetailsTooltip", "Open the Output Log tab.")
-				, FOnActionTokenExecuted::CreateLambda([]()
-					{
-						FGlobalTabmanager::Get()->TryInvokeTab(FName("OutputLog"));
-					}))
-				);
-
-		FMessageLog EditorErrors("MassEntity");
-		EditorErrors.AddMessages(Messages);
-		EditorErrors.Notify(SummaryMessage->ToText());
+		FMassEditorNotification Notification;
+		Notification.Message = FText::FormatOrdered(LOCTEXT("TraitResult", "Mass Entity Template validation:\n{0} errors and {1} warnings found"), ErrorCount, WarningCount);
+		Notification.Severity = ErrorCount ? EMessageSeverity::Error : EMessageSeverity::Warning;
+		Notification.bIncludeSeeOutputLogForDetails = true;
+		Notification.Show();
 	}
-#endif // WITH_MESSAGES
+#endif // WITH_UNREAL_DEVELOPER_TOOLS && WITH_EDITOR
 
 #undef IF_MESSAGES
-#undef WITH_MESSAGES
 
 	// only the Errors render the template invalid, Warnings just warn about stuff not being set up quite right, but we can recover.
 	return (ErrorCount == 0);
diff --git a/Engine/Plugins/Runtime/MassGameplay/Source/MassSpawner/Private/MassEntityTraitBase.cpp b/Engine/Plugins/Runtime/MassGameplay/Source/MassSpawner/Private/MassEntityTraitBase.cpp
new file mode 100644
--- /dev/null
+++ b/Engine/Plugins/Runtime/MassGameplay/Source/MassSpawner/Private/MassEntityTraitBase.cpp
@@ -0,0 +1,41 @@
+// Copyright Epic Games, Inc. All Rights Reserved.
+
+#include "MassEntityTraitBase.h"
+
+//----------------------------------------------------------------------//
+//  UMassEntityTraitBase
+//----------------------------------------------------------------------//
+void UMassEntityTraitBase::DestroyTemplate() const 
+{
+
+}
+
+bool UMassEntityTraitBase::ValidateTemplate(const FMassEntityTemplateBuildContext& BuildContext, const UWorld& World, FAdditionalTraitRequirements& OutTraitRequirements) const
+{
+	return true;
+}
+
+#if WITH_EDITOR
+UMassEntityTraitBase::FOnNewTraitType UMassEntityTraitBase::OnNewTraitTypeEvent;
+
+void UMassEntityTraitBase::PostInitProperties()
+{
+	Super::PostInitProperties();
+
+	if (HasAnyFlags(RF_ClassDefaultObject))
+	{
+		OnNewTraitTypeEvent.Broadcast(*this);
+	}
+}
+
+#endif // WITH_EDITOR
+
+//-----------------------------------------------------------------------------
+// DEPRECATED
+//-----------------------------------------------------------------------------
+bool UMassEntityTraitBase::ValidateTemplate(FMassEntityTemplateBuildContext& BuildContext, const UWorld& World) const
+{
+	TArray<const UStruct*> AdditionalRequirements;
+	FAdditionalTraitRequirements AdditionalTraitRequirementsWrapper(AdditionalRequirements);
+	return ValidateTemplate(BuildContext, World, AdditionalTraitRequirementsWrapper);
+}
\ No newline at end of file
diff --git a/Engine/Plugins/Runtime/MassGameplay/Source/MassSpawner/Public/MassEntityConfigAsset.h b/Engine/Plugins/Runtime/MassGameplay/Source/MassSpawner/Public/MassEntityConfigAsset.h
--- a/Engine/Plugins/Runtime/MassGameplay/Source/MassSpawner/Public/MassEntityConfigAsset.h
+++ b/Engine/Plugins/Runtime/MassGameplay/Source/MassSpawner/Public/MassEntityConfigAsset.h
@@ -47,7 +47,10 @@ struct MASSSPAWNER_API FMassEntityConfig
 	TConstArrayView<UMassEntityTraitBase*> GetTraits() const { return Traits; }
 
 	/** Looks for a trait of the indicated type, accepting all child classes as well, unless bExactMatch == true */
-	const UMassEntityTraitBase* FindTrait(TSubclassOf<UMassEntityTraitBase> TraitClass, const bool bExactMatch = false) const;
+	const UMassEntityTraitBase* FindTrait(TSubclassOf<UMassEntityTraitBase> TraitClass, const bool bExactMatch = false) const
+	{
+		return FindTraitInternal(TraitClass, bExactMatch);
+	}
 
 	/** Adds Trait to the collection of traits hosted by this FMassEntityConfig instance */
 	void AddTrait(UMassEntityTraitBase& Trait);
@@ -76,12 +79,17 @@ struct MASSSPAWNER_API FMassEntityConfig
 #if WITH_EDITOR
 	/** Needs to be called when the given config is being duplicated - ensured the ConfigGuid remains unique */
 	void PostDuplicate(bool bDuplicateForPIE);
+
+	UMassEntityTraitBase* FindMutableTrait(TSubclassOf<UMassEntityTraitBase> TraitClass, const bool bExactMatch = false);
 #endif // WITH_EDITOR
 	
 protected:
 	/** Combines traits based on the config hierarchy and returns list of unique traits */
 	void GetCombinedTraits(TArray<UMassEntityTraitBase*>& OutTraits) const;
 
+	/** Looks for a trait of the indicated type, accepting all child classes as well, unless bExactMatch == true */
+	UMassEntityTraitBase* FindTraitInternal(TSubclassOf<UMassEntityTraitBase> TraitClass, const bool bExactMatch = false) const;
+
 	/** Combines traits based on the config hierarchy and returns list of unique traits */
 	UE_DEPRECATED(5.3, "This flavor of GetCombinedTraits is deprecated. Use the one without the ConfigOwner parameter (now a property of the FMassEntityConfig itself)")
 	void GetCombinedTraits(TArray<UMassEntityTraitBase*>& OutTraits, TArray<const UObject*>& Visited, const UObject& ConfigOwner) const;
@@ -149,6 +157,12 @@ public:
 	UFUNCTION(CallInEditor, Category = "Entity Config")
 	void ValidateEntityConfig();
 
+	/**
+	 * Returns a mutable instance of given trait class. If an instance of the given class can be found in
+	 * the hosted FMassEntityConfig then that instance is returned. Otherwise one will be created, added to
+	 * the FMassEntityConfig and returned.
+	 */
+	UMassEntityTraitBase* AddTrait(TSubclassOf<UMassEntityTraitBase> TraitClass);
 #endif // WITH_EDITOR
 
 protected:
diff --git a/Engine/Plugins/Runtime/MassGameplay/Source/MassSpawner/Public/MassEntityTemplateRegistry.h b/Engine/Plugins/Runtime/MassGameplay/Source/MassSpawner/Public/MassEntityTemplateRegistry.h
--- a/Engine/Plugins/Runtime/MassGameplay/Source/MassSpawner/Public/MassEntityTemplateRegistry.h
+++ b/Engine/Plugins/Runtime/MassGameplay/Source/MassSpawner/Public/MassEntityTemplateRegistry.h
@@ -15,6 +15,35 @@
 class UWorld;
 class UMassEntityTraitBase;
 
+USTRUCT()
+struct MASSSPAWNER_API FMassMissingTraitMessage
+{
+	GENERATED_BODY()
+#if WITH_EDITORONLY_DATA
+	const UMassEntityTraitBase* RequestingTrait = nullptr;
+	const UStruct* MissingType = nullptr;
+#endif // WITH_EDITORONLY_DATA
+};
+
+USTRUCT()
+struct MASSSPAWNER_API FMassDuplicateElementsMessage
+{
+	GENERATED_BODY()
+#if WITH_EDITORONLY_DATA
+	const UMassEntityTraitBase* DuplicatingTrait = nullptr;
+	const UMassEntityTraitBase* OriginalTrait = nullptr;
+	const UStruct* Element = nullptr;
+#endif // WITH_EDITORONLY_DATA
+};
+
+#if WITH_EDITORONLY_DATA
+namespace UE::Mass::Debug
+{
+	extern MASSSPAWNER_API const FName TraitFailedValidation;
+	extern MASSSPAWNER_API const FName TraitIgnored;
+}
+#endif // WITH_EDITORONLY_DATA
+
 enum class EFragmentInitialization : uint8
 {
 	DefaultInitializer,
@@ -167,12 +196,14 @@ struct FMassEntityTemplateBuildContext
 	template<typename T>
 	void RequireFragment()
 	{
+		static_assert(TIsDerivedFrom<T, FMassTag>::IsDerived == false, "Given struct type is a valid fragment type.");
 		AddDependency(T::StaticStruct());
 	}
 
 	template<typename T>
 	void RequireTag()
 	{
+		static_assert(TIsDerivedFrom<T, FMassTag>::IsDerived, "Given struct type is not a valid tag type.");
 		AddDependency(T::StaticStruct());
 	}
 
@@ -199,6 +230,36 @@ struct FMassEntityTemplateBuildContext
 	 */
 	bool BuildFromTraits(TConstArrayView<UMassEntityTraitBase*> Traits, const UWorld& World);
 
+	/** 
+	 * The method that allows to distinguish between regular context use (using traits to build templates) and 
+	 * the "data investigation" mode (used for debugging and authoring purposes). Utilize this function to 
+	 * avoid UWorld-specific operations (like getting subsystems). This method should also be used when a trait 
+	 * contains conditional logic - in that case it's required for the trait to add all the types that are potentially
+	 * added at runtime (even if seemingly conflicting information will be added). 
+	 * 
+	 * @return whether this context is in data inspection mode.
+	 */
+#if WITH_EDITORONLY_DATA
+	bool IsInspectingData() const
+	{
+		return bIsInspectingData;
+	}
+#else
+	constexpr bool IsInspectingData() const
+	{
+		return false;
+	}
+#endif
+
+#if WITH_EDITORONLY_DATA
+	void EnableDataInvestigationMode()
+	{
+		checkf(TemplateData.IsEmpty(), TEXT("Marking a FMassEntityTemplateBuildContext as being in 'investigation mode` is only supported before the context is first used."));
+		bIsInspectingData = true;
+	}
+#endif // WITH_EDITORONLY_DATA
+
+
 protected:
 
 	/**
@@ -221,7 +282,7 @@ protected:
 	 * in nature and are not associated with any traits. In that case it's ok to have multiple SetTraitBeingProcessed(nullptr)
 	 * calls.
 	 */
-	bool SetTraitBeingProcessed(const UMassEntityTraitBase* Trait);
+	MASSSPAWNER_API bool SetTraitBeingProcessed(const UMassEntityTraitBase* Trait);
 
 	void ResetBuildTimeData()
 	{
@@ -244,9 +305,20 @@ protected:
 
 	FMassEntityTemplateData& TemplateData;
 	FMassEntityTemplateID TemplateID;
+
+#if WITH_EDITORONLY_DATA
+private:
+	/**
+	 * This being set to `true` indicates that the context is being used to gather information, not to create actual
+	 * entity templates.
+	 */
+	bool bIsInspectingData = false;
+#endif // WITH_EDITORONLY_DATA
 };
 
-/** @todo document 
+/** 
+ * Represents a repository storing all the FMassEntityTemplate that have been created and registered as part of FMassEntityConfig
+ * processing or via custom code (like we do in InstancedActors plugin).
  */
 struct MASSSPAWNER_API FMassEntityTemplateRegistry
 {
@@ -303,11 +375,4 @@ protected:
 	TWeakObjectPtr<UObject> Owner;
 };
 
-
-UCLASS(deprecated, meta = (DeprecationMessage = "UMassEntityTemplateRegistry is deprecated starting UE5.2. Use FMassEntityTemplateRegistry instead"))
-class MASSSPAWNER_API UDEPRECATED_MassEntityTemplateRegistry : public UObject
-{
-	GENERATED_BODY()
-};
-
 #undef ENSURE_SUPPORTED_TRAIT_OPERATION
diff --git a/Engine/Plugins/Runtime/MassGameplay/Source/MassSpawner/Public/MassEntityTraitBase.h b/Engine/Plugins/Runtime/MassGameplay/Source/MassSpawner/Public/MassEntityTraitBase.h
--- a/Engine/Plugins/Runtime/MassGameplay/Source/MassSpawner/Public/MassEntityTraitBase.h
+++ b/Engine/Plugins/Runtime/MassGameplay/Source/MassSpawner/Public/MassEntityTraitBase.h
@@ -2,7 +2,6 @@
 
 #pragma once 
 
-#include "CoreMinimal.h"
 #include "MassEntityTraitBase.generated.h"
 
 struct FMassEntityTemplateBuildContext;
@@ -20,20 +19,67 @@ class MASSSPAWNER_API UMassEntityTraitBase : public UObject
 	GENERATED_BODY()
 
 public:
+	/** 
+	 * This is a type wrapping an existing TArray to limit how users can interact with the contained data. 
+	 * We essentially limit users to just adding elements, no other operations.
+	 */
+	struct FAdditionalTraitRequirements
+	{
+		explicit FAdditionalTraitRequirements(TArray<const UStruct*>& InTargetContainer)
+			: TargetContainer(InTargetContainer)
+		{
+		}
+		
+		// Copying constructor and assignment deleted to prevent users storing copies of the type, 
+		// which wouldn't be safe due to this type hosting a reference to an array that can go out of scope.
+		FAdditionalTraitRequirements(const FAdditionalTraitRequirements&) = delete;
+		FAdditionalTraitRequirements& operator=(const FAdditionalTraitRequirements&) = delete;
+
+		FAdditionalTraitRequirements& Add(const UScriptStruct* RequiredType)
+		{
+			TargetContainer.Add(RequiredType);
+			return *this;
+		}
+
+	private:
+		TArray<const UStruct*>& TargetContainer;
+	};
+
 
 	/** Appends items into the entity template required for the trait. */
 	virtual void BuildTemplate(FMassEntityTemplateBuildContext& BuildContext, const UWorld& World) const PURE_VIRTUAL(UMassEntityTraitBase::BuildTemplate, return; );
 
-	virtual void DestroyTemplate() const {}
+	virtual void DestroyTemplate() const;
 
 	/**
 	 * Called once all traits have been processed and fragment requirements have been checked. Override this function
 	 * to perform additional Trait's configuration validation. Returning `false` will indicate that the trait instance
 	 * is not happy with the validation results - this result will be treated as an error.
+	 * @param TraitRequirements contains requirements declared by this trait and gives ValidateTemplate a chance to add 
+	 *		to the dependencies based on the state of BuildContext, which by this point should contain all the elements 
+	 *		added by all the relevant traits.
 	 * @return whether the validation was successful
 	 */
-	virtual bool ValidateTemplate(FMassEntityTemplateBuildContext& BuildContext, const UWorld& World) const
+	virtual bool ValidateTemplate(const FMassEntityTemplateBuildContext& BuildContext, const UWorld& World, FAdditionalTraitRequirements& OutTraitRequirements) const;
+	
+	UE_DEPRECATED(5.5, "This flavor of ValidateTemplate is deprecated. Use the three-parameter one instead.")
+	virtual bool ValidateTemplate(FMassEntityTemplateBuildContext& BuildContext, const UWorld& World) const;
+
+#if WITH_EDITOR
+	DECLARE_MULTICAST_DELEGATE_OneParam(FOnNewTraitType, UMassEntityTraitBase&);
+	static FOnNewTraitType& GetOnNewTraitTypeEvent()
 	{
-		return true;
+		return OnNewTraitTypeEvent;
 	}
+
+protected:
+	virtual void PostInitProperties() override;
+
+private:
+	static FOnNewTraitType OnNewTraitTypeEvent;
+#endif // WITH_EDITOR
 };
+
+#if UE_ENABLE_INCLUDE_ORDER_DEPRECATED_IN_5_5
+#include "CoreMinimal.h"
+#endif
