[Mass] Minor clean up after recent large CL. Converted a couple of FInstancedStruct::Make to FConstStructView::Make, which is a better fit in the context.
Also: [Main CIS] Issue 533091: Mac CIS fix

diff --git a/Engine/Plugins/Runtime/MassGameplay/Source/MassSpawner/Private/MassEntityTemplateRegistry.cpp b/Engine/Plugins/Runtime/MassGameplay/Source/MassSpawner/Private/MassEntityTemplateRegistry.cpp
--- a/Engine/Plugins/Runtime/MassGameplay/Source/MassSpawner/Private/MassEntityTemplateRegistry.cpp
+++ b/Engine/Plugins/Runtime/MassGameplay/Source/MassSpawner/Private/MassEntityTemplateRegistry.cpp
@@ -162,7 +162,7 @@ bool FMassEntityTemplateBuildContext::ValidateBuildContext(const UWorld& World)
 		{
 			++ErrorCount;
 			IF_MESSAGES(
-				FMassDebugger::DebugEvent(UE::Mass::Debug::TraitFailedValidation, FInstancedStruct::Make<FMassGenericDebugEvent>(TraitData.Trait));
+				FMassDebugger::DebugEvent(UE::Mass::Debug::TraitFailedValidation, FConstStructView::Make(FMassGenericDebugEvent(TraitData.Trait)));
 			);
 		}
 	}
@@ -189,7 +189,7 @@ bool FMassEntityTemplateBuildContext::ValidateBuildContext(const UWorld& World)
 				}
 				IF_MESSAGES(
 					FMassDebugger::DebugEvent(FMassDuplicateElementsMessage::StaticStruct()->GetFName()
-						, FInstancedStruct::Make<FMassDuplicateElementsMessage>(TraitData.Trait, SourceTrait, TypeAdded)
+						, FConstStructView::Make(FMassDuplicateElementsMessage(TraitData.Trait, SourceTrait, TypeAdded))
 						, UE::Mass::Debug::bReportDuplicatedFragmentsAsWarnings ? EMassDebugMessageSeverity::Warning : EMassDebugMessageSeverity::Info);
 				);
 			}
@@ -221,7 +221,7 @@ bool FMassEntityTemplateBuildContext::ValidateBuildContext(const UWorld& World)
 	for (const UMassEntityTraitBase* IgnoredTrait : IgnoredTraits)
 	{
 		IF_MESSAGES(
-			FMassDebugger::DebugEvent(UE::Mass::Debug::TraitIgnored, FInstancedStruct::Make<FMassGenericDebugEvent>(IgnoredTrait));
+			FMassDebugger::DebugEvent(UE::Mass::Debug::TraitIgnored, FConstStructView::Make(FMassGenericDebugEvent(IgnoredTrait)));
 		);
 		++WarningCount;
 	}
diff --git a/Engine/Source/Runtime/MassEntity/Public/MassDebugger.h b/Engine/Source/Runtime/MassEntity/Public/MassDebugger.h
--- a/Engine/Source/Runtime/MassEntity/Public/MassDebugger.h
+++ b/Engine/Source/Runtime/MassEntity/Public/MassDebugger.h
@@ -191,7 +191,7 @@ struct MASSENTITY_API FMassDebugger
 	static void DebugEvent(TArgs&&... InArgs)
 	{
 		DebugEvent(TMessage::StaticStruct()->GetFName()
-			, FInstancedStruct::Make<TMessage>(Forward<TArgs>(InArgs)...));
+			, FConstStructView::Make(TMessage(Forward<TArgs>(InArgs)...)));
 	}
 
 	static void RegisterEntityManager(FMassEntityManager& EntityManager);
