[Main CIS] Issue 533091: Mac CIS fix, take 3

diff --git a/Engine/Plugins/Runtime/MassGameplay/Source/MassSpawner/Private/MassEntityTemplateRegistry.cpp b/Engine/Plugins/Runtime/MassGameplay/Source/MassSpawner/Private/MassEntityTemplateRegistry.cpp
--- a/Engine/Plugins/Runtime/MassGameplay/Source/MassSpawner/Private/MassEntityTemplateRegistry.cpp
+++ b/Engine/Plugins/Runtime/MassGameplay/Source/MassSpawner/Private/MassEntityTemplateRegistry.cpp
@@ -145,7 +145,7 @@ bool FMassEntityTemplateBuildContext::SetTraitBeingProcessed(const UMassEntityTr
 
 bool FMassEntityTemplateBuildContext::ValidateBuildContext(const UWorld& World)
 {
-#if WITH_UNREAL_DEVELOPER_TOOLS && WITH_EDITOR
+#if WITH_UNREAL_DEVELOPER_TOOLS && WITH_EDITOR && WITH_EDITORONLY_DATA
 #define IF_MESSAGES(Message) if (GEditor) { Message }
 #else
 #define IF_MESSAGES(_)
@@ -162,7 +162,7 @@ bool FMassEntityTemplateBuildContext::ValidateBuildContext(const UWorld& World)
 		{
 			++ErrorCount;
 			IF_MESSAGES(
-				FMassDebugger::DebugEvent(UE::Mass::Debug::TraitFailedValidation, FConstStructView::Make(FMassGenericDebugEvent(TraitData.Trait)));
+				FMassDebugger::DebugEvent(UE::Mass::Debug::TraitFailedValidation, FConstStructView::Make(FMassGenericDebugEvent{TraitData.Trait}));
 			);
 		}
 	}
@@ -189,7 +189,7 @@ bool FMassEntityTemplateBuildContext::ValidateBuildContext(const UWorld& World)
 				}
 				IF_MESSAGES(
 					FMassDebugger::DebugEvent(FMassDuplicateElementsMessage::StaticStruct()->GetFName()
-						, FConstStructView::Make(FMassDuplicateElementsMessage(TraitData.Trait, SourceTrait, TypeAdded))
+						, FConstStructView::Make(FMassDuplicateElementsMessage{TraitData.Trait, SourceTrait, TypeAdded})
 						, UE::Mass::Debug::bReportDuplicatedFragmentsAsWarnings ? EMassDebugMessageSeverity::Warning : EMassDebugMessageSeverity::Info);
 				);
 			}
@@ -221,7 +221,7 @@ bool FMassEntityTemplateBuildContext::ValidateBuildContext(const UWorld& World)
 	for (const UMassEntityTraitBase* IgnoredTrait : IgnoredTraits)
 	{
 		IF_MESSAGES(
-			FMassDebugger::DebugEvent(UE::Mass::Debug::TraitIgnored, FConstStructView::Make(FMassGenericDebugEvent(IgnoredTrait)));
+			FMassDebugger::DebugEvent(UE::Mass::Debug::TraitIgnored, FConstStructView::Make(FMassGenericDebugEvent{IgnoredTrait}));
 		);
 		++WarningCount;
 	}
diff --git a/Engine/Plugins/Runtime/MassGameplay/Source/MassSpawner/Public/MassEntityTemplateRegistry.h b/Engine/Plugins/Runtime/MassGameplay/Source/MassSpawner/Public/MassEntityTemplateRegistry.h
--- a/Engine/Plugins/Runtime/MassGameplay/Source/MassSpawner/Public/MassEntityTemplateRegistry.h
+++ b/Engine/Plugins/Runtime/MassGameplay/Source/MassSpawner/Public/MassEntityTemplateRegistry.h
@@ -19,7 +19,12 @@ USTRUCT()
 struct MASSSPAWNER_API FMassMissingTraitMessage
 {
 	GENERATED_BODY()
+
 #if WITH_EDITORONLY_DATA
+	explicit FMassMissingTraitMessage(const UMassEntityTraitBase* InRequestingTrait = nullptr, const UStruct* InMissingType = nullptr)
+		: RequestingTrait(InRequestingTrait), MissingType(InMissingType)
+	{}
+
 	const UMassEntityTraitBase* RequestingTrait = nullptr;
 	const UStruct* MissingType = nullptr;
 #endif // WITH_EDITORONLY_DATA
