[Mass] Fixes to Trait validation on cooked editors where FMassDebugger is unavailable

diff --git a/Engine/Plugins/Runtime/MassGameplay/Source/MassGameplayEditor/Private/MassTraitRepository.cpp b/Engine/Plugins/Runtime/MassGameplay/Source/MassGameplayEditor/Private/MassTraitRepository.cpp
--- a/Engine/Plugins/Runtime/MassGameplay/Source/MassGameplayEditor/Private/MassTraitRepository.cpp
+++ b/Engine/Plugins/Runtime/MassGameplay/Source/MassGameplayEditor/Private/MassTraitRepository.cpp
@@ -96,6 +96,8 @@ void UMassTraitRepository::Initialize(FSubsystemCollectionBase& Collection)
 // the way certain events are handled. 
 void UMassTraitRepository::OnDebugEvent(const FName EventName, FConstStructView Payload, EMassDebugMessageSeverity SeverityOverride)
 {
+#if WITH_MASSENTITY_DEBUG
+
 	static const FName MissingTraitMessageName = FMassMissingTraitMessage::StaticStruct()->GetFName();
 	static const FName DuplicateElementsMessageName = FMassDuplicateElementsMessage::StaticStruct()->GetFName();
 
@@ -234,6 +236,7 @@ void UMassTraitRepository::OnDebugEvent(const FName EventName, FConstStructView
 	}
 
 #undef OVERRIDABLE_SEVERITY
+#endif // WITH_MASSENTITY_DEBUG
 }
 
 TConstArrayView<FName> UMassTraitRepository::GetTraitsNameAddingElements(const FName ElementName)
diff --git a/Engine/Plugins/Runtime/MassGameplay/Source/MassSpawner/Private/MassEntityTemplateRegistry.cpp b/Engine/Plugins/Runtime/MassGameplay/Source/MassSpawner/Private/MassEntityTemplateRegistry.cpp
--- a/Engine/Plugins/Runtime/MassGameplay/Source/MassSpawner/Private/MassEntityTemplateRegistry.cpp
+++ b/Engine/Plugins/Runtime/MassGameplay/Source/MassSpawner/Private/MassEntityTemplateRegistry.cpp
@@ -145,7 +145,7 @@ bool FMassEntityTemplateBuildContext::SetTraitBeingProcessed(const UMassEntityTr
 
 bool FMassEntityTemplateBuildContext::ValidateBuildContext(const UWorld& World)
 {
-#if WITH_UNREAL_DEVELOPER_TOOLS && WITH_EDITOR && WITH_EDITORONLY_DATA
+#if WITH_UNREAL_DEVELOPER_TOOLS && WITH_EDITOR && WITH_EDITORONLY_DATA && WITH_MASSENTITY_DEBUG
 #define IF_MESSAGES(Message) if (GEditor) { Message }
 #else
 #define IF_MESSAGES(_)
