[Mass] Made Trait fixes via Trait Validation transactional (undo/redo).

diff --git a/Engine/Plugins/Runtime/MassGameplay/Source/MassSpawner/Private/MassEntityConfigAsset.cpp b/Engine/Plugins/Runtime/MassGameplay/Source/MassSpawner/Private/MassEntityConfigAsset.cpp
--- a/Engine/Plugins/Runtime/MassGameplay/Source/MassSpawner/Private/MassEntityConfigAsset.cpp
+++ b/Engine/Plugins/Runtime/MassGameplay/Source/MassSpawner/Private/MassEntityConfigAsset.cpp
@@ -11,9 +11,8 @@
 #include "Engine/World.h"
 #if WITH_EDITOR
 #include "Editor.h"
-#include "Framework/Notifications/NotificationManager.h"
-#include "Widgets/Notifications/SNotificationList.h"
 #include "MassEntityEditor.h"
+#include "ScopedTransaction.h"
 #endif // WITH_EDITOR
 
 #define LOCTEXT_NAMESPACE "Mass"
@@ -266,10 +265,13 @@ UMassEntityTraitBase* UMassEntityConfigAsset::AddTrait(TSubclassOf<UMassEntityTr
 	UMassEntityTraitBase* TraitInstance = Config.FindMutableTrait(TraitClass, /*bExactMatch=*/true);
 	if (TraitInstance == nullptr)
 	{
-		TraitInstance = NewObject<UMassEntityTraitBase>(this, TraitClass);
+		const FScopedTransaction Transaction(LOCTEXT("ProcedurallyAddingTrait", "Adding a trait procedurally"));
+
+		Modify();
+
+		TraitInstance = NewObject<UMassEntityTraitBase>(this, TraitClass, FName(), RF_Transactional);
 		check(TraitInstance);
 		Config.AddTrait(*TraitInstance);
-		Modify();
 	}
 	return TraitInstance;
 }
