Issue 533091: compilation fixes

diff --git a/Engine/Plugins/Runtime/MassGameplay/Source/MassGameplayEditor/Private/MassTraitRepository.cpp b/Engine/Plugins/Runtime/MassGameplay/Source/MassGameplayEditor/Private/MassTraitRepository.cpp
--- a/Engine/Plugins/Runtime/MassGameplay/Source/MassGameplayEditor/Private/MassTraitRepository.cpp
+++ b/Engine/Plugins/Runtime/MassGameplay/Source/MassGameplayEditor/Private/MassTraitRepository.cpp
@@ -84,9 +84,8 @@ void UMassTraitRepository::Initialize(FSubsystemCollectionBase& Collection)
 {
 	Super::Initialize(Collection);
 
-	OnNewTraitTypeHandle = UMassEntityTraitBase::GetOnNewTraitTypeEvent().AddUObject(this, &UMassTraitRepository::OnNewTraitType);
-
 #if WITH_MASSENTITY_DEBUG
+	OnNewTraitTypeHandle = UMassEntityTraitBase::GetOnNewTraitTypeEvent().AddUObject(this, &UMassTraitRepository::OnNewTraitType);
 	FMassDebugger::OnDebugEvent.AddUObject(this, &UMassTraitRepository::OnDebugEvent);
 #endif // WITH_MASSENTITY_DEBUG
 }
@@ -260,6 +259,7 @@ void UMassTraitRepository::InitRepository()
 		return;
 	}
 
+#if WITH_MASSENTITY_DEBUG
 	UWorld::InitializationValues IVS;
 	IVS.InitializeScenes(false)
 		.AllowAudioPlayback(false)
@@ -295,21 +295,24 @@ void UMassTraitRepository::InitRepository()
 			OnNewTraitType(**ClassIterator);
 		}
 	}
+#endif // WITH_MASSENTITY_DEBUG
 }
 
 void UMassTraitRepository::Deinitialize()
 {
-	FCoreDelegates::OnPostEngineInit.RemoveAll(this);
+#if WITH_MASSENTITY_DEBUG
 	UMassEntityTraitBase::GetOnNewTraitTypeEvent().Remove(OnNewTraitTypeHandle);
 
 	if (IsValid(InvestigationWorld))
 	{
 		InvestigationWorld->DestroyWorld(/*bInformEngineOfWorld=*/false);
 	}
+#endif // WITH_MASSENTITY_DEBUG
 
 	Super::Deinitialize();
 }
 
+#if WITH_MASSENTITY_DEBUG
 void UMassTraitRepository::OnNewTraitType(UMassEntityTraitBase& Trait)
 {
 	if (bIsRepositoryInitialized == false)
@@ -385,6 +388,7 @@ void UMassTraitRepository::OnNewTraitType(UMassEntityTraitBase& Trait)
 	}
 	TraitClassNameToDataMap.Add(TraitName, MoveTemp(TraitData));
 }
+#endif // WITH_MASSENTITY_DEBUG
 
 UWorld* UMassTraitRepository::GetInvestigationWorld()
 {
