[Mass] fixed UMassMovableVisualizationTrait not respecting bAllowServerSideVisualization flag provided by it's parent class

diff --git a/Engine/Plugins/Runtime/MassGameplay/Source/MassRepresentation/Private/MassMovableVisualizationTrait.cpp b/Engine/Plugins/Runtime/MassGameplay/Source/MassRepresentation/Private/MassMovableVisualizationTrait.cpp
--- a/Engine/Plugins/Runtime/MassGameplay/Source/MassRepresentation/Private/MassMovableVisualizationTrait.cpp
+++ b/Engine/Plugins/Runtime/MassGameplay/Source/MassRepresentation/Private/MassMovableVisualizationTrait.cpp
@@ -6,7 +6,8 @@
 
 void UMassMovableVisualizationTrait::BuildTemplate(FMassEntityTemplateBuildContext& BuildContext, const UWorld& World) const
 {
-	if (World.IsNetMode(NM_DedicatedServer) && !BuildContext.IsInspectingData())
+	if (!bAllowServerSideVisualization && World.IsNetMode(NM_DedicatedServer)
+		&& !BuildContext.IsInspectingData())
 	{
 		return;
 	}
