PR #12277: [Mass Entity] Fix mass.debug.SetDebugEntityRange and movement styles are not working as expected

diff --git a/Engine/Plugins/Runtime/MassGameplay/Source/MassMovement/Private/MassMovementTrait.cpp b/Engine/Plugins/Runtime/MassGameplay/Source/MassMovement/Private/MassMovementTrait.cpp
--- a/Engine/Plugins/Runtime/MassGameplay/Source/MassMovement/Private/MassMovementTrait.cpp
+++ b/Engine/Plugins/Runtime/MassGameplay/Source/MassMovement/Private/MassMovementTrait.cpp
@@ -18,6 +18,7 @@ void UMassMovementTrait::BuildTemplate(FMassEntityTemplateBuildContext& BuildCon
 	BuildContext.AddFragment<FMassVelocityFragment>();
 	BuildContext.AddFragment<FMassForceFragment>();
 
-	const FConstSharedStruct MovementFragment = EntityManager.GetOrCreateConstSharedFragment(Movement);
+	const FMassMovementParameters MovementValidated = Movement.GetValidated();
+	const FConstSharedStruct MovementFragment = EntityManager.GetOrCreateConstSharedFragment(MovementValidated);
 	BuildContext.AddConstSharedFragment(MovementFragment);
 }
diff --git a/Engine/Source/Runtime/MassEntity/Private/MassDebugger.cpp b/Engine/Source/Runtime/MassEntity/Private/MassDebugger.cpp
--- a/Engine/Source/Runtime/MassEntity/Private/MassDebugger.cpp
+++ b/Engine/Source/Runtime/MassEntity/Private/MassDebugger.cpp
@@ -78,7 +78,7 @@ namespace UE::Mass::Debug
 	void SetDebugEntityRange(const int32 InDebugEntityBegin, const int32 InDebugEntityEnd)
 	{
 		DebugEntityBegin = InDebugEntityBegin;
-		DebugEntityEnd = InDebugEntityBegin;
+		DebugEntityEnd = InDebugEntityEnd;
 	}
 
 	static FAutoConsoleCommand SetDebugEntityRangeCommand(
