[Mass] Minor spelling fixes

diff --git a/Engine/Plugins/Runtime/MassGameplay/Source/MassSpawner/Public/MassEntityTraitBase.h b/Engine/Plugins/Runtime/MassGameplay/Source/MassSpawner/Public/MassEntityTraitBase.h
--- a/Engine/Plugins/Runtime/MassGameplay/Source/MassSpawner/Public/MassEntityTraitBase.h
+++ b/Engine/Plugins/Runtime/MassGameplay/Source/MassSpawner/Public/MassEntityTraitBase.h
@@ -58,7 +58,7 @@ public:
 	 * Called once all traits have been processed and fragment requirements have been checked. Override this function
 	 * to perform additional Trait's configuration validation. Returning `false` will indicate that the trait instance
 	 * is not happy with the validation results - this result will be treated as an error.
-	 * @param TraitRequirements contains requirements declared by this trait and gives ValidateTemplate a chance to add 
+	 * @param OutTraitRequirements contains requirements declared by this trait and gives ValidateTemplate a chance to add 
 	 *		to the dependencies based on the state of BuildContext, which by this point should contain all the elements 
 	 *		added by all the relevant traits.
 	 * @return whether the validation was successful
diff --git a/Engine/Source/Runtime/MassEntity/Internal/MassArchetypeData.h b/Engine/Source/Runtime/MassEntity/Internal/MassArchetypeData.h
--- a/Engine/Source/Runtime/MassEntity/Internal/MassArchetypeData.h
+++ b/Engine/Source/Runtime/MassEntity/Internal/MassArchetypeData.h
@@ -404,7 +404,7 @@ protected:
 	 * The function first creates new FMassArchetypeSharedFragmentValues instance combining existing values
 	 * and the contents of SharedFragmentValueOverrides. Then that is used to find the target chunk for Entity,
 	 * and if one cannot be found a new one will be created. 
-	 * @param SharedFragmentValueOverrides is expected to contain only instances of types already
+	 * @param SharedFragmentValueOverrides is expected to contain only instance of types already
 	 *    present in given archetypes FMassArchetypeSharedFragmentValues
 	 */
 	void SetSharedFragmentsData(const FMassEntityHandle Entity, TConstArrayView<FSharedStruct> SharedFragmentValueOverrides);
