[Mass] Fixed Mass debug FAgentDebugVisualization wire shape visualization no longer working

diff --git a/Engine/Plugins/Runtime/MassGameplay/Source/MassGameplayDebug/Private/DebugVisLocationProcessor.cpp b/Engine/Plugins/Runtime/MassGameplay/Source/MassGameplayDebug/Private/DebugVisLocationProcessor.cpp
--- a/Engine/Plugins/Runtime/MassGameplay/Source/MassGameplayDebug/Private/DebugVisLocationProcessor.cpp
+++ b/Engine/Plugins/Runtime/MassGameplay/Source/MassGameplayDebug/Private/DebugVisLocationProcessor.cpp
@@ -82,22 +82,23 @@ UMassProcessor_UpdateDebugVis::UMassProcessor_UpdateDebugVis()
 
 void UMassProcessor_UpdateDebugVis::ConfigureQueries() 
 {
-	// @todo only FDataFragment_DebugVis should be mandatory, rest optional 
 	EntityQuery.AddRequirement<FTransformFragment>(EMassFragmentAccess::ReadOnly);
 	EntityQuery.AddRequirement<FDataFragment_DebugVis>(EMassFragmentAccess::ReadWrite);
 	EntityQuery.AddRequirement<FAgentRadiusFragment>(EMassFragmentAccess::ReadWrite);
 	EntityQuery.AddTagRequirement<FMassDebuggableTag>(EMassFragmentPresence::All);
-	EntityQuery.AddSubsystemRequirement<UMassDebuggerSubsystem>(EMassFragmentAccess::ReadWrite);
+
+	ProcessorRequirements.AddSubsystemRequirement<UMassDebuggerSubsystem>(EMassFragmentAccess::ReadWrite);
 }
 
 void UMassProcessor_UpdateDebugVis::Execute(FMassEntityManager& EntityManager, FMassExecutionContext& Context)
 {
 	QUICK_SCOPE_CYCLE_COUNTER(UMassProcessor_UpdateDebugVis_Run);
 
-	EntityQuery.ForEachEntityChunk(EntityManager, Context, [this](FMassExecutionContext& Context)
-		{
-			UMassDebuggerSubsystem& Debugger = Context.GetMutableSubsystemChecked<UMassDebuggerSubsystem>();
+	UMassDebuggerSubsystem& Debugger = Context.GetMutableSubsystemChecked<UMassDebuggerSubsystem>();
+	Debugger.ResetDebugShapes();
 
+	EntityQuery.ForEachEntityChunk(EntityManager, Context, [&Debugger](FMassExecutionContext& Context)
+		{
 			const int32 NumEntities = Context.GetNumEntities();
 			const TConstArrayView<FTransformFragment> LocationList = Context.GetFragmentView<FTransformFragment>();
 			const TArrayView<FDataFragment_DebugVis> DebugVisList = Context.GetMutableFragmentView<FDataFragment_DebugVis>();
diff --git a/Engine/Plugins/Runtime/MassGameplay/Source/MassGameplayDebug/Private/MassDebuggerSubsystem.cpp b/Engine/Plugins/Runtime/MassGameplay/Source/MassGameplayDebug/Private/MassDebuggerSubsystem.cpp
--- a/Engine/Plugins/Runtime/MassGameplay/Source/MassGameplayDebug/Private/MassDebuggerSubsystem.cpp
+++ b/Engine/Plugins/Runtime/MassGameplay/Source/MassGameplayDebug/Private/MassDebuggerSubsystem.cpp
@@ -43,14 +43,17 @@ void UMassDebuggerSubsystem::OnEntitySelected(const FMassEntityManager& EntityMa
 }
 #endif // WITH_MASSENTITY_DEBUG
 
-void UMassDebuggerSubsystem::PreTickProcessors()
+void UMassDebuggerSubsystem::ResetDebugShapes()
 {
 	// get ready to receive new debug info
 	for (TArray<FShapeDesc>& Array : Shapes)
 	{
 		Array.Reset();
 	}
-	
+}
+
+void UMassDebuggerSubsystem::PreTickProcessors()
+{
 	Entities.Reset();
 	Locations.Reset();
 	SelectedEntityDetails.Empty();
diff --git a/Engine/Plugins/Runtime/MassGameplay/Source/MassGameplayDebug/Public/MassDebuggerSubsystem.h b/Engine/Plugins/Runtime/MassGameplay/Source/MassGameplayDebug/Public/MassDebuggerSubsystem.h
--- a/Engine/Plugins/Runtime/MassGameplay/Source/MassGameplayDebug/Public/MassDebuggerSubsystem.h
+++ b/Engine/Plugins/Runtime/MassGameplay/Source/MassGameplayDebug/Public/MassDebuggerSubsystem.h
@@ -29,6 +29,7 @@ public:
 
 	void AddShape(EMassEntityDebugShape Shape, FVector Location, float Size) { Shapes[uint8(Shape)].Add(FShapeDesc(Location, Size)); }
 	const TArray<FShapeDesc>* GetShapes() const { return Shapes; }
+	void ResetDebugShapes();
 
 	FMassEntityHandle GetSelectedEntity() const { return SelectedEntity; }
 	void SetSelectedEntity(const FMassEntityHandle InSelectedEntity);
