CitySample: Fix clang compile

#rnx

diff --git a/Engine/Plugins/Runtime/MassGameplay/Source/MassRepresentation/Public/MassRepresentationTypes.h b/Engine/Plugins/Runtime/MassGameplay/Source/MassRepresentation/Public/MassRepresentationTypes.h
--- a/Engine/Plugins/Runtime/MassGameplay/Source/MassRepresentation/Public/MassRepresentationTypes.h
+++ b/Engine/Plugins/Runtime/MassGameplay/Source/MassRepresentation/Public/MassRepresentationTypes.h
@@ -622,7 +622,7 @@ public:
 		return nullptr;
 	}
 
-	FORCEINLINE void AddBatchedTransform(const FMassEntityHandle EntityHandle, const FTransform& Transform, const FTransform& PrevTransform, const float LODSignificance, const float PrevLODSignificance = -1.0f)
+	void AddBatchedTransform(const FMassEntityHandle EntityHandle, const FTransform& Transform, const FTransform& PrevTransform, const float LODSignificance, const float PrevLODSignificance = -1.0f)
 	{
 		if (FMassLODSignificanceRange* Range = GetLODSignificanceRange(LODSignificance))
 		{
