[MassGameplay] Shuffled instructions in UMassSignalProcessorBase::OnSignalReceived to reduce time spent in the scoped spin lock

diff --git a/Engine/Plugins/Runtime/MassGameplay/Source/MassSignals/Private/MassSignalProcessorBase.cpp b/Engine/Plugins/Runtime/MassGameplay/Source/MassSignals/Private/MassSignalProcessorBase.cpp
--- a/Engine/Plugins/Runtime/MassGameplay/Source/MassSignals/Private/MassSignalProcessorBase.cpp
+++ b/Engine/Plugins/Runtime/MassGameplay/Source/MassSignals/Private/MassSignalProcessorBase.cpp
@@ -142,13 +142,15 @@ void UMassSignalProcessorBase::Execute(FMassEntityManager& EntityManager, FMassE
 
 void UMassSignalProcessorBase::OnSignalReceived(FName SignalName, TConstArrayView<FMassEntityHandle> Entities)
 {
+	FEntitySignalRange Range;
+	Range.SignalName = SignalName;
+
 	UE::TScopeLock<UE::FSpinLock> ScopeLock(ReceivedSignalLock);
 
 	FFrameReceivedSignals& CurrentFrameBuffer = FrameReceivedSignals[CurrentFrameBufferIndex];
 
-	FEntitySignalRange& Range = CurrentFrameBuffer.ReceivedSignalRanges.AddDefaulted_GetRef();
-	Range.SignalName = SignalName;
 	Range.Begin = CurrentFrameBuffer.SignaledEntities.Num();
 	CurrentFrameBuffer.SignaledEntities.Append(Entities.GetData(), Entities.Num());
 	Range.End = CurrentFrameBuffer.SignaledEntities.Num();
+	CurrentFrameBuffer.ReceivedSignalRanges.Add(MoveTemp(Range));
 }
