[MassGameplay] Marked the generic MassRepresentationProcessor as abstract since we never expect it to be used in practice.

diff --git a/Engine/Plugins/Runtime/MassGameplay/Source/MassRepresentation/Public/MassRepresentationProcessor.h b/Engine/Plugins/Runtime/MassGameplay/Source/MassRepresentation/Public/MassRepresentationProcessor.h
--- a/Engine/Plugins/Runtime/MassGameplay/Source/MassRepresentation/Public/MassRepresentationProcessor.h
+++ b/Engine/Plugins/Runtime/MassGameplay/Source/MassRepresentation/Public/MassRepresentationProcessor.h
@@ -33,7 +33,7 @@ struct FMassRepresentationUpdateParams
 	bool bTestCollisionAvailibilityForActorVisualization = true;
 };
 
-UCLASS()
+UCLASS(Abstract)
 class MASSREPRESENTATION_API UMassRepresentationProcessor : public UMassProcessor
 {
 	GENERATED_BODY()
