[MassGameplay] Marked UMassStationaryISMSwitcherProcessor to run in all net modes by default

diff --git a/Engine/Plugins/Runtime/MassGameplay/Source/MassRepresentation/Private/MassStationaryISMSwitcherProcessor.cpp b/Engine/Plugins/Runtime/MassGameplay/Source/MassRepresentation/Private/MassStationaryISMSwitcherProcessor.cpp
--- a/Engine/Plugins/Runtime/MassGameplay/Source/MassRepresentation/Private/MassStationaryISMSwitcherProcessor.cpp
+++ b/Engine/Plugins/Runtime/MassGameplay/Source/MassRepresentation/Private/MassStationaryISMSwitcherProcessor.cpp
@@ -14,6 +14,7 @@ UMassStationaryISMSwitcherProcessor::UMassStationaryISMSwitcherProcessor(const F
 	ExecutionOrder.ExecuteInGroup = UE::Mass::ProcessorGroupNames::Representation;
 	ExecutionOrder.ExecuteAfter.Add(UMassVisualizationProcessor::StaticClass()->GetFName());
 	bAutoRegisterWithProcessingPhases = true;
+	ExecutionFlags = static_cast<int32>(EProcessorExecutionFlags::AllNetModes);
 }
 
 void UMassStationaryISMSwitcherProcessor::ConfigureQueries()
