[MassGameplay] Added a way for Traits to remove tags from FMassEntityTemplateBuildContext. It comes with full editor-time explanation when a needed tag has been removed (via the validation tool).

diff --git a/Engine/Plugins/Runtime/MassGameplay/Source/MassGameplayEditor/Private/MassTraitRepository.cpp b/Engine/Plugins/Runtime/MassGameplay/Source/MassGameplayEditor/Private/MassTraitRepository.cpp
--- a/Engine/Plugins/Runtime/MassGameplay/Source/MassGameplayEditor/Private/MassTraitRepository.cpp
+++ b/Engine/Plugins/Runtime/MassGameplay/Source/MassGameplayEditor/Private/MassTraitRepository.cpp
@@ -119,77 +119,88 @@ void UMassTraitRepository::OnDebugEvent(const FName EventName, FConstStructView
 						, FText::FromName(TraitClassName))
 				));
 
-			UMassTraitRepository* TraitRepo = GEditor->GetEditorSubsystem<UMassTraitRepository>();
-			TConstArrayView<FName> SuggestedTraitNames = TraitRepo
-				? TraitRepo->GetTraitsNameAddingElements(MissingElement)
-				: TConstArrayView<FName>();
-
-			if (SuggestedTraitNames.Num())
+			// if the missing elements has been added but removed by some trait that's all we need to tell the user:
+			if (MissingTraitMessage->RemovedByTrait)
+			{
+				IntroMessage->AddText(FText::FormatOrdered(LOCTEXT("MissingElementSuggestionRemoved"
+					, "has unsatisfied dependency of {0}. The type has been explicitly removed by {1}.")
+					, FText::FromName(MissingElement)
+					, FText::FromName(MissingTraitMessage->RemovedByTrait->GetFName())));
+			}
+			else
 			{
-				// The FixController will coordinate IFixer instances and the relevant FFixToken to
-				// ensure that only one of them can be applied. Once any of the fixes is applied
-				// the rest will become inactive (the FFixToken tokens will become grayed out and non-clickable).
-				// @todo at the moment FMutuallyExclusiveFixSet doesn't care whether fixing was successful. 
-				//		Should be relatively easy to address but needs to be coordinated with the author
-				TSharedRef<UE::DataValidation::FMutuallyExclusiveFixSet> FixController = MakeShareable(new UE::DataValidation::FMutuallyExclusiveFixSet());
-
-				IntroMessage->AddText(FText::FormatOrdered(LOCTEXT("MissingElementSuggestionOptions"
-					, "has unsatisfied dependency of {0}. The following actions can address it:")
-					, FText::FromName(MissingElement)));
-
-				for (const FName& SuggestedTraitName : SuggestedTraitNames)
+				UMassTraitRepository* TraitRepo = GEditor->GetEditorSubsystem<UMassTraitRepository>();
+				TConstArrayView<FName> SuggestedTraitNames = TraitRepo
+					? TraitRepo->GetTraitsNameAddingElements(MissingElement)
+					: TConstArrayView<FName>();
+
+				if (SuggestedTraitNames.Num())
 				{
-					bool bFixable = false;
-					if (MissingTraitMessage->RequestingTrait)
+					// The FixController will coordinate IFixer instances and the relevant FFixToken to
+					// ensure that only one of them can be applied. Once any of the fixes is applied
+					// the rest will become inactive (the FFixToken tokens will become grayed out and non-clickable).
+					// @todo at the moment FMutuallyExclusiveFixSet doesn't care whether fixing was successful. 
+					//		Should be relatively easy to address but needs to be coordinated with the author
+					TSharedRef<UE::DataValidation::FMutuallyExclusiveFixSet> FixController = MakeShareable(new UE::DataValidation::FMutuallyExclusiveFixSet());
+
+					IntroMessage->AddText(FText::FormatOrdered(LOCTEXT("MissingElementSuggestionOptions"
+						, "has unsatisfied dependency of {0}. The following actions can address it:")
+						, FText::FromName(MissingElement)));
+
+					for (const FName& SuggestedTraitName : SuggestedTraitNames)
 					{
-						TWeakObjectPtr<UClass> WeakTraitClass = GetTraitClass(SuggestedTraitName);
-						UMassEntityConfigAsset* EntityConfigAsset = Cast<UMassEntityConfigAsset>(MissingTraitMessage->RequestingTrait->GetOuter());
-
-						if (EntityConfigAsset && WeakTraitClass.IsValid())
+						bool bFixable = false;
+						if (MissingTraitMessage->RequestingTrait)
 						{
-							bFixable = true;
+							TWeakObjectPtr<UClass> WeakTraitClass = GetTraitClass(SuggestedTraitName);
+							UMassEntityConfigAsset* EntityConfigAsset = Cast<UMassEntityConfigAsset>(MissingTraitMessage->RequestingTrait->GetOuter());
+
+							if (EntityConfigAsset && WeakTraitClass.IsValid())
+							{
+								bFixable = true;
 
-							TWeakObjectPtr<UMassEntityConfigAsset> WeakConfig = EntityConfigAsset;
+								TWeakObjectPtr<UMassEntityConfigAsset> WeakConfig = EntityConfigAsset;
 							
-							// capturing FixController to make sure it exists as long as the fixes are alive. The lambda will
-							// be destroyed once the Fixer tokens get destroyed, for example during MessageLog page clearing.
-							TFunction<FFixResult()> ApplyFix = [WeakConfig, WeakTraitClass, _ = FixController]()
+								// capturing FixController to make sure it exists as long as the fixes are alive. The lambda will
+								// be destroyed once the Fixer tokens get destroyed, for example during MessageLog page clearing.
+								TFunction<FFixResult()> ApplyFix = [WeakConfig, WeakTraitClass, _ = FixController]()
 								{
 									return UE::Mass::Private::AddTraitToConfigFix(WeakConfig, WeakTraitClass);
 								};
 
-							const TSharedRef<UE::DataValidation::IFixer> Fixer = UE::DataValidation::MakeFix(MoveTemp(ApplyFix));
+								const TSharedRef<UE::DataValidation::IFixer> Fixer = UE::DataValidation::MakeFix(MoveTemp(ApplyFix));
 
-							FixController->Add(FText::FormatOrdered(LOCTEXT("AddMissingTrait", "Add {0} trait to {1} entity config")
-										, FText::FromName(SuggestedTraitName)
-										, FText::FromName(EntityConfigAsset->GetFName()))
-									, Fixer);
+								FixController->Add(FText::FormatOrdered(LOCTEXT("AddMissingTrait", "Add {0} trait to {1} entity config")
+											, FText::FromName(SuggestedTraitName)
+											, FText::FromName(EntityConfigAsset->GetFName()))
+										, Fixer);
+							}
 						}
-					}
 					
-					if (bFixable == false)
-					{
-						// unfixable (since we're unable to determine the UMassEntityConfigAsset outer), so just report
-						Messages.Add_GetRef(FTokenizedMessage::Create(OVERRIDABLE_SEVERITY(EMessageSeverity::Info)))
-							->AddText(FText::FormatOrdered(LOCTEXT("MissingElementSuggestionUnfixable", "\t{0}")
-								, FText::FromName(SuggestedTraitName)));
+						if (bFixable == false)
+						{
+							// unfixable (since we're unable to determine the UMassEntityConfigAsset outer), so just report
+							Messages.Add_GetRef(FTokenizedMessage::Create(OVERRIDABLE_SEVERITY(EMessageSeverity::Info)))
+								->AddText(FText::FormatOrdered(LOCTEXT("MissingElementSuggestionUnfixable", "\t{0}")
+									, FText::FromName(SuggestedTraitName)));
+						}
 					}
-				}
 
-				// for every IFixer instance we created in the loop above the call bellow will create a FFixToken
-				// related to that specific "fix". We attach all the tokens to the initial "here are your options" message
-				FixController->CreateTokens([IntroMessage](TSharedRef<FFixToken> FixToken)
-					{
-						IntroMessage->AddToken(FixToken);
-					}
-				);
-			}
-			else
-			{
-				IntroMessage->AddText(FText::FormatOrdered(LOCTEXT("MissingElementSuggestionNoOptions"
-					, "has unsatisfied dependency of {0}. There are no registered Traits that provide the type. Try using {1}.")
-					, FText::FromName(MissingElement)
-					, FText::FromName(UMassAssortedFragmentsTrait::StaticClass()->GetFName())));
+					// for every IFixer instance we created in the loop above the call bellow will create a FFixToken
+					// related to that specific "fix". We attach all the tokens to the initial "here are your options" message
+					FixController->CreateTokens([IntroMessage](TSharedRef<FFixToken> FixToken)
+						{
+							IntroMessage->AddToken(FixToken);
+						}
+					);
+				}
+				else
+				{
+					IntroMessage->AddText(FText::FormatOrdered(LOCTEXT("MissingElementSuggestionNoOptions"
+						, "has unsatisfied dependency of {0}. There are no registered Traits that provide the type. Try using {1}.")
+						, FText::FromName(MissingElement)
+						, FText::FromName(UMassAssortedFragmentsTrait::StaticClass()->GetFName())));
+				}
 			}
 
 			MessageLog.AddMessages(Messages);
diff --git a/Engine/Plugins/Runtime/MassGameplay/Source/MassSpawner/Private/MassEntityTemplateRegistry.cpp b/Engine/Plugins/Runtime/MassGameplay/Source/MassSpawner/Private/MassEntityTemplateRegistry.cpp
--- a/Engine/Plugins/Runtime/MassGameplay/Source/MassSpawner/Private/MassEntityTemplateRegistry.cpp
+++ b/Engine/Plugins/Runtime/MassGameplay/Source/MassSpawner/Private/MassEntityTemplateRegistry.cpp
@@ -119,6 +119,14 @@ bool FMassEntityTemplateBuildContext::BuildFromTraits(TConstArrayView<UMassEntit
 			Trait->BuildTemplate(*this, World);
 		}
 	}
+	// now remove all that has been requested to be removed
+	// those are only tags for now, thus the shortcut of going directly for tags
+	for (FRemovedType& Removed : RemovedTags)
+	{
+		check(Removed.TypeRemoved);
+		TemplateData.RemoveTag(*CastChecked<UScriptStruct>(Removed.TypeRemoved));
+	}
+
 	bBuildInProgress = false;
 
 	const bool bTemplateValid = ValidateBuildContext(World);
@@ -207,12 +215,19 @@ bool FMassEntityTemplateBuildContext::ValidateBuildContext(const UWorld& World)
 		{
 			if (TypesAlreadyAdded.Contains(TypeRequired) == false)
 			{
-				UE_LOG(LogMass, Error, TEXT("%s: Missing required fragment of type %s")
+				UE_LOG(LogMass, Error, TEXT("%s: Missing required element of type %s")
 					, *GetNameSafe(TraitData.Trait), *GetNameSafe(TypeRequired));
 				++ErrorCount;
 				IF_MESSAGES(
 				{
-					FMassDebugger::DebugEvent<FMassMissingTraitMessage>(TraitData.Trait, TypeRequired);
+					// check if it was removed
+					const UMassEntityTraitBase* RemovedByTrait = nullptr;
+					if (const int32 RemoverIndex = RemovedTags.Find(FRemovedType({TypeRequired})) != INDEX_NONE)
+					{
+						RemovedByTrait = RemovedTags[RemoverIndex].Remover;
+					}
+
+					FMassDebugger::DebugEvent<FMassMissingTraitMessage>(TraitData.Trait, TypeRequired, RemovedByTrait);
 				});
 			}
 		}
diff --git a/Engine/Plugins/Runtime/MassGameplay/Source/MassSpawner/Public/MassEntityTemplateRegistry.h b/Engine/Plugins/Runtime/MassGameplay/Source/MassSpawner/Public/MassEntityTemplateRegistry.h
--- a/Engine/Plugins/Runtime/MassGameplay/Source/MassSpawner/Public/MassEntityTemplateRegistry.h
+++ b/Engine/Plugins/Runtime/MassGameplay/Source/MassSpawner/Public/MassEntityTemplateRegistry.h
@@ -21,12 +21,14 @@ struct MASSSPAWNER_API FMassMissingTraitMessage
 	GENERATED_BODY()
 
 #if WITH_EDITORONLY_DATA
-	explicit FMassMissingTraitMessage(const UMassEntityTraitBase* InRequestingTrait = nullptr, const UStruct* InMissingType = nullptr)
-		: RequestingTrait(InRequestingTrait), MissingType(InMissingType)
+	explicit FMassMissingTraitMessage(const UMassEntityTraitBase* InRequestingTrait = nullptr, const UStruct* InMissingType = nullptr, const UMassEntityTraitBase* InRemovedByTrait = nullptr)
+		: RequestingTrait(InRequestingTrait), MissingType(InMissingType), RemovedByTrait(InRemovedByTrait)
 	{}
 
 	const UMassEntityTraitBase* RequestingTrait = nullptr;
 	const UStruct* MissingType = nullptr;
+	// if set indicates that the missing type has been explicitly removed by given trait.
+	const UMassEntityTraitBase* RemovedByTrait = nullptr;
 #endif // WITH_EDITORONLY_DATA
 };
 
@@ -127,6 +129,28 @@ struct FMassEntityTemplateBuildContext
 		TemplateData.AddSharedFragment(InSharedFragment);
 	}
 
+	/**
+	 * Removes given tag from collected data. More precisely: it will store the information and apply upon template creation (an optimization). 
+	 * WARNING: use with caution and only in cases where you know for certain what the given tag does and which processors rely on it.
+	 *		Using this functionality makes most sense for removing tags that specifically mean that entities having it are to be
+	 *		processed by a given processor.
+	 */
+	void RemoveTag(const UScriptStruct& TagType)
+	{
+		checkf(TagType.IsChildOf(FMassTag::StaticStruct()), TEXT("Given struct doesn't represent a valid mass tag type. Make sure to inherit from FMassTag or one of its child-types."));
+		RemovedTags.Add({&TagType
+#if WITH_EDITORONLY_DATA
+			, TraitsData.Last().Trait
+#endif // WITH_EDITORONLY_DATA
+		});
+	}
+
+	template<typename T>
+	void RemoveTag()
+	{
+		RemoveTag(*T::StaticStruct());
+	}
+
 	template<typename T>
 	T* GetFragment()
 	{
@@ -294,6 +318,7 @@ protected:
 		TraitsData.Reset();
 		TraitsProcessed.Reset();
 		IgnoredTraits.Reset();
+		RemovedTags.Reset();
 		bBuildInProgress = false;
 	}
 
@@ -306,6 +331,24 @@ protected:
 	TArray<FTraitData> TraitsData;
 	TSet<const UMassEntityTraitBase*> TraitsProcessed;
 	TSet<const UMassEntityTraitBase*> IgnoredTraits;
+
+	struct FRemovedType
+	{
+		const UStruct* TypeRemoved = nullptr;
+#if WITH_EDITOR
+		const UMassEntityTraitBase* Remover = nullptr;
+#endif // WITH_EDITOR
+		bool operator==(const FRemovedType& Other) const
+		{
+			return TypeRemoved == Other.TypeRemoved;
+		}
+	};
+	/**
+	 * These tags will be removed from the resulting entity template
+	 * @see RemoveTag for more details
+	 */
+	TArray<FRemovedType> RemovedTags;
+
 	bool bBuildInProgress = false;
 
 	FMassEntityTemplateData& TemplateData;
@@ -356,12 +399,12 @@ struct MASSSPAWNER_API FMassEntityTemplateRegistry
 	UE_DEPRECATED(5.3, "We no longer support fething mutable templates from the TemplateRegistry. Stored templates are considered const.")
 	FMassEntityTemplate* FindMutableTemplateFromTemplateID(FMassEntityTemplateID TemplateID);
 
-	UE_DEPRECATED(5.3, "CreateTemplate is no longer available. Use AddTemplate instead.")
+	UE_DEPRECATED(5.3, "CreateTemplate is no longer available. Use FindOrAddTemplate instead.")
 	FMassEntityTemplate& CreateTemplate(const uint32 HashLookup, FMassEntityTemplateID TemplateID);
 
 	void DestroyTemplate(FMassEntityTemplateID TemplateID);
 
-	UE_DEPRECATED(5.3, "InitializeEntityTemplate is no longer available. Use AddTemplate instead.")
+	UE_DEPRECATED(5.3, "InitializeEntityTemplate is no longer available. Use FindOrAddTemplate instead.")
 	void InitializeEntityTemplate(FMassEntityTemplate& InOutTemplate) const;
 
 	FMassEntityManager& GetEntityManagerChecked() { check(EntityManager); return *EntityManager; }
