[Mass] CIS Issue 592177, fixed the offending condition

diff --git a/Engine/Plugins/Runtime/MassGameplay/Source/MassSpawner/Private/MassEntityTemplateRegistry.cpp b/Engine/Plugins/Runtime/MassGameplay/Source/MassSpawner/Private/MassEntityTemplateRegistry.cpp
--- a/Engine/Plugins/Runtime/MassGameplay/Source/MassSpawner/Private/MassEntityTemplateRegistry.cpp
+++ b/Engine/Plugins/Runtime/MassGameplay/Source/MassSpawner/Private/MassEntityTemplateRegistry.cpp
@@ -222,7 +222,8 @@ bool FMassEntityTemplateBuildContext::ValidateBuildContext(const UWorld& World)
 				{
 					// check if it was removed
 					const UMassEntityTraitBase* RemovedByTrait = nullptr;
-					if (const int32 RemoverIndex = RemovedTags.Find(FRemovedType({TypeRequired})) != INDEX_NONE)
+					const int32 RemoverIndex = RemovedTags.Find(FRemovedType({TypeRequired}));
+					if (RemoverIndex != INDEX_NONE)
 					{
 						RemovedByTrait = RemovedTags[RemoverIndex].Remover;
 					}
