[MassGameplay] fixed trait validaiton not properly reporting elements missing due to them being removed by other traits.

diff --git a/Engine/Plugins/Runtime/MassGameplay/Source/MassSpawner/Private/MassEntityTemplateRegistry.cpp b/Engine/Plugins/Runtime/MassGameplay/Source/MassSpawner/Private/MassEntityTemplateRegistry.cpp
--- a/Engine/Plugins/Runtime/MassGameplay/Source/MassSpawner/Private/MassEntityTemplateRegistry.cpp
+++ b/Engine/Plugins/Runtime/MassGameplay/Source/MassSpawner/Private/MassEntityTemplateRegistry.cpp
@@ -121,7 +121,7 @@ bool FMassEntityTemplateBuildContext::BuildFromTraits(TConstArrayView<UMassEntit
 	}
 	// now remove all that has been requested to be removed
 	// those are only tags for now, thus the shortcut of going directly for tags
-	for (FRemovedType& Removed : RemovedTags)
+	for (FRemovedType& Removed : RemovedTypes)
 	{
 		check(Removed.TypeRemoved);
 		TemplateData.RemoveTag(*CastChecked<UScriptStruct>(Removed.TypeRemoved));
@@ -208,6 +208,15 @@ bool FMassEntityTemplateBuildContext::ValidateBuildContext(const UWorld& World)
 		}
 	}
 
+	// now to properly test if something required was removed we need to filter TypesAlreadyAdded first
+	for (const FRemovedType& RemovedElement : RemovedTypes)
+	{
+		if (RemovedElement.TypeRemoved)
+		{
+			TypesAlreadyAdded.Remove(RemovedElement.TypeRemoved);
+		}
+	}
+
 	// these are critical, we're going to fail the validation if anything here fails
 	for (const FTraitData& TraitData : TraitsData)
 	{
@@ -222,10 +231,10 @@ bool FMassEntityTemplateBuildContext::ValidateBuildContext(const UWorld& World)
 				{
 					// check if it was removed
 					const UMassEntityTraitBase* RemovedByTrait = nullptr;
-					const int32 RemoverIndex = RemovedTags.Find(FRemovedType({TypeRequired}));
+					const int32 RemoverIndex = RemovedTypes.Find(FRemovedType({TypeRequired}));
 					if (RemoverIndex != INDEX_NONE)
 					{
-						RemovedByTrait = RemovedTags[RemoverIndex].Remover;
+						RemovedByTrait = RemovedTypes[RemoverIndex].Remover;
 					}
 
 					FMassDebugger::DebugEvent<FMassMissingTraitMessage>(TraitData.Trait, TypeRequired, RemovedByTrait);
diff --git a/Engine/Plugins/Runtime/MassGameplay/Source/MassSpawner/Public/MassEntityTemplateRegistry.h b/Engine/Plugins/Runtime/MassGameplay/Source/MassSpawner/Public/MassEntityTemplateRegistry.h
--- a/Engine/Plugins/Runtime/MassGameplay/Source/MassSpawner/Public/MassEntityTemplateRegistry.h
+++ b/Engine/Plugins/Runtime/MassGameplay/Source/MassSpawner/Public/MassEntityTemplateRegistry.h
@@ -138,7 +138,7 @@ struct FMassEntityTemplateBuildContext
 	void RemoveTag(const UScriptStruct& TagType)
 	{
 		checkf(TagType.IsChildOf(FMassTag::StaticStruct()), TEXT("Given struct doesn't represent a valid mass tag type. Make sure to inherit from FMassTag or one of its child-types."));
-		RemovedTags.Add({&TagType
+		RemovedTypes.Add({&TagType
 #if WITH_EDITORONLY_DATA
 			, TraitsData.Last().Trait
 #endif // WITH_EDITORONLY_DATA
@@ -318,7 +318,7 @@ protected:
 		TraitsData.Reset();
 		TraitsProcessed.Reset();
 		IgnoredTraits.Reset();
-		RemovedTags.Reset();
+		RemovedTypes.Reset();
 		bBuildInProgress = false;
 	}
 
@@ -347,7 +347,7 @@ protected:
 	 * These tags will be removed from the resulting entity template
 	 * @see RemoveTag for more details
 	 */
-	TArray<FRemovedType> RemovedTags;
+	TArray<FRemovedType> RemovedTypes;
 
 	bool bBuildInProgress = false;
 
