PR #12343: [Smart Objects] Fix a crash inside USmartObjectZoneAnnotations::DebugDraw()

diff --git a/Engine/Plugins/Runtime/MassGameplay/Source/MassSmartObjects/Private/SmartObjectZoneAnnotations.cpp b/Engine/Plugins/Runtime/MassGameplay/Source/MassSmartObjects/Private/SmartObjectZoneAnnotations.cpp
--- a/Engine/Plugins/Runtime/MassGameplay/Source/MassSmartObjects/Private/SmartObjectZoneAnnotations.cpp
+++ b/Engine/Plugins/Runtime/MassGameplay/Source/MassSmartObjects/Private/SmartObjectZoneAnnotations.cpp
@@ -204,7 +204,7 @@ void USmartObjectZoneAnnotations::DebugDraw(FZoneGraphAnnotationSceneProxy* Debu
 			}
 			const FSmartObjectLaneLocation& SOLaneLocation = AnnotationData.SmartObjectLaneLocations[*Index];
 
-			const FVector& ObjectLocation = Entry.GetComponent()->GetComponentLocation();
+			const FVector& ObjectLocation = Entry.GetTransform().GetLocation();
 			FZoneGraphLaneLocation EntryPointLocation;
 			UE::ZoneGraph::Query::CalculateLocationAlongLane(*ZoneStorage, SOLaneLocation.LaneIndex, SOLaneLocation.DistanceAlongLane, EntryPointLocation);
 			const FColor Color = FColor::Silver;
