[MassGameplay] Added a checked MassEntityManager getter to the MassSpawningSubsystem

diff --git a/Engine/Plugins/Runtime/MassGameplay/Source/MassSpawner/Public/MassSpawnerSubsystem.h b/Engine/Plugins/Runtime/MassGameplay/Source/MassSpawner/Public/MassSpawnerSubsystem.h
--- a/Engine/Plugins/Runtime/MassGameplay/Source/MassSpawner/Public/MassSpawnerSubsystem.h
+++ b/Engine/Plugins/Runtime/MassGameplay/Source/MassSpawner/Public/MassSpawnerSubsystem.h
@@ -46,6 +46,12 @@ public:
 
 	const FMassEntityTemplate* GetMassEntityTemplate(FMassEntityTemplateID TemplateID) const;
 
+	FMassEntityManager& GetEntityManagerChecked()
+	{
+		check(EntityManager.IsValid());
+		return *EntityManager.Get();
+	}
+
 protected:
 	// UWorldSubsystem BEGIN
 	virtual void Initialize(FSubsystemCollectionBase& Collection) override;
