[MassGameplay] Added a way to fetch FMassISMCSharedData indicated by an ISM component

diff --git a/Engine/Plugins/Runtime/MassGameplay/Source/MassRepresentation/Private/MassRepresentationSubsystem.cpp b/Engine/Plugins/Runtime/MassGameplay/Source/MassRepresentation/Private/MassRepresentationSubsystem.cpp
--- a/Engine/Plugins/Runtime/MassGameplay/Source/MassRepresentation/Private/MassRepresentationSubsystem.cpp
+++ b/Engine/Plugins/Runtime/MassGameplay/Source/MassRepresentation/Private/MassRepresentationSubsystem.cpp
@@ -44,6 +44,12 @@ const FMassISMCSharedData* UMassRepresentationSubsystem::GetISMCSharedDataForDes
 	return VisualizationComponent->GetISMCSharedDataForDescriptionIndex(DescriptionIndex);
 }
 
+const FMassISMCSharedData* UMassRepresentationSubsystem::GetISMCSharedDataForInstancedStaticMesh(const UInstancedStaticMeshComponent* ISMC) const
+{
+	check(VisualizationComponent);
+	return VisualizationComponent->GetISMCSharedDataForInstancedStaticMesh(ISMC);
+}
+
 void UMassRepresentationSubsystem::RemoveISMComponent(UInstancedStaticMeshComponent& ISMComponent)
 {
 PRAGMA_DISABLE_DEPRECATION_WARNINGS
diff --git a/Engine/Plugins/Runtime/MassGameplay/Source/MassRepresentation/Private/MassVisualizationComponent.cpp b/Engine/Plugins/Runtime/MassGameplay/Source/MassRepresentation/Private/MassVisualizationComponent.cpp
--- a/Engine/Plugins/Runtime/MassGameplay/Source/MassRepresentation/Private/MassVisualizationComponent.cpp
+++ b/Engine/Plugins/Runtime/MassGameplay/Source/MassRepresentation/Private/MassVisualizationComponent.cpp
@@ -151,6 +151,11 @@ const FMassISMCSharedData* UMassVisualizationComponent::GetISMCSharedDataForDesc
 	return ISMCSharedData.GetDataForIndex(DescriptionIndex);
 }
 
+const FMassISMCSharedData* UMassVisualizationComponent::GetISMCSharedDataForInstancedStaticMesh(const UInstancedStaticMeshComponent* ISMC) const
+{
+	return ISMCSharedData.GetDataForKey(ISMC);
+}
+
 void UMassVisualizationComponent::RemoveVisualDesc(const FStaticMeshInstanceVisualizationDescHandle VisualizationHandle)
 {
 	UE_MT_SCOPED_WRITE_ACCESS(InstancedStaticMeshInfosDetector);
diff --git a/Engine/Plugins/Runtime/MassGameplay/Source/MassRepresentation/Public/MassRepresentationSubsystem.h b/Engine/Plugins/Runtime/MassGameplay/Source/MassRepresentation/Public/MassRepresentationSubsystem.h
--- a/Engine/Plugins/Runtime/MassGameplay/Source/MassRepresentation/Public/MassRepresentationSubsystem.h
+++ b/Engine/Plugins/Runtime/MassGameplay/Source/MassRepresentation/Public/MassRepresentationSubsystem.h
@@ -55,6 +55,11 @@ public:
 	 */
 	const FMassISMCSharedData* GetISMCSharedDataForDescriptionIndex(const int32 DescriptionIndex) const;
 
+	/**
+	 * Fetches FMassISMCSharedData indicated by an ISMC, or nullptr if the ISMC is not represented by any shared data.
+	 */
+	const FMassISMCSharedData* GetISMCSharedDataForInstancedStaticMesh(const UInstancedStaticMeshComponent* ISMC) const;
+
 	/**
 	 * Removes the visualization data associated with the given ISM component. Note that this is safe to do only when
 	 * there are no entities relying on this data. No entity data patching will take place.
diff --git a/Engine/Plugins/Runtime/MassGameplay/Source/MassRepresentation/Public/MassRepresentationTypes.h b/Engine/Plugins/Runtime/MassGameplay/Source/MassRepresentation/Public/MassRepresentationTypes.h
--- a/Engine/Plugins/Runtime/MassGameplay/Source/MassRepresentation/Public/MassRepresentationTypes.h
+++ b/Engine/Plugins/Runtime/MassGameplay/Source/MassRepresentation/Public/MassRepresentationTypes.h
@@ -497,6 +497,14 @@ struct FMassISMCSharedDataMap
 		return Data.IsValidIndex(Index) ? &Data[Index] : nullptr;
 	}
 
+	const FMassISMCSharedData* GetDataForKey(const FISMCSharedDataKey Key) const
+	{
+		const int32* Index = Map.Find(Key);
+		return (Index && Data.IsValidIndex(*Index))
+			? &Data[*Index]
+			: nullptr;
+	}
+
 protected:
 	TArray<FMassISMCSharedData> Data;
 	/** Mapping from Owner (as FObjectKey) of data represented by FMassISMCSharedData to an index to Data */
diff --git a/Engine/Plugins/Runtime/MassGameplay/Source/MassRepresentation/Public/MassVisualizationComponent.h b/Engine/Plugins/Runtime/MassGameplay/Source/MassRepresentation/Public/MassVisualizationComponent.h
--- a/Engine/Plugins/Runtime/MassGameplay/Source/MassRepresentation/Public/MassVisualizationComponent.h
+++ b/Engine/Plugins/Runtime/MassGameplay/Source/MassRepresentation/Public/MassVisualizationComponent.h
@@ -48,6 +48,11 @@ public:
 	 */
 	const FMassISMCSharedData* GetISMCSharedDataForDescriptionIndex(const int32 DescriptionIndex) const;
 
+	/**
+	 * Fetches FMassISMCSharedData indicated by an ISMC, or nullptr if the ISMC is not represented by any shared data.
+	 */
+	const FMassISMCSharedData* GetISMCSharedDataForInstancedStaticMesh(const UInstancedStaticMeshComponent* ISMC) const;
+
 	/** 
 	 * Removes all the visualization data associated with the given ISM component. The function resolves the VisualizationIndex
 	 * associated with the given ISMComponent and calls RemoveVisualDescByIndex which will remove data on all ISMComponents
