[Mass][InstancedActors] The change adds a lot of comments and improves readability in multiple places

diff --git a/Engine/Plugins/Runtime/InstancedActors/Source/InstancedActors/Private/InstancedActorsInitializerProcessor.cpp b/Engine/Plugins/Runtime/InstancedActors/Source/InstancedActors/Private/InstancedActorsInitializerProcessor.cpp
--- a/Engine/Plugins/Runtime/InstancedActors/Source/InstancedActors/Private/InstancedActorsInitializerProcessor.cpp
+++ b/Engine/Plugins/Runtime/InstancedActors/Source/InstancedActors/Private/InstancedActorsInitializerProcessor.cpp
@@ -29,6 +29,11 @@ void UInstancedActorsInitializerProcessor::ConfigureQueries()
 	EntityQuery.AddRequirement<FMassGuidFragment>(EMassFragmentAccess::ReadWrite, EMassFragmentPresence::Optional);
 }
 
+/**
+* Initializes transform and GUID fragments. For the transform, we're going to apply the InstanceActor's parent transform (i.e.,
+* the ActorInstanceManager transform). As for the GUID, we'll assign it to an incremental index, that grows each time a fragment
+* is initialized.
+*/
 template<bool bApplyManagerTranslationOnly, bool bFilterInstanceTransforms>
 void InitInstanceFragments(UInstancedActorsData* InstanceData, int32& NextInstanceIndex, const FVector& ManagerLocation
 	, const FTransform& ManagerTransform, FMassExecutionContext& Context, TArray<FMassEntityHandle>& InOutEntitiesToSignal)
@@ -158,6 +163,7 @@ void UInstancedActorsInitializerProcessor::Execute(FMassEntityManager& EntityMan
 	}
 
 #if DO_CHECK
-	checkf(NumInitializedEntities == InstanceData->NumValidInstances, TEXT("UInstancedActorsInitializerProcessor expects to initialize all spawned entities at once and to have the same number of valid transforms to assign"));
+	checkf(NumInitializedEntities == InstanceData->NumValidInstances
+		, TEXT("UInstancedActorsInitializerProcessor expects to initialize all spawned entities at once and to have the same number of valid transforms to assign"));
 #endif // DO_CHECK
 }
diff --git a/Engine/Plugins/Runtime/InstancedActors/Source/InstancedActors/Private/InstancedActorsManager.cpp b/Engine/Plugins/Runtime/InstancedActors/Source/InstancedActors/Private/InstancedActorsManager.cpp
--- a/Engine/Plugins/Runtime/InstancedActors/Source/InstancedActors/Private/InstancedActorsManager.cpp
+++ b/Engine/Plugins/Runtime/InstancedActors/Source/InstancedActors/Private/InstancedActorsManager.cpp
@@ -674,7 +674,8 @@ UInstancedActorsData* AInstancedActorsManager::CreateNextInstanceActorData(TSubc
 {
 	check(InstancedActorsDataClass);
 
-	UInstancedActorsData* NewInstanceData = NewObject<UInstancedActorsData>(this, InstancedActorsDataClass);
+	const FString InstanceDataNameStr = FString::Printf(TEXT("InstancedActorsData_%s"), *ActorClass->GetFName().ToString());
+	UInstancedActorsData* NewInstanceData = NewObject<UInstancedActorsData>(this, InstancedActorsDataClass, FName(InstanceDataNameStr));
 	// @todo it's conceivable the NextInstanceDataID will overflow. We need to use some handle system in place instead. 
 	NewInstanceData->ID = NextInstanceDataID++;
 	NewInstanceData->ActorClass = ActorClass;
diff --git a/Engine/Plugins/Runtime/InstancedActors/Source/InstancedActors/Private/InstancedActorsRepresentationActorManagement.cpp b/Engine/Plugins/Runtime/InstancedActors/Source/InstancedActors/Private/InstancedActorsRepresentationActorManagement.cpp
--- a/Engine/Plugins/Runtime/InstancedActors/Source/InstancedActors/Private/InstancedActorsRepresentationActorManagement.cpp
+++ b/Engine/Plugins/Runtime/InstancedActors/Source/InstancedActors/Private/InstancedActorsRepresentationActorManagement.cpp
@@ -32,7 +32,7 @@ EMassActorSpawnRequestAction UInstancedActorsRepresentationActorManagement::OnPo
 		{
 			if (USceneComponent* RootComponent = Actor->GetRootComponent())
 			{
-				RootComponent->TransformUpdated.AddUObject(this, &UInstancedActorsRepresentationActorManagement::OnSpawnedBuildingActorMoved, MassActorSpawnRequest.MassAgent);
+				RootComponent->TransformUpdated.AddUObject(this, &UInstancedActorsRepresentationActorManagement::OnSpawnedActorMoved, MassActorSpawnRequest.MassAgent);
 			}
 		}
 
@@ -70,9 +70,8 @@ void UInstancedActorsRepresentationActorManagement::SetActorEnabled(const EMassA
 				});
 		}
 
-		// mz@todo IA: reconsider
+		// mz@todo IA: reconsider for generic use, should probably be configurable.
 		// Don't call Super::SetActorEnabled when enabling, to skip default implementation of enabling tick which we don't want to do
-		// for BuildingActor's etc
 	}
 	else
 	{
@@ -126,7 +125,7 @@ void UInstancedActorsRepresentationActorManagement::OnSpawnedActorDestroyed(AAct
 	}
 }
 
-void UInstancedActorsRepresentationActorManagement::OnSpawnedBuildingActorMoved(USceneComponent* MovedActorRootComponent, EUpdateTransformFlags TransformUpdateFlags, ETeleportType TeleportType, FMassEntityHandle EntityHandle) const
+void UInstancedActorsRepresentationActorManagement::OnSpawnedActorMoved(USceneComponent* MovedActorRootComponent, EUpdateTransformFlags TransformUpdateFlags, ETeleportType TeleportType, FMassEntityHandle EntityHandle) const
 {
 	check(IsValid(MovedActorRootComponent));
 
diff --git a/Engine/Plugins/Runtime/InstancedActors/Source/InstancedActors/Private/InstancedActorsStationaryLODBatchProcessor.cpp b/Engine/Plugins/Runtime/InstancedActors/Source/InstancedActors/Private/InstancedActorsStationaryLODBatchProcessor.cpp
--- a/Engine/Plugins/Runtime/InstancedActors/Source/InstancedActors/Private/InstancedActorsStationaryLODBatchProcessor.cpp
+++ b/Engine/Plugins/Runtime/InstancedActors/Source/InstancedActors/Private/InstancedActorsStationaryLODBatchProcessor.cpp
@@ -206,8 +206,7 @@ void UInstancedActorsStationaryLODBatchProcessor::Execute(FMassEntityManager& En
 						Settings.DetailedRepresentationLODDistance
 					);
 
-					EInstancedActorsBulkLOD NewBulkLOD = EInstancedActorsBulkLOD::Off;
-
+					// Calculates distance sqr from the viewer to the bounds of the InstancedActorManager who owns the FInstancedActorsDataSharedFragment
 					const FBox WorldSpaceBounds = InstanceData->Bounds.TransformBy(InstanceData->GetManagerChecked().GetActorTransform());
 					FVector::FReal DistanceSquared = TNumericLimits<FVector::FReal>::Max();
 
@@ -229,9 +228,11 @@ void UInstancedActorsStationaryLODBatchProcessor::Execute(FMassEntityManager& En
 					}
 #endif
 
-					// Compute scaled squared draw distance to the lowest LOD because cvar could change
+					// Calculates LOD for a given FInstancedActorsDataSharedFragment based on the distance from the viewer to its owner's bounds
+					// NOTE (1): It's called bulk LOD because we're only comparing the viewer to the InstancedActorManager, and not to a specific instance inside it
+					// NOTE (2): We're caching the scaled squared draw distance to the lowest LOD because the cvar could change
 					const float ScaledForceLowLODDrawDistance = InstanceData->LowLODDrawDistance / StaticMeshLODDistanceScale;
-
+					EInstancedActorsBulkLOD NewBulkLOD = EInstancedActorsBulkLOD::Off;
 					if (DistanceSquared < ForcedDetailedLevelDistanceSquared)
 					{
 						NewBulkLOD = EInstancedActorsBulkLOD::Detailed;
@@ -252,18 +253,20 @@ void UInstancedActorsStationaryLODBatchProcessor::Execute(FMassEntityManager& En
 					}
 
 					check(NewBulkLOD != EInstancedActorsBulkLOD::MAX);
+					// Updates the time at which the FInstancedActorsDataSharedFragment will tick depending on its bulk LOD value
 					NextTickTime = CurrentTime + (DelayPerBulkLOD[(int)NewBulkLOD] * 0.95 + FMath::FRand() * 0.1);
 
-					if (ManagerSharedFragment.BulkLOD != NewBulkLOD)
+					if (const bool bHasBulkLODChanged = (ManagerSharedFragment.BulkLOD != NewBulkLOD))
 					{
-						// Dec stats with current state
-						AInstancedActorsManager::UpdateInstanceStats(InstanceData->NumInstances, ManagerSharedFragment.BulkLOD, false);
-
-						ManagerSharedFragment.BulkLOD = NewBulkLOD;
-
-						// Inc stats with new state
-						AInstancedActorsManager::UpdateInstanceStats(InstanceData->NumInstances, ManagerSharedFragment.BulkLOD, true);
+						// Decrements stats with current state, and then increments stats with the new one
+						{
+							AInstancedActorsManager::UpdateInstanceStats(InstanceData->NumInstances, ManagerSharedFragment.BulkLOD, false);
+							ManagerSharedFragment.BulkLOD = NewBulkLOD;
+							AInstancedActorsManager::UpdateInstanceStats(InstanceData->NumInstances, ManagerSharedFragment.BulkLOD, true);
+						}
 
+						// Toggles physics state for the IA's ISM depending on the new bulk LOD value.
+						// If enabled = physics on, else = physics off.				
 						if (UE::Mass::Tweakables::bControlPhysicsState && Settings.bControlPhysicsState)
 						{
 							if (NewBulkLOD == EInstancedActorsBulkLOD::Detailed)
@@ -304,6 +307,8 @@ void UInstancedActorsStationaryLODBatchProcessor::Execute(FMassEntityManager& En
 						}
 
 						{
+							// Toggles visibility for the IA's ISM depending on the new bulk LOD value.
+							// If enabled = use default visibility (probably on), else = physics off.
 							if (NewBulkLOD != EInstancedActorsBulkLOD::Off)
 							{
 								const bool bForcedLowLOD = NewBulkLOD == EInstancedActorsBulkLOD::Low;
@@ -333,28 +338,28 @@ void UInstancedActorsStationaryLODBatchProcessor::Execute(FMassEntityManager& En
 								});
 							}
 						}
-
+						// Toggles MassProcessors on/off depending on the bulk LOD, by pushing or removing tags that are used by those processor's queries.
+						// NOTE: Forcibly updates the mass LOD to off or low when bulk LOD is smaller than Detailed
 						if (ManagerSharedFragment.BulkLOD == EInstancedActorsBulkLOD::Detailed)
 						{
 							EntityManager.Defer().PushCommand<UE::InstancedActors::FEnableDetailedLODCommand>(InstanceData->Entities);
 						}
 						else
 						{
-							// force given LOD for all the hosted entities 
+							// Force given LOD for all the hosted entities 
 							EMassLOD::Type NewLOD = EMassLOD::Off;
 							switch (ManagerSharedFragment.BulkLOD)
 							{
-							case EInstancedActorsBulkLOD::Medium:
-								// NewLOD = EMassLOD::Medium;
-								// break;
-								// right now falling through since we don't have a medium-level visualization
+							case EInstancedActorsBulkLOD::Medium: // right now falling through since we don't have a medium-level visualization
 							case EInstancedActorsBulkLOD::Low:
 								NewLOD = EMassLOD::Low;
 								break;
-							default: // defaulting to Off, as per initial NewLOD value
+							default:
+								NewLOD = EMassLOD::Off;
 								break;
 							}
 
+							// Grabs entity collections from the entities stored by the fragment we're processing, so that we can process them as chunks
 							TArray<FMassArchetypeEntityCollection> EntityCollections;
 							UE::Mass::Utils::CreateEntityCollections(EntityManager, InstanceData->Entities, FMassArchetypeEntityCollection::NoDuplicates, EntityCollections);
 
@@ -397,6 +402,7 @@ void UInstancedActorsStationaryLODBatchProcessor::Execute(FMassEntityManager& En
 		}
 
 		// Consume all pending explicitly dirtied instances to process
+		// NOTE: Those instances are dirtied whenever an InstancedActor is hydrated/dehydrated (check UInstancedActorsData::SetReplicatedActor)
 		TArray<FInstancedActorsInstanceHandle> DirtyRepresentationInstances;
 		InstancedActorSubsystem->PopAllDirtyRepresentationInstances(DirtyRepresentationInstances);
 
@@ -429,13 +435,14 @@ void UInstancedActorsStationaryLODBatchProcessor::Execute(FMassEntityManager& En
 
 			if (DirtyEntitiesByArchetype.Num())
 			{
+				// Converts collected mass entities to collections, which we'll then process afterwards as entity chunks
 				TArray<FMassArchetypeEntityCollection> DirtyEntityCollections;
 				for (TPair<const FMassArchetypeHandle, TArray<FMassEntityHandle>>& Pair : DirtyEntitiesByArchetype)
 				{
 					DirtyEntityCollections.Add(FMassArchetypeEntityCollection(Pair.Key, Pair.Value, FMassArchetypeEntityCollection::EDuplicatesHandling::FoldDuplicates));
 				}
 
-				// Ensure detailed representation update occurs for explicitly dirtied instanced actor entities with non-detailed BulkLOD 
+				// Ensure that detailed representation update occurs for explicitly dirtied instanced actor entities with non-detailed BulkLOD 
 				DirtyVisualizationEntityQuery.ForEachEntityChunkInCollections(DirtyEntityCollections, EntityManager, Context, [](FMassExecutionContext& Context)
 					{
 						// It's possible that we've only just switched to Non-Detailed this frame, the tag removal to prevent regular processing wouldn't have
diff --git a/Engine/Plugins/Runtime/InstancedActors/Source/InstancedActors/Private/InstancedActorsSubsystem.cpp b/Engine/Plugins/Runtime/InstancedActors/Source/InstancedActors/Private/InstancedActorsSubsystem.cpp
--- a/Engine/Plugins/Runtime/InstancedActors/Source/InstancedActors/Private/InstancedActorsSubsystem.cpp
+++ b/Engine/Plugins/Runtime/InstancedActors/Source/InstancedActors/Private/InstancedActorsSubsystem.cpp
@@ -474,7 +474,8 @@ FInstancedActorsInstanceHandle UInstancedActorsSubsystem::InstanceActor(TSubclas
 	FVector CellCenter(ForceInitToZero);
 	
 	// If this is a world partition world we want to be in the centre of a cell.
-	const bool bIsPartitionedWorld = Level->GetWorld()->IsPartitionedWorld();
+	const bool bIsPartitionedLevel = (Level->GetWorldPartitionRuntimeCell() != nullptr);
+	const bool bIsPartitionedWorld = bIsPartitionedLevel || Level->GetWorld()->IsPartitionedWorld();
 	if (bIsPartitionedWorld)
 	{
 		FBox CellBounds = UActorPartitionSubsystem::FCellCoord::GetCellBounds(CellCoord, ManagerGridSize);
@@ -482,7 +483,6 @@ FInstancedActorsInstanceHandle UInstancedActorsSubsystem::InstanceActor(TSubclas
 	}
 
 	// Note: These will be re-compiled at runtime in UInstancedActorsData::BeginPlay, and may differ as such.
-	//TSharedPtr<const FInstancedActorsSettings> Settings = GetOrCompileSettingsForActorClass(ActorClass);
 	FSharedStruct SharedSettings = GetOrCompileSettingsForActorClass(ActorClass);
 	const FInstancedActorsSettings& Settings = SharedSettings.Get<FInstancedActorsSettings>();
 
@@ -944,4 +944,4 @@ bool UInstancedActorsSubsystem::RegisterNewSharedFragmentsInternal(TConstStructV
 	check(SortedSharedFragments.Num() == AllSharedFragmentsOfType.Num());
 
 	return bParamStructFound;
-}
\ No newline at end of file
+}
diff --git a/Engine/Plugins/Runtime/InstancedActors/Source/InstancedActors/Private/InstancedActorsVisualizationTrait.cpp b/Engine/Plugins/Runtime/InstancedActors/Source/InstancedActors/Private/InstancedActorsVisualizationTrait.cpp
--- a/Engine/Plugins/Runtime/InstancedActors/Source/InstancedActors/Private/InstancedActorsVisualizationTrait.cpp
+++ b/Engine/Plugins/Runtime/InstancedActors/Source/InstancedActors/Private/InstancedActorsVisualizationTrait.cpp
@@ -25,6 +25,8 @@ UInstancedActorsVisualizationTrait::UInstancedActorsVisualizationTrait(const FOb
 	bAllowServerSideVisualization = true;
 	RepresentationSubsystemClass = UInstancedActorsRepresentationSubsystem::StaticClass();
 
+	// Avoids registering the Static Mesh Descriptor during BuildTemplate, as it's already added
+	// during UInstancedActorsManagers InitializeModifySpawn flow.
 	bRegisterStaticMeshDesc = false;
 }
 
diff --git a/Engine/Plugins/Runtime/InstancedActors/Source/InstancedActors/Public/InstancedActorsInitializerProcessor.h b/Engine/Plugins/Runtime/InstancedActors/Source/InstancedActors/Public/InstancedActorsInitializerProcessor.h
--- a/Engine/Plugins/Runtime/InstancedActors/Source/InstancedActors/Public/InstancedActorsInitializerProcessor.h
+++ b/Engine/Plugins/Runtime/InstancedActors/Source/InstancedActors/Public/InstancedActorsInitializerProcessor.h
@@ -16,6 +16,7 @@ struct FInstancedActorsMassSpawnData
 	TWeakObjectPtr<UInstancedActorsData> InstanceData;
 };
 
+/** Initializes the fragments of all entities that fit the query specified in ConfigureQueries, which are all considered Instanced Actors. */
 UCLASS()
 class UInstancedActorsInitializerProcessor : public UMassProcessor
 {
diff --git a/Engine/Plugins/Runtime/InstancedActors/Source/InstancedActors/Public/InstancedActorsRepresentationActorManagement.h b/Engine/Plugins/Runtime/InstancedActors/Source/InstancedActors/Public/InstancedActorsRepresentationActorManagement.h
--- a/Engine/Plugins/Runtime/InstancedActors/Source/InstancedActors/Public/InstancedActorsRepresentationActorManagement.h
+++ b/Engine/Plugins/Runtime/InstancedActors/Source/InstancedActors/Public/InstancedActorsRepresentationActorManagement.h
@@ -36,5 +36,5 @@ protected:
 	//~ End UMassRepresentationActorManagement Overrides
 
 	INSTANCEDACTORS_API void OnSpawnedActorDestroyed(AActor& DestroyedActor, FMassEntityHandle EntityHandle) const;
-	INSTANCEDACTORS_API void OnSpawnedBuildingActorMoved(USceneComponent* MovedActorRootComponent, EUpdateTransformFlags TransformUpdateFlags, ETeleportType TeleportType, FMassEntityHandle EntityHandle) const;
+	INSTANCEDACTORS_API void OnSpawnedActorMoved(USceneComponent* MovedActorRootComponent, EUpdateTransformFlags TransformUpdateFlags, ETeleportType TeleportType, FMassEntityHandle EntityHandle) const;
 };
diff --git a/Engine/Plugins/Runtime/InstancedActors/Source/InstancedActors/Public/InstancedActorsSubsystem.h b/Engine/Plugins/Runtime/InstancedActors/Source/InstancedActors/Public/InstancedActorsSubsystem.h
--- a/Engine/Plugins/Runtime/InstancedActors/Source/InstancedActors/Public/InstancedActorsSubsystem.h
+++ b/Engine/Plugins/Runtime/InstancedActors/Source/InstancedActors/Public/InstancedActorsSubsystem.h
@@ -49,14 +49,16 @@ public:
 	* Adds an instance of ActorClass at InstanceTransform location by spawning or reusing a AInstancedActorsManager at InstanceTransform's grid cell location.
 	* @see UInstancedActorsProjectSettings::GridSize
 	*/	
-	INSTANCEDACTORS_API FInstancedActorsInstanceHandle InstanceActor(TSubclassOf<AActor> ActorClass, FTransform InstanceTransform, ULevel* Level, const FGameplayTagContainer& InstanceTags = FGameplayTagContainer());
+	INSTANCEDACTORS_API FInstancedActorsInstanceHandle InstanceActor(TSubclassOf<AActor> ActorClass, FTransform InstanceTransform, ULevel* Level
+		, const FGameplayTagContainer& InstanceTags = FGameplayTagContainer());
 
 	/** 
 	* Adds an instance of ActorClass at InstanceTransform location by spawning or reusing a AInstancedActorsManager at InstanceTransform's grid cell location.
 	* @see UInstancedActorsProjectSettings::GridSize
 	*/	
 	UFUNCTION(BlueprintCallable, Category = InstancedActors)
-	INSTANCEDACTORS_API FInstancedActorsInstanceHandle InstanceActor(TSubclassOf<AActor> ActorClass, FTransform InstanceTransform, ULevel* Level, const FGameplayTagContainer& InstanceTags, TSubclassOf<AInstancedActorsManager> ManagerClass);
+	INSTANCEDACTORS_API FInstancedActorsInstanceHandle InstanceActor(TSubclassOf<AActor> ActorClass, FTransform InstanceTransform, ULevel* Level
+		, const FGameplayTagContainer& InstanceTags, TSubclassOf<AInstancedActorsManager> ManagerClass);
 
 	/**
 	 * Removes all instance data for InstanceHandle.	
@@ -73,9 +75,11 @@ public:
 	INSTANCEDACTORS_API bool RemoveActorInstance(const FInstancedActorsInstanceHandle& InstanceHandle, bool bDestroyManagerIfEmpty = true);
 #endif // WITH_EDITOR
 
-	INSTANCEDACTORS_API void ForEachManager(const FBox& QueryBounds, TFunctionRef<bool(AInstancedActorsManager&)> InOperation, TSubclassOf<AInstancedActorsManager> ManagerClass = AInstancedActorsManager::StaticClass()) const;
+	INSTANCEDACTORS_API void ForEachManager(const FBox& QueryBounds, TFunctionRef<bool(AInstancedActorsManager&)> InOperation
+		, TSubclassOf<AInstancedActorsManager> ManagerClass = AInstancedActorsManager::StaticClass()) const;
 	INSTANCEDACTORS_API void ForEachModifierVolume(const FBox& QueryBounds, TFunctionRef<bool(UInstancedActorsModifierVolumeComponent&)> InOperation) const;
-	INSTANCEDACTORS_API void ForEachInstance(const FBox& QueryBounds, TFunctionRef<bool(const FInstancedActorsInstanceHandle&, const FTransform&, FInstancedActorsIterationContext&)> InOperation) const;
+	INSTANCEDACTORS_API void ForEachInstance(const FBox& QueryBounds, TFunctionRef<bool(const FInstancedActorsInstanceHandle&
+		, const FTransform&, FInstancedActorsIterationContext&)> InOperation) const;
 
 	/** 
 	 * Checks whether there are any instanced actors representing ActorClass or its subclasses inside QueryBounds.
@@ -88,7 +92,7 @@ public:
 	INSTANCEDACTORS_API bool HasInstancesOfClass(const FBox& QueryBounds, TSubclassOf<AActor> ActorClass, const bool bTestActorsIfSpawned = false
 		, const EInstancedActorsBulkLODMask AllowedLODs = EInstancedActorsBulkLODMask::All) const;
 
-	FInstancedActorsManagerHandle AddManager(AInstancedActorsManager& Manager);
+	INSTANCEDACTORS_API FInstancedActorsManagerHandle AddManager(AInstancedActorsManager& Manager);
 	void RemoveManager(FInstancedActorsManagerHandle ManagerHandle);
 
 	FInstancedActorsModifierVolumeHandle AddModifierVolume(UInstancedActorsModifierVolumeComponent& ModifierVolume);
@@ -194,6 +198,11 @@ public:
 	TArray<UInstancedActorsSubsystem::FNextTickSharedFragment>& GetTickableSharedFragments();
 	void UpdateAndResetTickTime(TConstStructView<FInstancedActorsDataSharedFragment> InstancedActorsDataSharedFragment);
 
+	TSubclassOf<AInstancedActorsManager> GetInstancedActorsManagerClass() const 
+	{ 
+		return InstancedActorsManagerClass; 
+	}
+
 protected:
 	/** 
 	 * Fetches all registered FInstancedActorsDataSharedFragment from the EntityManager and adds the missing ones to SortedSharedFragments
diff --git a/Engine/Plugins/Runtime/MassGameplay/Source/MassRepresentation/Private/MassRepresentationProcessor.cpp b/Engine/Plugins/Runtime/MassGameplay/Source/MassRepresentation/Private/MassRepresentationProcessor.cpp
--- a/Engine/Plugins/Runtime/MassGameplay/Source/MassRepresentation/Private/MassRepresentationProcessor.cpp
+++ b/Engine/Plugins/Runtime/MassGameplay/Source/MassRepresentation/Private/MassRepresentationProcessor.cpp
@@ -51,6 +51,7 @@ void UMassRepresentationProcessor::UpdateRepresentation(FMassExecutionContext& C
 {
 	UMassRepresentationSubsystem* RepresentationSubsystem = Context.GetMutableSharedFragment<FMassRepresentationSubsystemSharedFragment>().RepresentationSubsystem;
 	check(RepresentationSubsystem);
+
 	const FMassRepresentationParameters& RepresentationParams = Context.GetConstSharedFragment<FMassRepresentationParameters>();
 	UMassRepresentationActorManagement* RepresentationActorManagement = RepresentationParams.CachedRepresentationActorManagement;
 	check(RepresentationActorManagement);
@@ -58,7 +59,8 @@ void UMassRepresentationProcessor::UpdateRepresentation(FMassExecutionContext& C
 	UMassActorSubsystem* MassActorSubsystem = Context.GetMutableSubsystem<UMassActorSubsystem>();
 
 	FMassEntityManager& CachedEntityManager = Context.GetEntityManagerChecked();
-
+	
+	// Get Transform, Representation, RepresentationLOD and Actor fragments from Context
 	const TConstArrayView<FTransformFragment> TransformList = Context.GetFragmentView<FTransformFragment>();
 	const TArrayView<FMassRepresentationFragment> RepresentationList = Context.GetMutableFragmentView<FMassRepresentationFragment>();
 	const TConstArrayView<FMassRepresentationLODFragment> RepresentationLODList = Context.GetFragmentView<FMassRepresentationLODFragment>();
@@ -66,6 +68,14 @@ void UMassRepresentationProcessor::UpdateRepresentation(FMassExecutionContext& C
 
 	const bool bDoKeepActorExtraFrame = UE::Mass::Representation::bAllowKeepActorExtraFrame ? RepresentationParams.bKeepLowResActors : false;
 
+	// Iterate over all entities, and:
+	// 1. Find their current EMassRepresentationType value based on their current RepresentationLOD;
+	// 2. Change EMassRepresentationType value based on some configs (not all flows will care about this);
+	// 3. Switch the in-game instance representation depending on the EMassRepresentationType;
+	// 		a. If EMassRepresentationType == HighResSpawnedActor or LowResSpawnedActor, sends an Actor SpawnRequest to MassRepresentationActorManagement;
+	// 		b. If EMassRepresentationType == StaticMeshInstance, sends an Actor disable request to MassRepresentationActorManagement;
+	//		c. If EMassRepresentationType == StaticMeshInstance, sends an Actor disable request to MassRepresentationActorManagement.
+	// 		NOTE: I guess this system assumes all instances are already represented by ISMs, which is why we're only dealing with Actor spawn/deactivation
 	const int32 NumEntities = Context.GetNumEntities();
 	for (int32 EntityIdx = 0; EntityIdx < NumEntities; EntityIdx++)
 	{
@@ -79,12 +89,15 @@ void UMassRepresentationProcessor::UpdateRepresentation(FMassExecutionContext& C
 		// Keeping a copy of the that last calculated previous representation
 		const EMassRepresentationType PrevRepresentationCopy = Representation.PrevRepresentation;
 		Representation.PrevRepresentation = Representation.CurrentRepresentation;
-		
+
+		// === 1. Find the current EMassRepresentationType value based on their current RepresentationLOD		
 		EMassRepresentationType WantedRepresentationType = RepresentationParams.LODRepresentation[FMath::Min((int32)RepresentationLOD.LOD, (int32)EMassLOD::Off)];
+		// === 1 end
 
+		// === 2. Change EMassRepresentationType value based on some configs (not all flows will care about this)
 		// Make sure we do not have actor spawned in areas not fully loaded
 		if (Params.bTestCollisionAvailibilityForActorVisualization
-			&& (WantedRepresentationType == EMassRepresentationType::HighResSpawnedActor || WantedRepresentationType == EMassRepresentationType::LowResSpawnedActor) 
+			&& (WantedRepresentationType == EMassRepresentationType::HighResSpawnedActor || WantedRepresentationType == EMassRepresentationType::LowResSpawnedActor)
 			&& !RepresentationSubsystem->IsCollisionLoaded(RepresentationParams.WorldPartitionGridNameContainingCollision, TransformFragment.GetTransform()))
 		{
 			WantedRepresentationType = RepresentationParams.CachedDefaultRepresentationType;
@@ -97,19 +110,29 @@ void UMassRepresentationProcessor::UpdateRepresentation(FMassExecutionContext& C
 		//
 		// Useful for server-authoritative Actor spawning, with replicated Actors inserting themselves into Mass whilst they're
 		// replicated, enforcing actor representation on clients whilst they're present.
-		if (RepresentationParams.bForceActorRepresentationForExternalActors && IsValid(Actor) && !ActorInfo.IsOwnedByMass())
+		//
+		// NOTE:
+		// IsOwnedByMass = Hydrated by mass
+		// !IsOwnedByMass = Hydrated by some external system
+		if (IsValid(Actor))
 		{
-			WantedRepresentationType = Representation.CurrentRepresentation == EMassRepresentationType::LowResSpawnedActor ? EMassRepresentationType::LowResSpawnedActor : EMassRepresentationType::HighResSpawnedActor;
+			if (RepresentationParams.bForceActorRepresentationForExternalActors && !ActorInfo.IsOwnedByMass())
+			{
+				WantedRepresentationType = (Representation.CurrentRepresentation == EMassRepresentationType::LowResSpawnedActor)
+					? EMassRepresentationType::LowResSpawnedActor
+					: EMassRepresentationType::HighResSpawnedActor;
+			}
 		}
-
 		// Has Actor unexpectedly been unset / destroyed since we last ran? 
-		if (!IsValid(Actor) && (Representation.CurrentRepresentation == EMassRepresentationType::LowResSpawnedActor || Representation.CurrentRepresentation == EMassRepresentationType::HighResSpawnedActor))
+		else if (Representation.CurrentRepresentation == EMassRepresentationType::LowResSpawnedActor 
+			|| Representation.CurrentRepresentation == EMassRepresentationType::HighResSpawnedActor)
 		{
 			// Set CurrentRepresentation = None so we get a chance to see CurrentRepresentation != WantedRepresentationType and spawn 
 			// another actor.
 			Representation.CurrentRepresentation = EMassRepresentationType::None;
 		}
-		
+		// === 2 end
+
 		auto DisableActorForISM = [&](AActor*& Actor)
 		{
 			if (!Actor || ActorInfo.IsOwnedByMass())
@@ -141,6 +164,7 @@ void UMassRepresentationProcessor::UpdateRepresentation(FMassExecutionContext& C
 			}
 		};
 
+		// === 3. Switch the in-game instance representation depending on the EMassRepresentationType;
 		// Process switch between representation if there is a change in the representation or there is a pending spawning request
 		if (WantedRepresentationType != Representation.CurrentRepresentation || Representation.ActorSpawnRequestHandle.IsValid())
 		{
@@ -235,12 +259,14 @@ void UMassRepresentationProcessor::UpdateRepresentation(FMassExecutionContext& C
 					break;
 			}
 		}
-		else if (bDoKeepActorExtraFrame && 
-				 Representation.PrevRepresentation == EMassRepresentationType::StaticMeshInstance &&
-			    (PrevRepresentationCopy == EMassRepresentationType::HighResSpawnedActor || PrevRepresentationCopy == EMassRepresentationType::LowResSpawnedActor))
+		else if (bDoKeepActorExtraFrame
+				&& Representation.PrevRepresentation == EMassRepresentationType::StaticMeshInstance 
+				&& (PrevRepresentationCopy == EMassRepresentationType::HighResSpawnedActor 
+					|| PrevRepresentationCopy == EMassRepresentationType::LowResSpawnedActor))
 		{
 			DisableActorForISM(Actor);
 		}
+		// === 3 end
 	}
 
 #if WITH_MASSGAMEPLAY_DEBUG
diff --git a/Engine/Plugins/Runtime/MassGameplay/Source/MassRepresentation/Private/MassStationaryISMSwitcherProcessor.cpp b/Engine/Plugins/Runtime/MassGameplay/Source/MassRepresentation/Private/MassStationaryISMSwitcherProcessor.cpp
--- a/Engine/Plugins/Runtime/MassGameplay/Source/MassRepresentation/Private/MassStationaryISMSwitcherProcessor.cpp
+++ b/Engine/Plugins/Runtime/MassGameplay/Source/MassRepresentation/Private/MassStationaryISMSwitcherProcessor.cpp
@@ -69,10 +69,10 @@ void UMassStationaryISMSwitcherProcessor::ProcessContext(FMassExecutionContext&
 		{
 			continue;
 		}
-		FMassInstancedStaticMeshInfo& ISMInfo = ISMInfosView[Representation.StaticMeshDescHandle.ToIndex()];
 
-		if (Representation.PrevRepresentation == EMassRepresentationType::StaticMeshInstance
-			&& Representation.CurrentRepresentation != EMassRepresentationType::StaticMeshInstance)
+		FMassInstancedStaticMeshInfo& ISMInfo = ISMInfosView[Representation.StaticMeshDescHandle.ToIndex()];
+		if (const bool bSwitchedAwayFromStaticMesh = (Representation.PrevRepresentation == EMassRepresentationType::StaticMeshInstance
+			&& Representation.CurrentRepresentation != EMassRepresentationType::StaticMeshInstance))
 		{
 			// note that we're using the PrevLODSignificance here, and the reason for it is that the Prev value matches the 
 			// PrevRepresentation - thus we need to remove from the "previously" used LODSignificance range.
@@ -86,8 +86,8 @@ void UMassStationaryISMSwitcherProcessor::ProcessContext(FMassExecutionContext&
 				SignalSubsystem.SignalEntity(UE::Mass::Signals::SwitchedToActor, EntityHandle);
 			}
 		}
-		else if (Representation.PrevRepresentation != EMassRepresentationType::StaticMeshInstance
-			&& Representation.CurrentRepresentation == EMassRepresentationType::StaticMeshInstance)
+		else if (const bool bSwitchedToStaticMesh = (Representation.PrevRepresentation != EMassRepresentationType::StaticMeshInstance
+			&& Representation.CurrentRepresentation == EMassRepresentationType::StaticMeshInstance))
 		{
 			const FTransform& Transform = TransformFragment.GetTransform();
 			const FTransform& PrevTransform = Representation.PrevTransform;
diff --git a/Engine/Plugins/Runtime/MassGameplay/Source/MassRepresentation/Private/MassVisualizationComponent.cpp b/Engine/Plugins/Runtime/MassGameplay/Source/MassRepresentation/Private/MassVisualizationComponent.cpp
--- a/Engine/Plugins/Runtime/MassGameplay/Source/MassRepresentation/Private/MassVisualizationComponent.cpp
+++ b/Engine/Plugins/Runtime/MassGameplay/Source/MassRepresentation/Private/MassVisualizationComponent.cpp
@@ -112,9 +112,10 @@ FStaticMeshInstanceVisualizationDescHandle UMassVisualizationComponent::AddVisua
 	
 	UE_MT_SCOPED_WRITE_ACCESS(InstancedStaticMeshInfosDetector);
 
+
+	// 0. Iterate over all meshes in the visualization desc. Each mesh is a descriptor.
 	FStaticMeshInstanceVisualizationDescHandle VisualHandle;
 	TArray<UInstancedStaticMeshComponent*> ISMComponentsUsed;
-
 	for (int32 EntryIndex = 0; EntryIndex < Desc.Meshes.Num(); ++EntryIndex)
 	{
 		const FMassStaticMeshInstanceVisualizationMeshDesc& MeshDesc = Desc.Meshes[EntryIndex];
@@ -125,17 +126,23 @@ FStaticMeshInstanceVisualizationDescHandle UMassVisualizationComponent::AddVisua
 			continue;
 		}
 	
+		// 1. Creates a VisualHandle that will be used for all ISMCs
 		if (!VisualHandle.IsValid())
 		{
 			VisualHandle = AddInstancedStaticMeshInfo(Desc);
 			check(VisualHandle.IsValid());
 		}
 
+		// 2. Stores one FMassISMCSharedData for each ISMC in ISMCSharedData
+		// NOTE: FMassISMCSharedData stores a reference to an ISMC and instance transform updates performed in the current frame.
 		FMassISMCSharedData& NewData = ISMCSharedData.FindOrAdd(ISMComponents[EntryIndex], FMassISMCSharedData(ISMComponents[EntryIndex], /*bInRequiresExternalInstanceIDTracking=*/true));
+		// 3. Stores the newly created FMassISMCSharedData in a VisualHandle -> FMassISMCSharedData array
+		// NOTE: This means that one VisualHandle may reference multiple FMassISMCSharedData
 		InstancedStaticMeshInfos[VisualHandle.ToIndex()].AddISMComponent(NewData);
-		ISMComponentsUsed.Add(ISMComponents[EntryIndex]);
-	
+		// 4. Stores an ISMC -> VisualHandle map, which allows us to later on use an ISMC to query for its FMassISMCSharedData
 		ISMComponentMap.Add(ISMComponents[EntryIndex], VisualHandle);
+
+		ISMComponentsUsed.Add(ISMComponents[EntryIndex]);
 	}
 
 	if (VisualHandle.IsValid())
diff --git a/Engine/Plugins/Runtime/MassGameplay/Source/MassSpawner/Private/MassSpawnerSubsystem.cpp b/Engine/Plugins/Runtime/MassGameplay/Source/MassSpawner/Private/MassSpawnerSubsystem.cpp
--- a/Engine/Plugins/Runtime/MassGameplay/Source/MassSpawner/Private/MassSpawnerSubsystem.cpp
+++ b/Engine/Plugins/Runtime/MassGameplay/Source/MassSpawner/Private/MassSpawnerSubsystem.cpp
@@ -128,19 +128,20 @@ void UMassSpawnerSubsystem::DoSpawning(const FMassEntityTemplate& EntityTemplate
 	//TRACE_CPUPROFILER_EVENT_SCOPE_STR("MassSpawnerSubsystem DoSpawning");
 
 	// 1. Create required number of entities with EntityTemplate.Archetype
-	// 2. Copy data from FMassEntityTemplate.Fragments.
-	//		a. @todo, could be done as part of creation?
-	// 3. Run SpawlDataInitializer if set
-	// 4. "OnEntitiesCreated" notifies will be sent out once the CreationContext gets destroyed (via its destructor).
-
 	TArray<FMassEntityHandle> SpawnedEntities;
 	TSharedRef<FMassEntityManager::FEntityCreationContext> CreationContext
 		= EntityManager->BatchCreateEntities(EntityTemplate.GetArchetype(), EntityTemplate.GetSharedFragmentValues(), NumToSpawn, SpawnedEntities);
 
+	// 2. Copy data from FMassEntityTemplate.Fragments.
+	//		a. @todo, could be done as part of creation?
 	TConstArrayView<FInstancedStruct> FragmentInstances = EntityTemplate.GetInitialFragmentValues();
 	EntityManager->BatchSetEntityFragmentsValues(CreationContext->GetEntityCollections(), FragmentInstances);
 	
-	UMassProcessor* SpawnDataInitializer = SpawnData.IsValid() ? GetSpawnDataInitializer(InitializerClass) : nullptr;
+	// 3. Run SpawnDataInitializer, if set. This is a special type of processor that operates on the entities to initialize them.
+	// e.g., will run UInstancedActorsInitializerProcessor for Mass InstancedActors
+	UMassProcessor* SpawnDataInitializer = SpawnData.IsValid() 
+		? GetSpawnDataInitializer(InitializerClass) 
+		: nullptr;
 
 	if (SpawnDataInitializer)
 	{
@@ -150,6 +151,7 @@ void UMassSpawnerSubsystem::DoSpawning(const FMassEntityTemplate& EntityTemplate
 	}
 
 	OutEntities.Append(MoveTemp(SpawnedEntities));
+	// 4. "OnEntitiesCreated" notifies will be sent out once the CreationContext gets destroyed (via its destructor).
 }
 
 const FMassEntityTemplate* UMassSpawnerSubsystem::GetMassEntityTemplate(FMassEntityTemplateID TemplateID) const
diff --git a/Engine/Source/Developer/MassEntityTestSuite/Private/MassQueryTest.cpp b/Engine/Source/Developer/MassEntityTestSuite/Private/MassQueryTest.cpp
--- a/Engine/Source/Developer/MassEntityTestSuite/Private/MassQueryTest.cpp
+++ b/Engine/Source/Developer/MassEntityTestSuite/Private/MassQueryTest.cpp
@@ -12,16 +12,6 @@
 
 #define LOCTEXT_NAMESPACE "MassTest"
 
-/**
-mz@todo:
-	- add a test for requirement mode (read/write, absent, [future] optional)
-	- test Present/Absent modes for tags
-	- test Present/Absent modes for non-tag fragments
-	- read/write tag fragments 
-	- constructs like Query.AddRequirement<FTestFragment_Tag>(EMassFragmentAccess::ReadWrite);
-	- constructs like Query.AddRequirement<FTestFragment_Int>(EMassFragmentAccess::ReadWrite, Absent);
-*/
-
 
 namespace FMassQueryTest
 {
diff --git a/Engine/Source/Runtime/MassEntity/Private/MassArchetypeData.cpp b/Engine/Source/Runtime/MassEntity/Private/MassArchetypeData.cpp
--- a/Engine/Source/Runtime/MassEntity/Private/MassArchetypeData.cpp
+++ b/Engine/Source/Runtime/MassEntity/Private/MassArchetypeData.cpp
@@ -470,7 +470,7 @@ void FMassArchetypeData::ExecuteFunction(FMassExecutionContext& RunContext, cons
 		return;
 	}
 
-	// mz@todo to be removed
+	// @todo do we really want users to check composition of the archetype being processed at the moment?
 	RunContext.SetCurrentArchetypeCompositionDescriptor(GetCompositionDescriptor());
 
 	uint32 PrevSharedFragmentValuesHash = UINT32_MAX;
@@ -510,7 +510,6 @@ void FMassArchetypeData::ExecuteFunction(FMassExecutionContext& RunContext, cons
 		return;
 	}
 
-	// mz@todo to be removed
 	RunContext.SetCurrentArchetypeCompositionDescriptor(GetCompositionDescriptor());
 
 	uint32 PrevSharedFragmentValuesHash = UINT32_MAX;
@@ -656,6 +655,7 @@ void FMassArchetypeData::GetRequirementsFragmentMapping(TConstArrayView<FMassFra
 	}
 }
 
+// @todo make ChunkRequirements a dedicated type, so that we can ensure that the contents are sorted as expected by the for loop below
 void FMassArchetypeData::GetRequirementsChunkFragmentMapping(TConstArrayView<FMassFragmentRequirementDescription> ChunkRequirements, FMassFragmentIndicesMapping& OutFragmentIndices) const
 {
 	int32 LastFoundFragmentIndex = -1;
@@ -665,7 +665,6 @@ void FMassArchetypeData::GetRequirementsChunkFragmentMapping(TConstArrayView<FMa
 		if (Requirement.RequiresBinding())
 		{
 			int32 FragmentIndex = INDEX_NONE;
-			// mz@todo Add comment here as this code seems to be assuming a certain order for chunk fragments, please explain
 			for (int32 i = LastFoundFragmentIndex + 1; i < ChunkFragmentsTemplate.Num(); ++i)
 			{
 				if (ChunkFragmentsTemplate[i].GetScriptStruct()->IsChildOf(Requirement.StructType))
