[SmartObject] tweaked verbosity for some logs

diff --git a/Engine/Plugins/Runtime/GameplayBehaviorSmartObjects/Source/GameplayBehaviorSmartObjectsModule/Private/AI/BTTask_FindAndUseGameplayBehaviorSmartObject.cpp b/Engine/Plugins/Runtime/GameplayBehaviorSmartObjects/Source/GameplayBehaviorSmartObjectsModule/Private/AI/BTTask_FindAndUseGameplayBehaviorSmartObject.cpp
--- a/Engine/Plugins/Runtime/GameplayBehaviorSmartObjects/Source/GameplayBehaviorSmartObjectsModule/Private/AI/BTTask_FindAndUseGameplayBehaviorSmartObject.cpp
+++ b/Engine/Plugins/Runtime/GameplayBehaviorSmartObjects/Source/GameplayBehaviorSmartObjectsModule/Private/AI/BTTask_FindAndUseGameplayBehaviorSmartObject.cpp
@@ -90,7 +90,7 @@ EBTNodeResult::Type UBTTask_FindAndUseGameplayBehaviorSmartObject::ExecuteTask(U
 					UseClaimedSmartObject(OwnerComp, ClaimHandle, *MyMemory);
 
 					NodeResult = EBTNodeResult::InProgress;
-					UE_VLOG_UELOG(MyController, LogSmartObject, Verbose, TEXT("%s claimed smart object: %s"), *GetNodeName(), *LexToString(ClaimHandle));
+					UE_VLOG_UELOG(MyController, LogSmartObject, Log, TEXT("%s claimed smart object: %s"), *GetNodeName(), *LexToString(ClaimHandle));
 					break;
 				}
 			}
@@ -236,7 +236,7 @@ void UBTTask_FindAndUseGameplayBehaviorSmartObject::OnQueryFinished(TSharedPtr<F
 					UseClaimedSmartObject(*BTComponent, ClaimHandle, *MyMemory);
 					bSmartObjectClaimed = true;
 
-					UE_VLOG_UELOG(BTComponent, LogSmartObject, Verbose, TEXT("%s claimed EQS-found smart object: %s"), *GetNodeName(), *LexToString(ClaimHandle));
+					UE_VLOG_UELOG(BTComponent, LogSmartObject, Log, TEXT("%s claimed EQS-found smart object: %s"), *GetNodeName(), *LexToString(ClaimHandle));
 					break;
 				}
 			}
diff --git a/Engine/Plugins/Runtime/MassGameplay/Source/MassSmartObjects/Private/MassSmartObjectProcessor.cpp b/Engine/Plugins/Runtime/MassGameplay/Source/MassSmartObjects/Private/MassSmartObjectProcessor.cpp
--- a/Engine/Plugins/Runtime/MassGameplay/Source/MassSmartObjects/Private/MassSmartObjectProcessor.cpp
+++ b/Engine/Plugins/Runtime/MassGameplay/Source/MassSmartObjects/Private/MassSmartObjectProcessor.cpp
@@ -82,7 +82,7 @@ void UMassSmartObjectCandidatesFinderProcessor::Execute(FMassEntityManager& Enti
 		Result.bProcessed = true;
 
 #if WITH_MASSGAMEPLAY_DEBUG
-		UE_VLOG(LogOwner, LogSmartObject, Verbose, TEXT("[%s] search completed: found %d"), *Entity.DebugGetDescription(), Result.Candidates.NumSlots);
+		UE_VLOG(LogOwner, LogSmartObject, VeryVerbose, TEXT("[%s] search completed: found %d"), *Entity.DebugGetDescription(), Result.Candidates.NumSlots);
 #endif // WITH_MASSGAMEPLAY_DEBUG
 	};
 
@@ -351,7 +351,7 @@ void UMassSmartObjectTimedBehaviorProcessor::Execute(FMassEntityManager& EntityM
 			if (bIsDebuggingEntity)
 			{
 				UE_CVLOG(bMustRelease, &SmartObjectSubsystem, LogSmartObject, Log, TEXT("[%s] stops using [%s]"), *Entity.DebugGetDescription(), *LexToString(SOUser.InteractionHandle));
-				UE_CVLOG(!bMustRelease, &SmartObjectSubsystem, LogSmartObject, Verbose, TEXT("[%s] using [%s]. Time left: %.1f"), *Entity.DebugGetDescription(), *LexToString(SOUser.InteractionHandle), UseTime);
+				UE_CVLOG(!bMustRelease, &SmartObjectSubsystem, LogSmartObject, VeryVerbose, TEXT("[%s] using [%s]. Time left: %.1f"), *Entity.DebugGetDescription(), *LexToString(SOUser.InteractionHandle), UseTime);
 
 				const TOptional<FTransform> Transform = SmartObjectSubsystem.GetSlotTransform(SOUser.InteractionHandle);
 				if (Transform.IsSet())
