[Mass] headers cleanup

diff --git a/Engine/Plugins/AI/MassAI/Source/MassAIBehavior/Public/MassLookAtTargetTrait.h b/Engine/Plugins/AI/MassAI/Source/MassAIBehavior/Public/MassLookAtTargetTrait.h
--- a/Engine/Plugins/AI/MassAI/Source/MassAIBehavior/Public/MassLookAtTargetTrait.h
+++ b/Engine/Plugins/AI/MassAI/Source/MassAIBehavior/Public/MassLookAtTargetTrait.h
@@ -2,7 +2,6 @@
 
 #pragma once
 
-#include "CoreMinimal.h"
 #include "MassEntityTraitBase.h"
 #include "MassLookAtTargetTrait.generated.h"
 
@@ -15,3 +14,7 @@ protected:
 
 	virtual void BuildTemplate(FMassEntityTemplateBuildContext& BuildContext, const UWorld& World) const override;
 };
+
+#if UE_ENABLE_INCLUDE_ORDER_DEPRECATED_IN_5_6
+#include "CoreMinimal.h"
+#endif // UE_ENABLE_INCLUDE_ORDER_DEPRECATED_IN_5_6
\ No newline at end of file
diff --git a/Engine/Plugins/AI/MassAI/Source/MassAIBehavior/Public/MassLookAtTrait.h b/Engine/Plugins/AI/MassAI/Source/MassAIBehavior/Public/MassLookAtTrait.h
--- a/Engine/Plugins/AI/MassAI/Source/MassAIBehavior/Public/MassLookAtTrait.h
+++ b/Engine/Plugins/AI/MassAI/Source/MassAIBehavior/Public/MassLookAtTrait.h
@@ -2,7 +2,6 @@
 
 #pragma once
 
-#include "CoreMinimal.h"
 #include "MassEntityTraitBase.h"
 #include "MassLookAtTrait.generated.h"
 
@@ -15,3 +14,7 @@ protected:
 
 	virtual void BuildTemplate(FMassEntityTemplateBuildContext& BuildContext, const UWorld& World) const override;
 };
+
+#if UE_ENABLE_INCLUDE_ORDER_DEPRECATED_IN_5_6
+#include "CoreMinimal.h"
+#endif // UE_ENABLE_INCLUDE_ORDER_DEPRECATED_IN_5_6
\ No newline at end of file
diff --git a/Engine/Plugins/AI/MassAI/Source/MassAIBehavior/Public/MassStateTreeTypes.h b/Engine/Plugins/AI/MassAI/Source/MassAIBehavior/Public/MassStateTreeTypes.h
--- a/Engine/Plugins/AI/MassAI/Source/MassAIBehavior/Public/MassStateTreeTypes.h
+++ b/Engine/Plugins/AI/MassAI/Source/MassAIBehavior/Public/MassStateTreeTypes.h
@@ -2,7 +2,6 @@
 
 #pragma once
 
-#include "CoreMinimal.h"
 #include "StateTreeEvaluatorBase.h" 
 #include "StateTreeTaskBase.h"
 #include "MassStateTreeTypes.generated.h"
@@ -65,3 +64,6 @@ protected:
 	int32 Generation = 0;
 };
 
+#if UE_ENABLE_INCLUDE_ORDER_DEPRECATED_IN_5_6
+#include "CoreMinimal.h"
+#endif // UE_ENABLE_INCLUDE_ORDER_DEPRECATED_IN_5_6
\ No newline at end of file
diff --git a/Engine/Plugins/AI/MassAI/Source/MassAIDebug/Private/GameplayDebuggerCategory_Mass.cpp b/Engine/Plugins/AI/MassAI/Source/MassAIDebug/Private/GameplayDebuggerCategory_Mass.cpp
--- a/Engine/Plugins/AI/MassAI/Source/MassAIDebug/Private/GameplayDebuggerCategory_Mass.cpp
+++ b/Engine/Plugins/AI/MassAI/Source/MassAIDebug/Private/GameplayDebuggerCategory_Mass.cpp
@@ -23,13 +23,12 @@
 #include "MassZoneGraphNavigationFragments.h"
 #include "MassSmartObjectFragments.h"
 #include "SmartObjectSubsystem.h"
-#include "Util/ColorConstants.h"
 #include "MassSimulationLOD.h"
 #include "CanvasItem.h"
 #include "Engine/World.h"
 #include "MassDebugger.h"
 #include "MassLODSubsystem.h"
-
+#include "MassMovementFragments.h"
 
 namespace UE::Mass::Debug
 {
diff --git a/Engine/Plugins/AI/MassAI/Source/MassAIDebug/Public/GameplayDebuggerCategory_Mass.h b/Engine/Plugins/AI/MassAI/Source/MassAIDebug/Public/GameplayDebuggerCategory_Mass.h
--- a/Engine/Plugins/AI/MassAI/Source/MassAIDebug/Public/GameplayDebuggerCategory_Mass.h
+++ b/Engine/Plugins/AI/MassAI/Source/MassAIDebug/Public/GameplayDebuggerCategory_Mass.h
@@ -7,7 +7,9 @@
 
 #if WITH_GAMEPLAY_DEBUGGER && WITH_MASSGAMEPLAY_DEBUG
 
+#include "MassEntityTypes.h"
 #include "HAL/IConsoleManager.h"
+#include "UObject/WeakObjectPtrTemplates.h"
 #include "GameplayDebuggerCategory.h"
 
 struct FMassEntityManager;
diff --git a/Engine/Plugins/AI/MassAI/Source/MassAITestSuite/Public/MassAITestSuiteModule.h b/Engine/Plugins/AI/MassAI/Source/MassAITestSuite/Public/MassAITestSuiteModule.h
--- a/Engine/Plugins/AI/MassAI/Source/MassAITestSuite/Public/MassAITestSuiteModule.h
+++ b/Engine/Plugins/AI/MassAI/Source/MassAITestSuite/Public/MassAITestSuiteModule.h
@@ -2,7 +2,6 @@
 
 #pragma  once
 
-#include "CoreMinimal.h"
 #include "Modules/ModuleManager.h"
 
 /**
@@ -35,3 +34,6 @@ public:
 	}
 };
 
+#if UE_ENABLE_INCLUDE_ORDER_DEPRECATED_IN_5_6
+#include "CoreMinimal.h"
+#endif // UE_ENABLE_INCLUDE_ORDER_DEPRECATED_IN_5_6
\ No newline at end of file
diff --git a/Engine/Plugins/AI/MassAI/Source/MassNavigation/Public/MassNavigationFragments.h b/Engine/Plugins/AI/MassAI/Source/MassNavigation/Public/MassNavigationFragments.h
--- a/Engine/Plugins/AI/MassAI/Source/MassNavigation/Public/MassNavigationFragments.h
+++ b/Engine/Plugins/AI/MassAI/Source/MassNavigation/Public/MassNavigationFragments.h
@@ -2,12 +2,7 @@
 
 #pragma once
 
-#if UE_ENABLE_INCLUDE_ORDER_DEPRECATED_IN_5_5
-#include "CoreMinimal.h"
-#endif
-
 #include "AI/Navigation/NavigationElement.h"
-#include "MassCommonFragments.h"
 #include "MassCommonTypes.h"
 #include "MassNavigationSubsystem.h"
 #include "MassNavigationTypes.h"
@@ -221,3 +216,7 @@ struct FNavigationRelevantFragment : public FMassFragment
 	FNavigationElementHandle Handle;
 };
 
+#if UE_ENABLE_INCLUDE_ORDER_DEPRECATED_IN_5_5
+#include "CoreMinimal.h"
+#include "MassCommonFragments.h"
+#endif
diff --git a/Engine/Plugins/AI/MassAI/Source/MassNavigation/Public/MassNavigationSubsystem.h b/Engine/Plugins/AI/MassAI/Source/MassNavigation/Public/MassNavigationSubsystem.h
--- a/Engine/Plugins/AI/MassAI/Source/MassNavigation/Public/MassNavigationSubsystem.h
+++ b/Engine/Plugins/AI/MassAI/Source/MassNavigation/Public/MassNavigationSubsystem.h
@@ -5,6 +5,7 @@
 #include "MassEntityTypes.h"
 #include "HierarchicalHashGrid2D.h"
 #include "MassSubsystemBase.h"
+#include "MassExternalSubsystemTraits.h"
 #include "MassNavigationSubsystem.generated.h"
 
 UENUM()
diff --git a/Engine/Plugins/AI/MassAI/Source/MassNavigation/Public/Steering/MassSteeringFragments.h b/Engine/Plugins/AI/MassAI/Source/MassNavigation/Public/Steering/MassSteeringFragments.h
--- a/Engine/Plugins/AI/MassAI/Source/MassNavigation/Public/Steering/MassSteeringFragments.h
+++ b/Engine/Plugins/AI/MassAI/Source/MassNavigation/Public/Steering/MassSteeringFragments.h
@@ -1,9 +1,8 @@
-﻿// Copyright Epic Games, Inc. All Rights Reserved.
+// Copyright Epic Games, Inc. All Rights Reserved.
 
 #pragma once
 
-#include "CoreMinimal.h"
-#include "MassMovementFragments.h"
+#include "MassEntityTypes.h"
 #include "MassSteeringFragments.generated.h"
 
 class UWorld;
@@ -95,3 +94,8 @@ struct FMassStandingSteeringParameters : public FMassConstSharedFragment
 	UPROPERTY(EditAnywhere, Category = "Standing", meta = (ForceUnits="cm"))
 	float DeadZoneRadius = 15.0f;
 };
+
+#if UE_ENABLE_INCLUDE_ORDER_DEPRECATED_IN_5_6
+#include "CoreMinimal.h"
+#include "MassMovementFragments.h"
+#endif // UE_ENABLE_INCLUDE_ORDER_DEPRECATED_IN_5_6
\ No newline at end of file
diff --git a/Engine/Plugins/AI/MassAI/Source/MassNavigationEditor/Public/MassNavigationEditorModule.h b/Engine/Plugins/AI/MassAI/Source/MassNavigationEditor/Public/MassNavigationEditorModule.h
--- a/Engine/Plugins/AI/MassAI/Source/MassNavigationEditor/Public/MassNavigationEditorModule.h
+++ b/Engine/Plugins/AI/MassAI/Source/MassNavigationEditor/Public/MassNavigationEditorModule.h
@@ -2,9 +2,7 @@
 
 #pragma  once
 
-#include "CoreMinimal.h"
 #include "Modules/ModuleManager.h"
-#include "AssetTypeCategories.h"
 #include "Toolkits/IToolkitHost.h"
 #include "Toolkits/AssetEditorToolkit.h"
 
@@ -18,6 +16,9 @@ public:
 	virtual void StartupModule() override;
 	virtual void ShutdownModule() override;
 	// End IModuleInterface
-
-private:
 };
+
+#if UE_ENABLE_INCLUDE_ORDER_DEPRECATED_IN_5_6
+#include "CoreMinimal.h"
+#include "AssetTypeCategories.h"
+#endif // UE_ENABLE_INCLUDE_ORDER_DEPRECATED_IN_5_6
\ No newline at end of file
diff --git a/Engine/Plugins/AI/MassCrowd/Source/MassCrowd/Public/MassCrowdSubsystem.h b/Engine/Plugins/AI/MassCrowd/Source/MassCrowd/Public/MassCrowdSubsystem.h
--- a/Engine/Plugins/AI/MassCrowd/Source/MassCrowd/Public/MassCrowdSubsystem.h
+++ b/Engine/Plugins/AI/MassCrowd/Source/MassCrowd/Public/MassCrowdSubsystem.h
@@ -5,6 +5,7 @@
 #include "MassCrowdTypes.h"
 #include "ZoneGraphTypes.h"
 #include "MassSubsystemBase.h"
+#include "MassExternalSubsystemTraits.h"
 #include "MassCrowdSubsystem.generated.h"
 
 class UZoneGraphAnnotationSubsystem;
diff --git a/Engine/Plugins/Runtime/MassGameplay/Source/MassActors/Public/MassActorSpawnerSubsystem.h b/Engine/Plugins/Runtime/MassGameplay/Source/MassActors/Public/MassActorSpawnerSubsystem.h
--- a/Engine/Plugins/Runtime/MassGameplay/Source/MassActors/Public/MassActorSpawnerSubsystem.h
+++ b/Engine/Plugins/Runtime/MassGameplay/Source/MassActors/Public/MassActorSpawnerSubsystem.h
@@ -5,11 +5,10 @@
 #include "StructUtils/InstancedStruct.h"
 #include "IndexedHandle.h"
 #include "Delegates/Delegate.h"
-#include "MassCommonTypes.h"
 #include "Templates/SubclassOf.h"
 #include "GameFramework/Actor.h"
 #include "MassSubsystemBase.h"
-
+#include "MassEntityTypes.h"
 #include "MassActorSpawnerSubsystem.generated.h"
 
 
@@ -291,3 +290,7 @@ public:
 	UE_DEPRECATED(5.4, "This flavor of GetNextRequestToSpawn is deprecated. Use the alternative taking an int32& parameter")
 	virtual FMassActorSpawnRequestHandle GetNextRequestToSpawn() const final;
 };
+
+#if UE_ENABLE_INCLUDE_ORDER_DEPRECATED_IN_5_6
+#include "MassCommonTypes.h"
+#endif // UE_ENABLE_INCLUDE_ORDER_DEPRECATED_IN_5_6
\ No newline at end of file
diff --git a/Engine/Plugins/Runtime/MassGameplay/Source/MassActors/Public/MassActorSubsystem.h b/Engine/Plugins/Runtime/MassGameplay/Source/MassActors/Public/MassActorSubsystem.h
--- a/Engine/Plugins/Runtime/MassGameplay/Source/MassActors/Public/MassActorSubsystem.h
+++ b/Engine/Plugins/Runtime/MassGameplay/Source/MassActors/Public/MassActorSubsystem.h
@@ -8,6 +8,7 @@
 #include "Misc/MTAccessDetector.h"
 #include "UObject/ObjectKey.h"
 #include "MassSubsystemBase.h"
+#include "MassExternalSubsystemTraits.h"
 #include "MassActorSubsystem.generated.h"
 
 struct FMassEntityHandle;
diff --git a/Engine/Plugins/Runtime/MassGameplay/Source/MassActors/Public/MassAgentTraits.h b/Engine/Plugins/Runtime/MassGameplay/Source/MassActors/Public/MassAgentTraits.h
--- a/Engine/Plugins/Runtime/MassGameplay/Source/MassActors/Public/MassAgentTraits.h
+++ b/Engine/Plugins/Runtime/MassGameplay/Source/MassActors/Public/MassAgentTraits.h
@@ -2,7 +2,6 @@
 
 #pragma once
 
-#include "CoreMinimal.h"
 #include "MassEntityTraitBase.h"
 #include "MassTranslator.h"
 #include "MassAgentTraits.generated.h"
@@ -64,3 +63,6 @@ protected:
 	virtual void BuildTemplate(FMassEntityTemplateBuildContext& BuildContext, const UWorld& World) const override;
 };
 
+#if UE_ENABLE_INCLUDE_ORDER_DEPRECATED_IN_5_6
+#include "CoreMinimal.h"
+#endif // UE_ENABLE_INCLUDE_ORDER_DEPRECATED_IN_5_6
\ No newline at end of file
diff --git a/Engine/Plugins/Runtime/MassGameplay/Source/MassCommon/Public/MassCommonTypes.h b/Engine/Plugins/Runtime/MassGameplay/Source/MassCommon/Public/MassCommonTypes.h
--- a/Engine/Plugins/Runtime/MassGameplay/Source/MassCommon/Public/MassCommonTypes.h
+++ b/Engine/Plugins/Runtime/MassGameplay/Source/MassCommon/Public/MassCommonTypes.h
@@ -2,9 +2,11 @@
 
 #pragma once
 
+#if UE_ENABLE_INCLUDE_ORDER_DEPRECATED_IN_5_6
 #include "MassEntityTypes.h"
-#include "MassProcessingTypes.h"
 #include "StructUtils/InstancedStruct.h"
+#endif // UE_ENABLE_INCLUDE_ORDER_DEPRECATED_IN_5_6
+#include "MassProcessingTypes.h"
 #include "SequentialID.h"
 #include "MassCommonTypes.generated.h"
 
diff --git a/Engine/Plugins/Runtime/MassGameplay/Source/MassEQS/Public/MassEQSSubsystem.h b/Engine/Plugins/Runtime/MassGameplay/Source/MassEQS/Public/MassEQSSubsystem.h
--- a/Engine/Plugins/Runtime/MassGameplay/Source/MassEQS/Public/MassEQSSubsystem.h
+++ b/Engine/Plugins/Runtime/MassGameplay/Source/MassEQS/Public/MassEQSSubsystem.h
@@ -3,12 +3,10 @@
 #pragma once
 
 #include "Containers/Queue.h"
-
 #include "EnvironmentQuery/EnvQueryNode.h"
 #include "MassEQSTypes.h"
-
 #include "MassSubsystemBase.h"
-
+#include "MassExternalSubsystemTraits.h"
 #include "MassEQSSubsystem.generated.h"
 
 /**
diff --git a/Engine/Plugins/Runtime/MassGameplay/Source/MassEQS/Public/MassEQSTypes.h b/Engine/Plugins/Runtime/MassGameplay/Source/MassEQS/Public/MassEQSTypes.h
--- a/Engine/Plugins/Runtime/MassGameplay/Source/MassEQS/Public/MassEQSTypes.h
+++ b/Engine/Plugins/Runtime/MassGameplay/Source/MassEQS/Public/MassEQSTypes.h
@@ -4,7 +4,7 @@
 
 #include "IndexedHandle.h"
 #include "MassEntityTypes.h"
-
+#include "UObject/Interface.h"
 #include "MassEQSTypes.generated.h"
 
 struct FEnvQueryInstance;
diff --git a/Engine/Plugins/Runtime/MassGameplay/Source/MassGameplayDebug/Public/MassDebuggerSubsystem.h b/Engine/Plugins/Runtime/MassGameplay/Source/MassGameplayDebug/Public/MassDebuggerSubsystem.h
--- a/Engine/Plugins/Runtime/MassGameplay/Source/MassGameplayDebug/Public/MassDebuggerSubsystem.h
+++ b/Engine/Plugins/Runtime/MassGameplay/Source/MassGameplayDebug/Public/MassDebuggerSubsystem.h
@@ -4,6 +4,7 @@
 
 #include "MassGameplayDebugTypes.h"
 #include "MassSubsystemBase.h"
+#include "MassExternalSubsystemTraits.h"
 #include "MassDebuggerSubsystem.generated.h"
 
 
diff --git a/Engine/Plugins/Runtime/MassGameplay/Source/MassGameplayDebug/Public/MassGameplayDebugTypes.h b/Engine/Plugins/Runtime/MassGameplay/Source/MassGameplayDebug/Public/MassGameplayDebugTypes.h
--- a/Engine/Plugins/Runtime/MassGameplay/Source/MassGameplayDebug/Public/MassGameplayDebugTypes.h
+++ b/Engine/Plugins/Runtime/MassGameplay/Source/MassGameplayDebug/Public/MassGameplayDebugTypes.h
@@ -2,7 +2,6 @@
 
 #pragma once
 
-#include "CoreMinimal.h"
 #include "UObject/ObjectMacros.h"
 #include "Logging/LogMacros.h"
 #include "Engine/DataTable.h"
@@ -101,3 +100,7 @@ struct FAgentDebugVisualization : public FTableRowBase
 	UPROPERTY(EditAnywhere, Category = "Mass|Debug")
 	EMassEntityDebugShape WireShape = EMassEntityDebugShape::Box;
 };
+
+#if UE_ENABLE_INCLUDE_ORDER_DEPRECATED_IN_5_6
+#include "CoreMinimal.h"
+#endif // UE_ENABLE_INCLUDE_ORDER_DEPRECATED_IN_5_6
\ No newline at end of file
diff --git a/Engine/Plugins/Runtime/MassGameplay/Source/MassGameplayEditor/Public/MassGameplayEditorModule.h b/Engine/Plugins/Runtime/MassGameplay/Source/MassGameplayEditor/Public/MassGameplayEditorModule.h
--- a/Engine/Plugins/Runtime/MassGameplay/Source/MassGameplayEditor/Public/MassGameplayEditorModule.h
+++ b/Engine/Plugins/Runtime/MassGameplay/Source/MassGameplayEditor/Public/MassGameplayEditorModule.h
@@ -16,8 +16,3 @@ public:
 protected:
 	void RegisterSectionMappings();
 };
-
-#if UE_ENABLE_INCLUDE_ORDER_DEPRECATED_IN_5_3
-#include "CoreMinimal.h"
-#include "Toolkits/AssetEditorToolkit.h"
-#endif
diff --git a/Engine/Plugins/Runtime/MassGameplay/Source/MassGameplayTestSuite/Public/MassGameplayTestSuiteModule.h b/Engine/Plugins/Runtime/MassGameplay/Source/MassGameplayTestSuite/Public/MassGameplayTestSuiteModule.h
--- a/Engine/Plugins/Runtime/MassGameplay/Source/MassGameplayTestSuite/Public/MassGameplayTestSuiteModule.h
+++ b/Engine/Plugins/Runtime/MassGameplay/Source/MassGameplayTestSuite/Public/MassGameplayTestSuiteModule.h
@@ -2,7 +2,6 @@
 
 #pragma once
 
-#include "CoreMinimal.h"
 #include "Modules/ModuleManager.h"
 
 /**
@@ -35,3 +34,6 @@ public:
 	}
 };
 
+#if UE_ENABLE_INCLUDE_ORDER_DEPRECATED_IN_5_6
+#include "CoreMinimal.h"
+#endif // UE_ENABLE_INCLUDE_ORDER_DEPRECATED_IN_5_6
\ No newline at end of file
diff --git a/Engine/Plugins/Runtime/MassGameplay/Source/MassMovement/Private/MassMovementDelegates.cpp b/Engine/Plugins/Runtime/MassGameplay/Source/MassMovement/Private/MassMovementDelegates.cpp
--- a/Engine/Plugins/Runtime/MassGameplay/Source/MassMovement/Private/MassMovementDelegates.cpp
+++ b/Engine/Plugins/Runtime/MassGameplay/Source/MassMovement/Private/MassMovementDelegates.cpp
@@ -1,7 +1,6 @@
 // Copyright Epic Games, Inc. All Rights Reserved.
 
 #include "MassMovementDelegates.h"
-#include "CoreMinimal.h"
 
 namespace UE::MassMovement::Delegates
 {
diff --git a/Engine/Plugins/Runtime/MassGameplay/Source/MassMovement/Public/MassMovementDelegates.h b/Engine/Plugins/Runtime/MassGameplay/Source/MassMovement/Public/MassMovementDelegates.h
--- a/Engine/Plugins/Runtime/MassGameplay/Source/MassMovement/Public/MassMovementDelegates.h
+++ b/Engine/Plugins/Runtime/MassGameplay/Source/MassMovement/Public/MassMovementDelegates.h
@@ -2,7 +2,8 @@
 
 #pragma once
 
-#include "CoreMinimal.h"
+#include "Delegates/Delegate.h"
+#include "Delegates/DelegateCombinations.h"
 
 namespace UE::MassMovement::Delegates
 {
@@ -12,3 +13,7 @@ namespace UE::MassMovement::Delegates
 	extern MASSMOVEMENT_API FOnMassMovementNamesChanged OnMassMovementNamesChanged;
 #endif // WITH_EDITOR
 } // UE::MassMovement::Delegates
+
+#if UE_ENABLE_INCLUDE_ORDER_DEPRECATED_IN_5_6
+#include "CoreMinimal.h"
+#endif // UE_ENABLE_INCLUDE_ORDER_DEPRECATED_IN_5_6
\ No newline at end of file
diff --git a/Engine/Plugins/Runtime/MassGameplay/Source/MassMovement/Public/MassMovementFragments.h b/Engine/Plugins/Runtime/MassGameplay/Source/MassMovement/Public/MassMovementFragments.h
--- a/Engine/Plugins/Runtime/MassGameplay/Source/MassMovement/Public/MassMovementFragments.h
+++ b/Engine/Plugins/Runtime/MassGameplay/Source/MassMovement/Public/MassMovementFragments.h
@@ -2,8 +2,7 @@
 
 #pragma once
 
-#include "CoreMinimal.h"
-#include "MassCommonTypes.h"
+#include "MassEntityTypes.h"
 #include "MassMovementTypes.h"
 #include "RandomSequence.h"
 #include "MassMovementFragments.generated.h"
@@ -91,3 +90,8 @@ struct MASSMOVEMENT_API FMassMovementParameters : public FMassConstSharedFragmen
 	UPROPERTY(EditAnywhere, Category = "Movement")
 	TArray<FMassMovementStyleParameters> MovementStyles;
 };
+
+#if UE_ENABLE_INCLUDE_ORDER_DEPRECATED_IN_5_6
+#include "CoreMinimal.h"
+#include "MassCommonTypes.h"
+#endif // UE_ENABLE_INCLUDE_ORDER_DEPRECATED_IN_5_6
diff --git a/Engine/Plugins/Runtime/MassGameplay/Source/MassMovement/Public/MassMovementSettings.h b/Engine/Plugins/Runtime/MassGameplay/Source/MassMovement/Public/MassMovementSettings.h
--- a/Engine/Plugins/Runtime/MassGameplay/Source/MassMovement/Public/MassMovementSettings.h
+++ b/Engine/Plugins/Runtime/MassGameplay/Source/MassMovement/Public/MassMovementSettings.h
@@ -2,7 +2,6 @@
 
 #pragma once
 
-#include "CoreMinimal.h"
 #include "MassSettings.h"
 #include "MassMovementTypes.h"
 #include "MassMovementSettings.generated.h"
@@ -32,3 +31,7 @@ private:
 	UPROPERTY(EditAnywhere, config, Category = Movement);
     TArray<FMassMovementStyle> MovementStyles;
 };
+
+#if UE_ENABLE_INCLUDE_ORDER_DEPRECATED_IN_5_6
+#include "CoreMinimal.h"
+#endif // UE_ENABLE_INCLUDE_ORDER_DEPRECATED_IN_5_6
\ No newline at end of file
diff --git a/Engine/Plugins/Runtime/MassGameplay/Source/MassMovementEditor/Public/MassMovementEditorModule.h b/Engine/Plugins/Runtime/MassGameplay/Source/MassMovementEditor/Public/MassMovementEditorModule.h
--- a/Engine/Plugins/Runtime/MassGameplay/Source/MassMovementEditor/Public/MassMovementEditorModule.h
+++ b/Engine/Plugins/Runtime/MassGameplay/Source/MassMovementEditor/Public/MassMovementEditorModule.h
@@ -2,9 +2,7 @@
 
 #pragma  once
 
-#include "CoreMinimal.h"
 #include "Modules/ModuleManager.h"
-#include "AssetTypeCategories.h"
 #include "Toolkits/IToolkitHost.h"
 #include "Toolkits/AssetEditorToolkit.h"
 
@@ -20,6 +18,9 @@ public:
 	// Begin IModuleInterface
 	virtual void StartupModule() override;
 	virtual void ShutdownModule() override;
-
-private:
 };
+
+#if UE_ENABLE_INCLUDE_ORDER_DEPRECATED_IN_5_6
+#include "CoreMinimal.h"
+#include "AssetTypeCategories.h"
+#endif // UE_ENABLE_INCLUDE_ORDER_DEPRECATED_IN_5_6
diff --git a/Engine/Plugins/Runtime/MassGameplay/Source/MassMovementEditor/Public/MassMovementPropertyUtils.h b/Engine/Plugins/Runtime/MassGameplay/Source/MassMovementEditor/Public/MassMovementPropertyUtils.h
--- a/Engine/Plugins/Runtime/MassGameplay/Source/MassMovementEditor/Public/MassMovementPropertyUtils.h
+++ b/Engine/Plugins/Runtime/MassGameplay/Source/MassMovementEditor/Public/MassMovementPropertyUtils.h
@@ -1,6 +1,6 @@
 ﻿// Copyright Epic Games, Inc. All Rights Reserved.
 #pragma once
-#include "CoreMinimal.h"
+
 #include "Editor.h"
 #include "PropertyHandle.h"
 
@@ -110,4 +110,8 @@ void SetValue(TSharedPtr<IPropertyHandle>& ValueProperty, T NewValue, EPropertyV
 
 } // UE::MassMovement::PropertyUtils
 
-#undef LOCTEXT_NAMESPACE
\ No newline at end of file
+#undef LOCTEXT_NAMESPACE
+
+#if UE_ENABLE_INCLUDE_ORDER_DEPRECATED_IN_5_6
+#include "CoreMinimal.h"
+#endif // UE_ENABLE_INCLUDE_ORDER_DEPRECATED_IN_5_6
\ No newline at end of file
diff --git a/Engine/Plugins/Runtime/MassGameplay/Source/MassMovementEditor/Public/MassMovementStyleRefDetails.h b/Engine/Plugins/Runtime/MassGameplay/Source/MassMovementEditor/Public/MassMovementStyleRefDetails.h
--- a/Engine/Plugins/Runtime/MassGameplay/Source/MassMovementEditor/Public/MassMovementStyleRefDetails.h
+++ b/Engine/Plugins/Runtime/MassGameplay/Source/MassMovementEditor/Public/MassMovementStyleRefDetails.h
@@ -1,6 +1,6 @@
 // Copyright Epic Games, Inc. All Rights Reserved.
 #pragma once
-#include "CoreMinimal.h"
+
 #include "UnrealClient.h"
 #include "IPropertyTypeCustomization.h"
 
@@ -32,4 +32,8 @@ private:
 
 	class IPropertyUtilities* PropUtils = nullptr;
 	TSharedPtr<IPropertyHandle> StructProperty;
-};
\ No newline at end of file
+};
+
+#if UE_ENABLE_INCLUDE_ORDER_DEPRECATED_IN_5_6
+#include "CoreMinimal.h"
+#endif // UE_ENABLE_INCLUDE_ORDER_DEPRECATED_IN_5_6
diff --git a/Engine/Plugins/Runtime/MassGameplay/Source/MassReplication/Public/MassReplicationSubsystem.h b/Engine/Plugins/Runtime/MassGameplay/Source/MassReplication/Public/MassReplicationSubsystem.h
--- a/Engine/Plugins/Runtime/MassGameplay/Source/MassReplication/Public/MassReplicationSubsystem.h
+++ b/Engine/Plugins/Runtime/MassGameplay/Source/MassReplication/Public/MassReplicationSubsystem.h
@@ -6,7 +6,7 @@
 #include "MassEntityTypes.h"
 #include "MassReplicationTypes.h"
 #include "MassSubsystemBase.h"
-
+#include "MassExternalSubsystemTraits.h"
 #include "MassReplicationSubsystem.generated.h"
 
 struct FMassEntityManager;
diff --git a/Engine/Plugins/Runtime/MassGameplay/Source/MassRepresentation/Public/MassRepresentationSubsystem.h b/Engine/Plugins/Runtime/MassGameplay/Source/MassRepresentation/Public/MassRepresentationSubsystem.h
--- a/Engine/Plugins/Runtime/MassGameplay/Source/MassRepresentation/Public/MassRepresentationSubsystem.h
+++ b/Engine/Plugins/Runtime/MassGameplay/Source/MassRepresentation/Public/MassRepresentationSubsystem.h
@@ -3,11 +3,11 @@
 #pragma once
 
 #include "GameFramework/Actor.h"
-#include "MassCommonTypes.h"
 #include "Misc/MTAccessDetector.h"
 #include "MassRepresentationTypes.h"
 #include "MassActorSpawnerSubsystem.h"
 #include "MassSubsystemBase.h"
+#include "MassExternalSubsystemTraits.h"
 #include "MassRepresentationSubsystem.generated.h"
 
 class UMassVisualizationComponent;
@@ -233,3 +233,7 @@ struct TMassExternalSubsystemTraits<UMassRepresentationSubsystem> final
 		GameThreadOnly = true
 	};
 };
+
+#if UE_ENABLE_INCLUDE_ORDER_DEPRECATED_IN_5_6
+#include "MassCommonTypes.h"
+#endif // UE_ENABLE_INCLUDE_ORDER_DEPRECATED_IN_5_6
\ No newline at end of file
diff --git a/Engine/Plugins/Runtime/MassGameplay/Source/MassSignals/Public/MassSignalSubsystem.h b/Engine/Plugins/Runtime/MassGameplay/Source/MassSignals/Public/MassSignalSubsystem.h
--- a/Engine/Plugins/Runtime/MassGameplay/Source/MassSignals/Public/MassSignalSubsystem.h
+++ b/Engine/Plugins/Runtime/MassGameplay/Source/MassSignals/Public/MassSignalSubsystem.h
@@ -5,6 +5,7 @@
 #include "MassEntityManager.h"
 #include "MassSubsystemBase.h"
 #include "Misc/MTAccessDetector.h"
+#include "MassExternalSubsystemTraits.h"
 #include "MassSignalSubsystem.generated.h"
 
 namespace UE::MassSignal 
diff --git a/Engine/Plugins/Runtime/MassGameplay/Source/MassSignals/Public/MassSignalTypes.h b/Engine/Plugins/Runtime/MassGameplay/Source/MassSignals/Public/MassSignalTypes.h
--- a/Engine/Plugins/Runtime/MassGameplay/Source/MassSignals/Public/MassSignalTypes.h
+++ b/Engine/Plugins/Runtime/MassGameplay/Source/MassSignals/Public/MassSignalTypes.h
@@ -2,13 +2,11 @@
 
 #pragma once
 
-#include "CoreMinimal.h"
 #include "Logging/LogMacros.h"
 #include "MassEntityManager.h"
-//#include "MassSignalTypes.generated.h"
 
-DECLARE_LOG_CATEGORY_EXTERN(LogMassSignals, Log, All)
 
+DECLARE_LOG_CATEGORY_EXTERN(LogMassSignals, Log, All)
 
 /**
  * MassSignalNameLookup stores list of Signal names for each entity. The names are stored per entity as a bitmask,
@@ -48,4 +46,8 @@ protected:
 
 	/** Map from entity id to name bitmask */
 	TMap<FMassEntityHandle, uint64> EntitySignals;
-};
\ No newline at end of file
+};
+
+#if UE_ENABLE_INCLUDE_ORDER_DEPRECATED_IN_5_6
+#include "CoreMinimal.h"
+#endif // UE_ENABLE_INCLUDE_ORDER_DEPRECATED_IN_5_6
diff --git a/Engine/Plugins/Runtime/MassGameplay/Source/MassSpawner/Public/MassAssortedFragmentsTrait.h b/Engine/Plugins/Runtime/MassGameplay/Source/MassSpawner/Public/MassAssortedFragmentsTrait.h
--- a/Engine/Plugins/Runtime/MassGameplay/Source/MassSpawner/Public/MassAssortedFragmentsTrait.h
+++ b/Engine/Plugins/Runtime/MassGameplay/Source/MassSpawner/Public/MassAssortedFragmentsTrait.h
@@ -2,7 +2,6 @@
 
 #pragma once 
 
-#include "CoreMinimal.h"
 #include "MassEntityTraitBase.h"
 #include "StructUtils/InstancedStruct.h"
 #include "MassAssortedFragmentsTrait.generated.h"
@@ -25,3 +24,7 @@ protected:
 	UPROPERTY(Category="Mass", EditAnywhere, meta = (BaseStruct = "/Script/MassEntity.MassTag", ExcludeBaseStruct))
 	TArray<FInstancedStruct> Tags;
 };
+
+#if UE_ENABLE_INCLUDE_ORDER_DEPRECATED_IN_5_6
+#include "CoreMinimal.h"
+#endif // UE_ENABLE_INCLUDE_ORDER_DEPRECATED_IN_5_6
diff --git a/Engine/Plugins/Runtime/MassGameplay/Source/MassSpawner/Public/MassEntityConfigAsset.h b/Engine/Plugins/Runtime/MassGameplay/Source/MassSpawner/Public/MassEntityConfigAsset.h
--- a/Engine/Plugins/Runtime/MassGameplay/Source/MassSpawner/Public/MassEntityConfigAsset.h
+++ b/Engine/Plugins/Runtime/MassGameplay/Source/MassSpawner/Public/MassEntityConfigAsset.h
@@ -2,10 +2,8 @@
 
 #pragma once 
 
-#include "CoreMinimal.h"
 #include "Engine/DataAsset.h"
 #include "MassEntityTemplate.h"
-#include "MassEntityTemplateRegistry.h"
 #include "MassEntityConfigAsset.generated.h"
 
 class UMassEntityConfigAsset;
@@ -153,4 +151,9 @@ protected:
 	/** The config described in this asset. */
 	UPROPERTY(Category = "Entity Config", EditAnywhere)
 	FMassEntityConfig Config;
-};
\ No newline at end of file
+};
+
+#if UE_ENABLE_INCLUDE_ORDER_DEPRECATED_IN_5_6
+#include "CoreMinimal.h"
+#include "MassEntityTemplateRegistry.h"
+#endif // UE_ENABLE_INCLUDE_ORDER_DEPRECATED_IN_5_6
diff --git a/Engine/Plugins/Runtime/MassGameplay/Source/MassSpawner/Public/MassEntityEQSSpawnPointsGenerator.h b/Engine/Plugins/Runtime/MassGameplay/Source/MassSpawner/Public/MassEntityEQSSpawnPointsGenerator.h
--- a/Engine/Plugins/Runtime/MassGameplay/Source/MassSpawner/Public/MassEntityEQSSpawnPointsGenerator.h
+++ b/Engine/Plugins/Runtime/MassGameplay/Source/MassSpawner/Public/MassEntityEQSSpawnPointsGenerator.h
@@ -2,7 +2,6 @@
 
 #pragma once 
 
-#include "CoreMinimal.h"
 #include "MassEntitySpawnDataGeneratorBase.h"
 #include "EnvironmentQuery/EnvQueryTypes.h"
 #include "MassEntityEQSSpawnPointsGenerator.generated.h"
@@ -31,3 +30,7 @@ protected:
 	UPROPERTY(Category = "Query", EditAnywhere)
 	FEQSParametrizedQueryExecutionRequest EQSRequest;
 };
+
+#if UE_ENABLE_INCLUDE_ORDER_DEPRECATED_IN_5_6
+#include "CoreMinimal.h"
+#endif // UE_ENABLE_INCLUDE_ORDER_DEPRECATED_IN_5_6
diff --git a/Engine/Plugins/Runtime/MassGameplay/Source/MassSpawner/Public/MassEntitySpawnDataGeneratorBase.h b/Engine/Plugins/Runtime/MassGameplay/Source/MassSpawner/Public/MassEntitySpawnDataGeneratorBase.h
--- a/Engine/Plugins/Runtime/MassGameplay/Source/MassSpawner/Public/MassEntitySpawnDataGeneratorBase.h
+++ b/Engine/Plugins/Runtime/MassGameplay/Source/MassSpawner/Public/MassEntitySpawnDataGeneratorBase.h
@@ -2,7 +2,6 @@
 
 #pragma once 
 
-#include "CoreMinimal.h"
 #include "StructUtils/InstancedStruct.h"
 #include "MassSpawnerTypes.h"
 #include "MassEntitySpawnDataGeneratorBase.generated.h"
@@ -76,3 +75,7 @@ protected:
 	UPROPERTY()
 	uint8 bUseCustomSeed : 1;
 };
+
+#if UE_ENABLE_INCLUDE_ORDER_DEPRECATED_IN_5_6
+#include "CoreMinimal.h"
+#endif // UE_ENABLE_INCLUDE_ORDER_DEPRECATED_IN_5_6
diff --git a/Engine/Plugins/Runtime/MassGameplay/Source/MassSpawner/Public/MassEntityTemplateRegistry.h b/Engine/Plugins/Runtime/MassGameplay/Source/MassSpawner/Public/MassEntityTemplateRegistry.h
--- a/Engine/Plugins/Runtime/MassGameplay/Source/MassSpawner/Public/MassEntityTemplateRegistry.h
+++ b/Engine/Plugins/Runtime/MassGameplay/Source/MassSpawner/Public/MassEntityTemplateRegistry.h
@@ -2,9 +2,7 @@
 
 #pragma once
 
-#include "CoreMinimal.h"
 #include "UObject/Object.h"
-#include "MassCommonTypes.h"
 #include "MassEntityTemplate.h"
 #include "MassTranslator.h"
 #include "MassEntityTemplateRegistry.generated.h"
@@ -415,3 +413,8 @@ protected:
 };
 
 #undef ENSURE_SUPPORTED_TRAIT_OPERATION
+
+#if UE_ENABLE_INCLUDE_ORDER_DEPRECATED_IN_5_6
+#include "CoreMinimal.h"
+#include "MassCommonTypes.h"
+#endif // UE_ENABLE_INCLUDE_ORDER_DEPRECATED_IN_5_6
diff --git a/Engine/Plugins/Runtime/MassGameplay/Source/MassSpawner/Public/MassEntityZoneGraphSpawnPointsGenerator.h b/Engine/Plugins/Runtime/MassGameplay/Source/MassSpawner/Public/MassEntityZoneGraphSpawnPointsGenerator.h
--- a/Engine/Plugins/Runtime/MassGameplay/Source/MassSpawner/Public/MassEntityZoneGraphSpawnPointsGenerator.h
+++ b/Engine/Plugins/Runtime/MassGameplay/Source/MassSpawner/Public/MassEntityZoneGraphSpawnPointsGenerator.h
@@ -2,10 +2,8 @@
 
 #pragma once 
 
-#include "CoreMinimal.h"
 #include "MassEntitySpawnDataGeneratorBase.h"
 #include "ZoneGraphTypes.h"
-
 #include "MassEntityZoneGraphSpawnPointsGenerator.generated.h"
 
 class AZoneGraphData;
@@ -36,3 +34,7 @@ protected:
 	UPROPERTY(EditAnywhere, Category = "ZoneGraph Generator Config")
 	float MaxGap = 300;
 };
+
+#if UE_ENABLE_INCLUDE_ORDER_DEPRECATED_IN_5_6
+#include "CoreMinimal.h"
+#endif // UE_ENABLE_INCLUDE_ORDER_DEPRECATED_IN_5_6
diff --git a/Engine/Source/Developer/MassEntityTestSuite/Public/MassEntityTestFarmPlot.h b/Engine/Source/Developer/MassEntityTestSuite/Public/MassEntityTestFarmPlot.h
--- a/Engine/Source/Developer/MassEntityTestSuite/Public/MassEntityTestFarmPlot.h
+++ b/Engine/Source/Developer/MassEntityTestSuite/Public/MassEntityTestFarmPlot.h
@@ -2,15 +2,12 @@
 
 #pragma once
 
-#include "CoreMinimal.h"
 #include "GameFramework/Actor.h"
 #include "Engine/DataTable.h"
-
 #include "MassEntityTypes.h"
 #include "MassEntityManager.h"
 #include "MassCommandBuffer.h"
 #include "MassProcessor.h"
-
 #include "MassEntityTestFarmPlot.generated.h"
 
 class UStaticMesh;
@@ -249,3 +246,7 @@ public:
 	virtual void BeginPlay() override;
 	virtual void TickActor(float DeltaTime, enum ELevelTick TickType, FActorTickFunction& ThisTickFunction) override;
 };
+
+#if UE_ENABLE_INCLUDE_ORDER_DEPRECATED_IN_5_6
+#include "CoreMinimal.h"
+#endif // UE_ENABLE_INCLUDE_ORDER_DEPRECATED_IN_5_6
diff --git a/Engine/Source/Editor/MassEntityEditor/Public/MassEntityEditorModule.h b/Engine/Source/Editor/MassEntityEditor/Public/MassEntityEditorModule.h
--- a/Engine/Source/Editor/MassEntityEditor/Public/MassEntityEditorModule.h
+++ b/Engine/Source/Editor/MassEntityEditor/Public/MassEntityEditorModule.h
@@ -46,10 +46,3 @@ protected:
 	TSharedPtr<FExtensibilityManager> MenuExtensibilityManager;
 	TSharedPtr<FExtensibilityManager> ToolBarExtensibilityManager;
 };
-
-#if UE_ENABLE_INCLUDE_ORDER_DEPRECATED_IN_5_3
-#include "CoreMinimal.h"
-#include "Modules/ModuleManager.h"
-#include "AssetTypeCategories.h"
-#include "Toolkits/IToolkitHost.h"
-#endif
diff --git a/Engine/Source/Runtime/MassEntity/Public/MassEntityTypes.h b/Engine/Source/Runtime/MassEntity/Public/MassEntityTypes.h
--- a/Engine/Source/Runtime/MassEntity/Public/MassEntityTypes.h
+++ b/Engine/Source/Runtime/MassEntity/Public/MassEntityTypes.h
@@ -6,7 +6,6 @@
 #include "MassProcessingTypes.h"
 #include "StructUtils/StructArrayView.h"
 #include "Subsystems/Subsystem.h"
-#include "MassExternalSubsystemTraits.h"
 #include "StructUtils/SharedStruct.h"
 #include "MassEntityTypes.generated.h"
 #ifdef WITH_AITESTSUITE
@@ -803,3 +802,7 @@ struct FMassArchetypeCreationParams
 	/** Name to identify the archetype while debugging*/
 	FName DebugName;
 };
+
+#if UE_ENABLE_INCLUDE_ORDER_DEPRECATED_IN_5_6
+#include "MassExternalSubsystemTraits.h"
+#endif // UE_ENABLE_INCLUDE_ORDER_DEPRECATED_IN_5_6
diff --git a/Engine/Source/Runtime/MassEntity/Public/MassProcessingPhaseManager.h b/Engine/Source/Runtime/MassEntity/Public/MassProcessingPhaseManager.h
--- a/Engine/Source/Runtime/MassEntity/Public/MassProcessingPhaseManager.h
+++ b/Engine/Source/Runtime/MassEntity/Public/MassProcessingPhaseManager.h
@@ -2,7 +2,6 @@
 
 #pragma once
 
-#include "CoreMinimal.h"
 #include "UObject/UObjectGlobals.h"
 #include "UObject/Object.h"
 #include "UObject/GCObject.h"
@@ -225,3 +224,7 @@ protected:
 	EProcessorExecutionFlags ProcessorExecutionFlags = EProcessorExecutionFlags::None;
 	bool bIsAllowedToTick = false;
 };
+
+#if UE_ENABLE_INCLUDE_ORDER_DEPRECATED_IN_5_6
+#include "CoreMinimal.h"
+#endif // UE_ENABLE_INCLUDE_ORDER_DEPRECATED_IN_5_6
diff --git a/Engine/Source/Runtime/MassEntity/Public/MassProcessor.h b/Engine/Source/Runtime/MassEntity/Public/MassProcessor.h
--- a/Engine/Source/Runtime/MassEntity/Public/MassProcessor.h
+++ b/Engine/Source/Runtime/MassEntity/Public/MassProcessor.h
@@ -2,7 +2,6 @@
 
 #pragma once
 
-#include "CoreMinimal.h"
 #include "MassEntityManager.h"
 #include "MassProcessingTypes.h"
 #include "Async/TaskGraphInterfaces.h"
@@ -424,3 +423,7 @@ inline bool UMassCompositeProcessor::IsEmpty() const
 {
 	return ChildPipeline.IsEmpty();
 }
+
+#if UE_ENABLE_INCLUDE_ORDER_DEPRECATED_IN_5_6
+#include "CoreMinimal.h"
+#endif // UE_ENABLE_INCLUDE_ORDER_DEPRECATED_IN_5_6
