[Mass] Issue 623014 fix, added missing headers

diff --git a/Engine/Plugins/Runtime/MassGameplay/Source/MassCommon/Private/MassCommonUtils.cpp b/Engine/Plugins/Runtime/MassGameplay/Source/MassCommon/Private/MassCommonUtils.cpp
--- a/Engine/Plugins/Runtime/MassGameplay/Source/MassCommon/Private/MassCommonUtils.cpp
+++ b/Engine/Plugins/Runtime/MassGameplay/Source/MassCommon/Private/MassCommonUtils.cpp
@@ -2,6 +2,7 @@
 
 #include "MassCommonUtils.h"
 #include "Misc/App.h"
+#include "HAL/IConsoleManager.h"
 
 namespace UE::Mass::Private
 {
diff --git a/Engine/Plugins/Runtime/MassGameplay/Source/MassCommon/Public/MassCommonUtils.h b/Engine/Plugins/Runtime/MassGameplay/Source/MassCommon/Public/MassCommonUtils.h
--- a/Engine/Plugins/Runtime/MassGameplay/Source/MassCommon/Public/MassCommonUtils.h
+++ b/Engine/Plugins/Runtime/MassGameplay/Source/MassCommon/Public/MassCommonUtils.h
@@ -5,6 +5,7 @@
 #include "MassCommonTypes.h"
 #include "Containers/Array.h"
 #include "Containers/Queue.h"
+#include "MassEntityTypes.h"
 
 namespace UE::Mass::Utils
 {
