[MassGameplay] More non-unity/monolithic fixes

#rnx

diff --git a/Engine/Plugins/Runtime/MassGameplay/Source/MassEQS/Private/MassEQSBlueprintLibrary.cpp b/Engine/Plugins/Runtime/MassGameplay/Source/MassEQS/Private/MassEQSBlueprintLibrary.cpp
--- a/Engine/Plugins/Runtime/MassGameplay/Source/MassEQS/Private/MassEQSBlueprintLibrary.cpp
+++ b/Engine/Plugins/Runtime/MassGameplay/Source/MassEQS/Private/MassEQSBlueprintLibrary.cpp
@@ -1,7 +1,8 @@
 // Copyright Epic Games, Inc. All Rights Reserved.
 
 #include "MassEQSBlueprintLibrary.h"
-
+#include "Engine/World.h"
+#include "GameFramework/Actor.h"
 #include "MassCommonFragments.h"
 #include "MassSignalSubsystem.h"
 #include "MassEntitySubsystem.h"
diff --git a/Engine/Plugins/Runtime/MassGameplay/Source/MassEQS/Private/MassEQSTypes.cpp b/Engine/Plugins/Runtime/MassGameplay/Source/MassEQS/Private/MassEQSTypes.cpp
--- a/Engine/Plugins/Runtime/MassGameplay/Source/MassEQS/Private/MassEQSTypes.cpp
+++ b/Engine/Plugins/Runtime/MassGameplay/Source/MassEQS/Private/MassEQSTypes.cpp
@@ -1,6 +1,7 @@
 // Copyright Epic Games, Inc. All Rights Reserved.
 
 #include "MassEQSTypes.h"
+#include "Engine/World.h"
 #include "EnvironmentQuery/EnvQueryTypes.h"
 #include "MassEQSSubsystem.h"
 #include "MassEQS.h"
diff --git a/Engine/Plugins/Runtime/MassGameplay/Source/MassGameplayEditor/Private/MassActorEditorSubsystem.cpp b/Engine/Plugins/Runtime/MassGameplay/Source/MassGameplayEditor/Private/MassActorEditorSubsystem.cpp
--- a/Engine/Plugins/Runtime/MassGameplay/Source/MassGameplayEditor/Private/MassActorEditorSubsystem.cpp
+++ b/Engine/Plugins/Runtime/MassGameplay/Source/MassGameplayEditor/Private/MassActorEditorSubsystem.cpp
@@ -1,6 +1,7 @@
 // Copyright Epic Games, Inc. All Rights Reserved.
 
 #include "MassActorEditorSubsystem.h"
+#include "Subsystems/SubsystemCollection.h"
 #include "MassActorSubsystem.h"
 #include "MassEntityManager.h"
 #include "MassEntityEditorSubsystem.h"
