[MassGameplay] More non-unit header fixed

#rnx

diff --git a/Engine/Plugins/Runtime/MassGameplay/Source/MassEQS/Private/Processors/MassEnvQueryProcessorBase.cpp b/Engine/Plugins/Runtime/MassGameplay/Source/MassEQS/Private/Processors/MassEnvQueryProcessorBase.cpp
--- a/Engine/Plugins/Runtime/MassGameplay/Source/MassEQS/Private/Processors/MassEnvQueryProcessorBase.cpp
+++ b/Engine/Plugins/Runtime/MassGameplay/Source/MassEQS/Private/Processors/MassEnvQueryProcessorBase.cpp
@@ -1,7 +1,7 @@
 // Copyright Epic Games, Inc. All Rights Reserved.
 
 #include "Processors/MassEnvQueryProcessorBase.h"
-
+#include "Engine/World.h"
 #include "MassEQSSubsystem.h"
 
 void UMassEnvQueryProcessorBase::Initialize(UObject& Owner)
