[Backout] - CL37013364 - CIS Issue 623586
Original CL Desc
-----------------------------------------------------------------
[Mass] headers cleanup

diff --git a/Engine/Plugins/AI/MassAI/Source/MassAIBehavior/Public/MassLookAtTargetTrait.h b/Engine/Plugins/AI/MassAI/Source/MassAIBehavior/Public/MassLookAtTargetTrait.h
--- a/Engine/Plugins/AI/MassAI/Source/MassAIBehavior/Public/MassLookAtTargetTrait.h
+++ b/Engine/Plugins/AI/MassAI/Source/MassAIBehavior/Public/MassLookAtTargetTrait.h
@@ -2,6 +2,7 @@
 
 #pragma once
 
+#include "CoreMinimal.h"
 #include "MassEntityTraitBase.h"
 #include "MassLookAtTargetTrait.generated.h"
 
@@ -14,7 +15,3 @@ protected:
 
 	virtual void BuildTemplate(FMassEntityTemplateBuildContext& BuildContext, const UWorld& World) const override;
 };
-
-#if UE_ENABLE_INCLUDE_ORDER_DEPRECATED_IN_5_6
-#include "CoreMinimal.h"
-#endif // UE_ENABLE_INCLUDE_ORDER_DEPRECATED_IN_5_6
\ No newline at end of file
diff --git a/Engine/Plugins/AI/MassAI/Source/MassAIBehavior/Public/MassLookAtTrait.h b/Engine/Plugins/AI/MassAI/Source/MassAIBehavior/Public/MassLookAtTrait.h
--- a/Engine/Plugins/AI/MassAI/Source/MassAIBehavior/Public/MassLookAtTrait.h
+++ b/Engine/Plugins/AI/MassAI/Source/MassAIBehavior/Public/MassLookAtTrait.h
@@ -2,6 +2,7 @@
 
 #pragma once
 
+#include "CoreMinimal.h"
 #include "MassEntityTraitBase.h"
 #include "MassLookAtTrait.generated.h"
 
@@ -14,7 +15,3 @@ protected:
 
 	virtual void BuildTemplate(FMassEntityTemplateBuildContext& BuildContext, const UWorld& World) const override;
 };
-
-#if UE_ENABLE_INCLUDE_ORDER_DEPRECATED_IN_5_6
-#include "CoreMinimal.h"
-#endif // UE_ENABLE_INCLUDE_ORDER_DEPRECATED_IN_5_6
\ No newline at end of file
diff --git a/Engine/Plugins/AI/MassAI/Source/MassAIBehavior/Public/MassStateTreeTypes.h b/Engine/Plugins/AI/MassAI/Source/MassAIBehavior/Public/MassStateTreeTypes.h
--- a/Engine/Plugins/AI/MassAI/Source/MassAIBehavior/Public/MassStateTreeTypes.h
+++ b/Engine/Plugins/AI/MassAI/Source/MassAIBehavior/Public/MassStateTreeTypes.h
@@ -2,6 +2,7 @@
 
 #pragma once
 
+#include "CoreMinimal.h"
 #include "StateTreeEvaluatorBase.h" 
 #include "StateTreeTaskBase.h"
 #include "MassStateTreeTypes.generated.h"
@@ -64,6 +65,3 @@ protected:
 	int32 Generation = 0;
 };
 
-#if UE_ENABLE_INCLUDE_ORDER_DEPRECATED_IN_5_6
-#include "CoreMinimal.h"
-#endif // UE_ENABLE_INCLUDE_ORDER_DEPRECATED_IN_5_6
\ No newline at end of file
diff --git a/Engine/Plugins/AI/MassAI/Source/MassAIDebug/Private/GameplayDebuggerCategory_Mass.cpp b/Engine/Plugins/AI/MassAI/Source/MassAIDebug/Private/GameplayDebuggerCategory_Mass.cpp
--- a/Engine/Plugins/AI/MassAI/Source/MassAIDebug/Private/GameplayDebuggerCategory_Mass.cpp
+++ b/Engine/Plugins/AI/MassAI/Source/MassAIDebug/Private/GameplayDebuggerCategory_Mass.cpp
@@ -23,12 +23,13 @@
 #include "MassZoneGraphNavigationFragments.h"
 #include "MassSmartObjectFragments.h"
 #include "SmartObjectSubsystem.h"
+#include "Util/ColorConstants.h"
 #include "MassSimulationLOD.h"
 #include "CanvasItem.h"
 #include "Engine/World.h"
 #include "MassDebugger.h"
 #include "MassLODSubsystem.h"
-#include "MassMovementFragments.h"
+
 
 namespace UE::Mass::Debug
 {
diff --git a/Engine/Plugins/AI/MassAI/Source/MassAIDebug/Public/GameplayDebuggerCategory_Mass.h b/Engine/Plugins/AI/MassAI/Source/MassAIDebug/Public/GameplayDebuggerCategory_Mass.h
--- a/Engine/Plugins/AI/MassAI/Source/MassAIDebug/Public/GameplayDebuggerCategory_Mass.h
+++ b/Engine/Plugins/AI/MassAI/Source/MassAIDebug/Public/GameplayDebuggerCategory_Mass.h
@@ -7,9 +7,7 @@
 
 #if WITH_GAMEPLAY_DEBUGGER && WITH_MASSGAMEPLAY_DEBUG
 
-#include "MassEntityTypes.h"
 #include "HAL/IConsoleManager.h"
-#include "UObject/WeakObjectPtrTemplates.h"
 #include "GameplayDebuggerCategory.h"
 
 struct FMassEntityManager;
diff --git a/Engine/Plugins/AI/MassAI/Source/MassAITestSuite/Public/MassAITestSuiteModule.h b/Engine/Plugins/AI/MassAI/Source/MassAITestSuite/Public/MassAITestSuiteModule.h
--- a/Engine/Plugins/AI/MassAI/Source/MassAITestSuite/Public/MassAITestSuiteModule.h
+++ b/Engine/Plugins/AI/MassAI/Source/MassAITestSuite/Public/MassAITestSuiteModule.h
@@ -2,6 +2,7 @@
 
 #pragma  once
 
+#include "CoreMinimal.h"
 #include "Modules/ModuleManager.h"
 
 /**
@@ -34,6 +35,3 @@ public:
 	}
 };
 
-#if UE_ENABLE_INCLUDE_ORDER_DEPRECATED_IN_5_6
-#include "CoreMinimal.h"
-#endif // UE_ENABLE_INCLUDE_ORDER_DEPRECATED_IN_5_6
\ No newline at end of file
diff --git a/Engine/Plugins/AI/MassAI/Source/MassNavigation/Public/MassNavigationFragments.h b/Engine/Plugins/AI/MassAI/Source/MassNavigation/Public/MassNavigationFragments.h
--- a/Engine/Plugins/AI/MassAI/Source/MassNavigation/Public/MassNavigationFragments.h
+++ b/Engine/Plugins/AI/MassAI/Source/MassNavigation/Public/MassNavigationFragments.h
@@ -2,7 +2,12 @@
 
 #pragma once
 
+#if UE_ENABLE_INCLUDE_ORDER_DEPRECATED_IN_5_5
+#include "CoreMinimal.h"
+#endif
+
 #include "AI/Navigation/NavigationElement.h"
+#include "MassCommonFragments.h"
 #include "MassCommonTypes.h"
 #include "MassNavigationSubsystem.h"
 #include "MassNavigationTypes.h"
@@ -216,7 +221,3 @@ struct FNavigationRelevantFragment : public FMassFragment
 	FNavigationElementHandle Handle;
 };
 
-#if UE_ENABLE_INCLUDE_ORDER_DEPRECATED_IN_5_5
-#include "CoreMinimal.h"
-#include "MassCommonFragments.h"
-#endif
diff --git a/Engine/Plugins/AI/MassAI/Source/MassNavigation/Public/MassNavigationSubsystem.h b/Engine/Plugins/AI/MassAI/Source/MassNavigation/Public/MassNavigationSubsystem.h
--- a/Engine/Plugins/AI/MassAI/Source/MassNavigation/Public/MassNavigationSubsystem.h
+++ b/Engine/Plugins/AI/MassAI/Source/MassNavigation/Public/MassNavigationSubsystem.h
@@ -5,7 +5,6 @@
 #include "MassEntityTypes.h"
 #include "HierarchicalHashGrid2D.h"
 #include "MassSubsystemBase.h"
-#include "MassExternalSubsystemTraits.h"
 #include "MassNavigationSubsystem.generated.h"
 
 UENUM()
diff --git a/Engine/Plugins/AI/MassAI/Source/MassNavigation/Public/Steering/MassSteeringFragments.h b/Engine/Plugins/AI/MassAI/Source/MassNavigation/Public/Steering/MassSteeringFragments.h
--- a/Engine/Plugins/AI/MassAI/Source/MassNavigation/Public/Steering/MassSteeringFragments.h
+++ b/Engine/Plugins/AI/MassAI/Source/MassNavigation/Public/Steering/MassSteeringFragments.h
@@ -1,8 +1,9 @@
-// Copyright Epic Games, Inc. All Rights Reserved.
+﻿// Copyright Epic Games, Inc. All Rights Reserved.
 
 #pragma once
 
-#include "MassEntityTypes.h"
+#include "CoreMinimal.h"
+#include "MassMovementFragments.h"
 #include "MassSteeringFragments.generated.h"
 
 class UWorld;
@@ -94,8 +95,3 @@ struct FMassStandingSteeringParameters : public FMassConstSharedFragment
 	UPROPERTY(EditAnywhere, Category = "Standing", meta = (ForceUnits="cm"))
 	float DeadZoneRadius = 15.0f;
 };
-
-#if UE_ENABLE_INCLUDE_ORDER_DEPRECATED_IN_5_6
-#include "CoreMinimal.h"
-#include "MassMovementFragments.h"
-#endif // UE_ENABLE_INCLUDE_ORDER_DEPRECATED_IN_5_6
\ No newline at end of file
diff --git a/Engine/Plugins/AI/MassAI/Source/MassNavigationEditor/Public/MassNavigationEditorModule.h b/Engine/Plugins/AI/MassAI/Source/MassNavigationEditor/Public/MassNavigationEditorModule.h
--- a/Engine/Plugins/AI/MassAI/Source/MassNavigationEditor/Public/MassNavigationEditorModule.h
+++ b/Engine/Plugins/AI/MassAI/Source/MassNavigationEditor/Public/MassNavigationEditorModule.h
@@ -2,7 +2,9 @@
 
 #pragma  once
 
+#include "CoreMinimal.h"
 #include "Modules/ModuleManager.h"
+#include "AssetTypeCategories.h"
 #include "Toolkits/IToolkitHost.h"
 #include "Toolkits/AssetEditorToolkit.h"
 
@@ -16,9 +18,6 @@ public:
 	virtual void StartupModule() override;
 	virtual void ShutdownModule() override;
 	// End IModuleInterface
-};
 
-#if UE_ENABLE_INCLUDE_ORDER_DEPRECATED_IN_5_6
-#include "CoreMinimal.h"
-#include "AssetTypeCategories.h"
-#endif // UE_ENABLE_INCLUDE_ORDER_DEPRECATED_IN_5_6
\ No newline at end of file
+private:
+};
diff --git a/Engine/Plugins/AI/MassCrowd/Source/MassCrowd/Public/MassCrowdSubsystem.h b/Engine/Plugins/AI/MassCrowd/Source/MassCrowd/Public/MassCrowdSubsystem.h
--- a/Engine/Plugins/AI/MassCrowd/Source/MassCrowd/Public/MassCrowdSubsystem.h
+++ b/Engine/Plugins/AI/MassCrowd/Source/MassCrowd/Public/MassCrowdSubsystem.h
@@ -5,7 +5,6 @@
 #include "MassCrowdTypes.h"
 #include "ZoneGraphTypes.h"
 #include "MassSubsystemBase.h"
-#include "MassExternalSubsystemTraits.h"
 #include "MassCrowdSubsystem.generated.h"
 
 class UZoneGraphAnnotationSubsystem;
diff --git a/Engine/Plugins/Runtime/MassGameplay/Source/MassActors/Public/MassActorSpawnerSubsystem.h b/Engine/Plugins/Runtime/MassGameplay/Source/MassActors/Public/MassActorSpawnerSubsystem.h
--- a/Engine/Plugins/Runtime/MassGameplay/Source/MassActors/Public/MassActorSpawnerSubsystem.h
+++ b/Engine/Plugins/Runtime/MassGameplay/Source/MassActors/Public/MassActorSpawnerSubsystem.h
@@ -5,10 +5,11 @@
 #include "StructUtils/InstancedStruct.h"
 #include "IndexedHandle.h"
 #include "Delegates/Delegate.h"
+#include "MassCommonTypes.h"
 #include "Templates/SubclassOf.h"
 #include "GameFramework/Actor.h"
 #include "MassSubsystemBase.h"
-#include "MassEntityTypes.h"
+
 #include "MassActorSpawnerSubsystem.generated.h"
 
 
@@ -290,7 +291,3 @@ public:
 	UE_DEPRECATED(5.4, "This flavor of GetNextRequestToSpawn is deprecated. Use the alternative taking an int32& parameter")
 	virtual FMassActorSpawnRequestHandle GetNextRequestToSpawn() const final;
 };
-
-#if UE_ENABLE_INCLUDE_ORDER_DEPRECATED_IN_5_6
-#include "MassCommonTypes.h"
-#endif // UE_ENABLE_INCLUDE_ORDER_DEPRECATED_IN_5_6
\ No newline at end of file
diff --git a/Engine/Plugins/Runtime/MassGameplay/Source/MassActors/Public/MassActorSubsystem.h b/Engine/Plugins/Runtime/MassGameplay/Source/MassActors/Public/MassActorSubsystem.h
--- a/Engine/Plugins/Runtime/MassGameplay/Source/MassActors/Public/MassActorSubsystem.h
+++ b/Engine/Plugins/Runtime/MassGameplay/Source/MassActors/Public/MassActorSubsystem.h
@@ -8,7 +8,6 @@
 #include "Misc/MTAccessDetector.h"
 #include "UObject/ObjectKey.h"
 #include "MassSubsystemBase.h"
-#include "MassExternalSubsystemTraits.h"
 #include "MassActorSubsystem.generated.h"
 
 struct FMassEntityHandle;
diff --git a/Engine/Plugins/Runtime/MassGameplay/Source/MassActors/Public/MassAgentTraits.h b/Engine/Plugins/Runtime/MassGameplay/Source/MassActors/Public/MassAgentTraits.h
--- a/Engine/Plugins/Runtime/MassGameplay/Source/MassActors/Public/MassAgentTraits.h
+++ b/Engine/Plugins/Runtime/MassGameplay/Source/MassActors/Public/MassAgentTraits.h
@@ -2,6 +2,7 @@
 
 #pragma once
 
+#include "CoreMinimal.h"
 #include "MassEntityTraitBase.h"
 #include "MassTranslator.h"
 #include "MassAgentTraits.generated.h"
@@ -63,6 +64,3 @@ protected:
 	virtual void BuildTemplate(FMassEntityTemplateBuildContext& BuildContext, const UWorld& World) const override;
 };
 
-#if UE_ENABLE_INCLUDE_ORDER_DEPRECATED_IN_5_6
-#include "CoreMinimal.h"
-#endif // UE_ENABLE_INCLUDE_ORDER_DEPRECATED_IN_5_6
\ No newline at end of file
diff --git a/Engine/Plugins/Runtime/MassGameplay/Source/MassCommon/Public/MassCommonTypes.h b/Engine/Plugins/Runtime/MassGameplay/Source/MassCommon/Public/MassCommonTypes.h
--- a/Engine/Plugins/Runtime/MassGameplay/Source/MassCommon/Public/MassCommonTypes.h
+++ b/Engine/Plugins/Runtime/MassGameplay/Source/MassCommon/Public/MassCommonTypes.h
@@ -2,11 +2,9 @@
 
 #pragma once
 
-#if UE_ENABLE_INCLUDE_ORDER_DEPRECATED_IN_5_6
 #include "MassEntityTypes.h"
-#include "StructUtils/InstancedStruct.h"
-#endif // UE_ENABLE_INCLUDE_ORDER_DEPRECATED_IN_5_6
 #include "MassProcessingTypes.h"
+#include "StructUtils/InstancedStruct.h"
 #include "SequentialID.h"
 #include "MassCommonTypes.generated.h"
 
diff --git a/Engine/Plugins/Runtime/MassGameplay/Source/MassEQS/Public/MassEQSSubsystem.h b/Engine/Plugins/Runtime/MassGameplay/Source/MassEQS/Public/MassEQSSubsystem.h
--- a/Engine/Plugins/Runtime/MassGameplay/Source/MassEQS/Public/MassEQSSubsystem.h
+++ b/Engine/Plugins/Runtime/MassGameplay/Source/MassEQS/Public/MassEQSSubsystem.h
@@ -3,10 +3,12 @@
 #pragma once
 
 #include "Containers/Queue.h"
+
 #include "EnvironmentQuery/EnvQueryNode.h"
 #include "MassEQSTypes.h"
+
 #include "MassSubsystemBase.h"
-#include "MassExternalSubsystemTraits.h"
+
 #include "MassEQSSubsystem.generated.h"
 
 /**
diff --git a/Engine/Plugins/Runtime/MassGameplay/Source/MassEQS/Public/MassEQSTypes.h b/Engine/Plugins/Runtime/MassGameplay/Source/MassEQS/Public/MassEQSTypes.h
--- a/Engine/Plugins/Runtime/MassGameplay/Source/MassEQS/Public/MassEQSTypes.h
+++ b/Engine/Plugins/Runtime/MassGameplay/Source/MassEQS/Public/MassEQSTypes.h
@@ -4,7 +4,7 @@
 
 #include "IndexedHandle.h"
 #include "MassEntityTypes.h"
-#include "UObject/Interface.h"
+
 #include "MassEQSTypes.generated.h"
 
 struct FEnvQueryInstance;
diff --git a/Engine/Plugins/Runtime/MassGameplay/Source/MassGameplayDebug/Public/MassDebuggerSubsystem.h b/Engine/Plugins/Runtime/MassGameplay/Source/MassGameplayDebug/Public/MassDebuggerSubsystem.h
--- a/Engine/Plugins/Runtime/MassGameplay/Source/MassGameplayDebug/Public/MassDebuggerSubsystem.h
+++ b/Engine/Plugins/Runtime/MassGameplay/Source/MassGameplayDebug/Public/MassDebuggerSubsystem.h
@@ -4,7 +4,6 @@
 
 #include "MassGameplayDebugTypes.h"
 #include "MassSubsystemBase.h"
-#include "MassExternalSubsystemTraits.h"
 #include "MassDebuggerSubsystem.generated.h"
 
 
diff --git a/Engine/Plugins/Runtime/MassGameplay/Source/MassGameplayDebug/Public/MassGameplayDebugTypes.h b/Engine/Plugins/Runtime/MassGameplay/Source/MassGameplayDebug/Public/MassGameplayDebugTypes.h
--- a/Engine/Plugins/Runtime/MassGameplay/Source/MassGameplayDebug/Public/MassGameplayDebugTypes.h
+++ b/Engine/Plugins/Runtime/MassGameplay/Source/MassGameplayDebug/Public/MassGameplayDebugTypes.h
@@ -2,6 +2,7 @@
 
 #pragma once
 
+#include "CoreMinimal.h"
 #include "UObject/ObjectMacros.h"
 #include "Logging/LogMacros.h"
 #include "Engine/DataTable.h"
@@ -100,7 +101,3 @@ struct FAgentDebugVisualization : public FTableRowBase
 	UPROPERTY(EditAnywhere, Category = "Mass|Debug")
 	EMassEntityDebugShape WireShape = EMassEntityDebugShape::Box;
 };
-
-#if UE_ENABLE_INCLUDE_ORDER_DEPRECATED_IN_5_6
-#include "CoreMinimal.h"
-#endif // UE_ENABLE_INCLUDE_ORDER_DEPRECATED_IN_5_6
\ No newline at end of file
diff --git a/Engine/Plugins/Runtime/MassGameplay/Source/MassGameplayEditor/Public/MassGameplayEditorModule.h b/Engine/Plugins/Runtime/MassGameplay/Source/MassGameplayEditor/Public/MassGameplayEditorModule.h
--- a/Engine/Plugins/Runtime/MassGameplay/Source/MassGameplayEditor/Public/MassGameplayEditorModule.h
+++ b/Engine/Plugins/Runtime/MassGameplay/Source/MassGameplayEditor/Public/MassGameplayEditorModule.h
@@ -16,3 +16,8 @@ public:
 protected:
 	void RegisterSectionMappings();
 };
+
+#if UE_ENABLE_INCLUDE_ORDER_DEPRECATED_IN_5_3
+#include "CoreMinimal.h"
+#include "Toolkits/AssetEditorToolkit.h"
+#endif
diff --git a/Engine/Plugins/Runtime/MassGameplay/Source/MassGameplayTestSuite/Public/MassGameplayTestSuiteModule.h b/Engine/Plugins/Runtime/MassGameplay/Source/MassGameplayTestSuite/Public/MassGameplayTestSuiteModule.h
--- a/Engine/Plugins/Runtime/MassGameplay/Source/MassGameplayTestSuite/Public/MassGameplayTestSuiteModule.h
+++ b/Engine/Plugins/Runtime/MassGameplay/Source/MassGameplayTestSuite/Public/MassGameplayTestSuiteModule.h
@@ -2,6 +2,7 @@
 
 #pragma once
 
+#include "CoreMinimal.h"
 #include "Modules/ModuleManager.h"
 
 /**
@@ -34,6 +35,3 @@ public:
 	}
 };
 
-#if UE_ENABLE_INCLUDE_ORDER_DEPRECATED_IN_5_6
-#include "CoreMinimal.h"
-#endif // UE_ENABLE_INCLUDE_ORDER_DEPRECATED_IN_5_6
\ No newline at end of file
diff --git a/Engine/Plugins/Runtime/MassGameplay/Source/MassMovement/Private/MassMovementDelegates.cpp b/Engine/Plugins/Runtime/MassGameplay/Source/MassMovement/Private/MassMovementDelegates.cpp
--- a/Engine/Plugins/Runtime/MassGameplay/Source/MassMovement/Private/MassMovementDelegates.cpp
+++ b/Engine/Plugins/Runtime/MassGameplay/Source/MassMovement/Private/MassMovementDelegates.cpp
@@ -1,6 +1,7 @@
 // Copyright Epic Games, Inc. All Rights Reserved.
 
 #include "MassMovementDelegates.h"
+#include "CoreMinimal.h"
 
 namespace UE::MassMovement::Delegates
 {
diff --git a/Engine/Plugins/Runtime/MassGameplay/Source/MassMovement/Public/MassMovementDelegates.h b/Engine/Plugins/Runtime/MassGameplay/Source/MassMovement/Public/MassMovementDelegates.h
--- a/Engine/Plugins/Runtime/MassGameplay/Source/MassMovement/Public/MassMovementDelegates.h
+++ b/Engine/Plugins/Runtime/MassGameplay/Source/MassMovement/Public/MassMovementDelegates.h
@@ -2,8 +2,7 @@
 
 #pragma once
 
-#include "Delegates/Delegate.h"
-#include "Delegates/DelegateCombinations.h"
+#include "CoreMinimal.h"
 
 namespace UE::MassMovement::Delegates
 {
@@ -13,7 +12,3 @@ namespace UE::MassMovement::Delegates
 	extern MASSMOVEMENT_API FOnMassMovementNamesChanged OnMassMovementNamesChanged;
 #endif // WITH_EDITOR
 } // UE::MassMovement::Delegates
-
-#if UE_ENABLE_INCLUDE_ORDER_DEPRECATED_IN_5_6
-#include "CoreMinimal.h"
-#endif // UE_ENABLE_INCLUDE_ORDER_DEPRECATED_IN_5_6
\ No newline at end of file
diff --git a/Engine/Plugins/Runtime/MassGameplay/Source/MassMovement/Public/MassMovementFragments.h b/Engine/Plugins/Runtime/MassGameplay/Source/MassMovement/Public/MassMovementFragments.h
--- a/Engine/Plugins/Runtime/MassGameplay/Source/MassMovement/Public/MassMovementFragments.h
+++ b/Engine/Plugins/Runtime/MassGameplay/Source/MassMovement/Public/MassMovementFragments.h
@@ -2,7 +2,8 @@
 
 #pragma once
 
-#include "MassEntityTypes.h"
+#include "CoreMinimal.h"
+#include "MassCommonTypes.h"
 #include "MassMovementTypes.h"
 #include "RandomSequence.h"
 #include "MassMovementFragments.generated.h"
@@ -90,8 +91,3 @@ struct MASSMOVEMENT_API FMassMovementParameters : public FMassConstSharedFragmen
 	UPROPERTY(EditAnywhere, Category = "Movement")
 	TArray<FMassMovementStyleParameters> MovementStyles;
 };
-
-#if UE_ENABLE_INCLUDE_ORDER_DEPRECATED_IN_5_6
-#include "CoreMinimal.h"
-#include "MassCommonTypes.h"
-#endif // UE_ENABLE_INCLUDE_ORDER_DEPRECATED_IN_5_6
diff --git a/Engine/Plugins/Runtime/MassGameplay/Source/MassMovement/Public/MassMovementSettings.h b/Engine/Plugins/Runtime/MassGameplay/Source/MassMovement/Public/MassMovementSettings.h
--- a/Engine/Plugins/Runtime/MassGameplay/Source/MassMovement/Public/MassMovementSettings.h
+++ b/Engine/Plugins/Runtime/MassGameplay/Source/MassMovement/Public/MassMovementSettings.h
@@ -2,6 +2,7 @@
 
 #pragma once
 
+#include "CoreMinimal.h"
 #include "MassSettings.h"
 #include "MassMovementTypes.h"
 #include "MassMovementSettings.generated.h"
@@ -31,7 +32,3 @@ private:
 	UPROPERTY(EditAnywhere, config, Category = Movement);
     TArray<FMassMovementStyle> MovementStyles;
 };
-
-#if UE_ENABLE_INCLUDE_ORDER_DEPRECATED_IN_5_6
-#include "CoreMinimal.h"
-#endif // UE_ENABLE_INCLUDE_ORDER_DEPRECATED_IN_5_6
\ No newline at end of file
diff --git a/Engine/Plugins/Runtime/MassGameplay/Source/MassMovementEditor/Public/MassMovementEditorModule.h b/Engine/Plugins/Runtime/MassGameplay/Source/MassMovementEditor/Public/MassMovementEditorModule.h
--- a/Engine/Plugins/Runtime/MassGameplay/Source/MassMovementEditor/Public/MassMovementEditorModule.h
+++ b/Engine/Plugins/Runtime/MassGameplay/Source/MassMovementEditor/Public/MassMovementEditorModule.h
@@ -2,7 +2,9 @@
 
 #pragma  once
 
+#include "CoreMinimal.h"
 #include "Modules/ModuleManager.h"
+#include "AssetTypeCategories.h"
 #include "Toolkits/IToolkitHost.h"
 #include "Toolkits/AssetEditorToolkit.h"
 
@@ -18,9 +20,6 @@ public:
 	// Begin IModuleInterface
 	virtual void StartupModule() override;
 	virtual void ShutdownModule() override;
-};
 
-#if UE_ENABLE_INCLUDE_ORDER_DEPRECATED_IN_5_6
-#include "CoreMinimal.h"
-#include "AssetTypeCategories.h"
-#endif // UE_ENABLE_INCLUDE_ORDER_DEPRECATED_IN_5_6
+private:
+};
diff --git a/Engine/Plugins/Runtime/MassGameplay/Source/MassMovementEditor/Public/MassMovementPropertyUtils.h b/Engine/Plugins/Runtime/MassGameplay/Source/MassMovementEditor/Public/MassMovementPropertyUtils.h
--- a/Engine/Plugins/Runtime/MassGameplay/Source/MassMovementEditor/Public/MassMovementPropertyUtils.h
+++ b/Engine/Plugins/Runtime/MassGameplay/Source/MassMovementEditor/Public/MassMovementPropertyUtils.h
@@ -1,6 +1,6 @@
 ﻿// Copyright Epic Games, Inc. All Rights Reserved.
 #pragma once
-
+#include "CoreMinimal.h"
 #include "Editor.h"
 #include "PropertyHandle.h"
 
@@ -110,8 +110,4 @@ void SetValue(TSharedPtr<IPropertyHandle>& ValueProperty, T NewValue, EPropertyV
 
 } // UE::MassMovement::PropertyUtils
 
-#undef LOCTEXT_NAMESPACE
-
-#if UE_ENABLE_INCLUDE_ORDER_DEPRECATED_IN_5_6
-#include "CoreMinimal.h"
-#endif // UE_ENABLE_INCLUDE_ORDER_DEPRECATED_IN_5_6
\ No newline at end of file
+#undef LOCTEXT_NAMESPACE
\ No newline at end of file
diff --git a/Engine/Plugins/Runtime/MassGameplay/Source/MassMovementEditor/Public/MassMovementStyleRefDetails.h b/Engine/Plugins/Runtime/MassGameplay/Source/MassMovementEditor/Public/MassMovementStyleRefDetails.h
--- a/Engine/Plugins/Runtime/MassGameplay/Source/MassMovementEditor/Public/MassMovementStyleRefDetails.h
+++ b/Engine/Plugins/Runtime/MassGameplay/Source/MassMovementEditor/Public/MassMovementStyleRefDetails.h
@@ -1,6 +1,6 @@
 // Copyright Epic Games, Inc. All Rights Reserved.
 #pragma once
-
+#include "CoreMinimal.h"
 #include "UnrealClient.h"
 #include "IPropertyTypeCustomization.h"
 
@@ -32,8 +32,4 @@ private:
 
 	class IPropertyUtilities* PropUtils = nullptr;
 	TSharedPtr<IPropertyHandle> StructProperty;
-};
-
-#if UE_ENABLE_INCLUDE_ORDER_DEPRECATED_IN_5_6
-#include "CoreMinimal.h"
-#endif // UE_ENABLE_INCLUDE_ORDER_DEPRECATED_IN_5_6
+};
\ No newline at end of file
diff --git a/Engine/Plugins/Runtime/MassGameplay/Source/MassReplication/Public/MassReplicationSubsystem.h b/Engine/Plugins/Runtime/MassGameplay/Source/MassReplication/Public/MassReplicationSubsystem.h
--- a/Engine/Plugins/Runtime/MassGameplay/Source/MassReplication/Public/MassReplicationSubsystem.h
+++ b/Engine/Plugins/Runtime/MassGameplay/Source/MassReplication/Public/MassReplicationSubsystem.h
@@ -6,7 +6,7 @@
 #include "MassEntityTypes.h"
 #include "MassReplicationTypes.h"
 #include "MassSubsystemBase.h"
-#include "MassExternalSubsystemTraits.h"
+
 #include "MassReplicationSubsystem.generated.h"
 
 struct FMassEntityManager;
diff --git a/Engine/Plugins/Runtime/MassGameplay/Source/MassRepresentation/Public/MassRepresentationSubsystem.h b/Engine/Plugins/Runtime/MassGameplay/Source/MassRepresentation/Public/MassRepresentationSubsystem.h
--- a/Engine/Plugins/Runtime/MassGameplay/Source/MassRepresentation/Public/MassRepresentationSubsystem.h
+++ b/Engine/Plugins/Runtime/MassGameplay/Source/MassRepresentation/Public/MassRepresentationSubsystem.h
@@ -3,11 +3,11 @@
 #pragma once
 
 #include "GameFramework/Actor.h"
+#include "MassCommonTypes.h"
 #include "Misc/MTAccessDetector.h"
 #include "MassRepresentationTypes.h"
 #include "MassActorSpawnerSubsystem.h"
 #include "MassSubsystemBase.h"
-#include "MassExternalSubsystemTraits.h"
 #include "MassRepresentationSubsystem.generated.h"
 
 class UMassVisualizationComponent;
@@ -233,7 +233,3 @@ struct TMassExternalSubsystemTraits<UMassRepresentationSubsystem> final
 		GameThreadOnly = true
 	};
 };
-
-#if UE_ENABLE_INCLUDE_ORDER_DEPRECATED_IN_5_6
-#include "MassCommonTypes.h"
-#endif // UE_ENABLE_INCLUDE_ORDER_DEPRECATED_IN_5_6
\ No newline at end of file
diff --git a/Engine/Plugins/Runtime/MassGameplay/Source/MassSignals/Public/MassSignalSubsystem.h b/Engine/Plugins/Runtime/MassGameplay/Source/MassSignals/Public/MassSignalSubsystem.h
--- a/Engine/Plugins/Runtime/MassGameplay/Source/MassSignals/Public/MassSignalSubsystem.h
+++ b/Engine/Plugins/Runtime/MassGameplay/Source/MassSignals/Public/MassSignalSubsystem.h
@@ -5,7 +5,6 @@
 #include "MassEntityManager.h"
 #include "MassSubsystemBase.h"
 #include "Misc/MTAccessDetector.h"
-#include "MassExternalSubsystemTraits.h"
 #include "MassSignalSubsystem.generated.h"
 
 namespace UE::MassSignal 
diff --git a/Engine/Plugins/Runtime/MassGameplay/Source/MassSignals/Public/MassSignalTypes.h b/Engine/Plugins/Runtime/MassGameplay/Source/MassSignals/Public/MassSignalTypes.h
--- a/Engine/Plugins/Runtime/MassGameplay/Source/MassSignals/Public/MassSignalTypes.h
+++ b/Engine/Plugins/Runtime/MassGameplay/Source/MassSignals/Public/MassSignalTypes.h
@@ -2,12 +2,14 @@
 
 #pragma once
 
+#include "CoreMinimal.h"
 #include "Logging/LogMacros.h"
 #include "MassEntityManager.h"
-
+//#include "MassSignalTypes.generated.h"
 
 DECLARE_LOG_CATEGORY_EXTERN(LogMassSignals, Log, All)
 
+
 /**
  * MassSignalNameLookup stores list of Signal names for each entity. The names are stored per entity as a bitmask,
  * you can allocate new name using GetOrAddSignalName(). This limits the names to 64.
@@ -46,8 +48,4 @@ protected:
 
 	/** Map from entity id to name bitmask */
 	TMap<FMassEntityHandle, uint64> EntitySignals;
-};
-
-#if UE_ENABLE_INCLUDE_ORDER_DEPRECATED_IN_5_6
-#include "CoreMinimal.h"
-#endif // UE_ENABLE_INCLUDE_ORDER_DEPRECATED_IN_5_6
+};
\ No newline at end of file
diff --git a/Engine/Plugins/Runtime/MassGameplay/Source/MassSpawner/Public/MassAssortedFragmentsTrait.h b/Engine/Plugins/Runtime/MassGameplay/Source/MassSpawner/Public/MassAssortedFragmentsTrait.h
--- a/Engine/Plugins/Runtime/MassGameplay/Source/MassSpawner/Public/MassAssortedFragmentsTrait.h
+++ b/Engine/Plugins/Runtime/MassGameplay/Source/MassSpawner/Public/MassAssortedFragmentsTrait.h
@@ -2,6 +2,7 @@
 
 #pragma once 
 
+#include "CoreMinimal.h"
 #include "MassEntityTraitBase.h"
 #include "StructUtils/InstancedStruct.h"
 #include "MassAssortedFragmentsTrait.generated.h"
@@ -24,7 +25,3 @@ protected:
 	UPROPERTY(Category="Mass", EditAnywhere, meta = (BaseStruct = "/Script/MassEntity.MassTag", ExcludeBaseStruct))
 	TArray<FInstancedStruct> Tags;
 };
-
-#if UE_ENABLE_INCLUDE_ORDER_DEPRECATED_IN_5_6
-#include "CoreMinimal.h"
-#endif // UE_ENABLE_INCLUDE_ORDER_DEPRECATED_IN_5_6
diff --git a/Engine/Plugins/Runtime/MassGameplay/Source/MassSpawner/Public/MassEntityConfigAsset.h b/Engine/Plugins/Runtime/MassGameplay/Source/MassSpawner/Public/MassEntityConfigAsset.h
--- a/Engine/Plugins/Runtime/MassGameplay/Source/MassSpawner/Public/MassEntityConfigAsset.h
+++ b/Engine/Plugins/Runtime/MassGameplay/Source/MassSpawner/Public/MassEntityConfigAsset.h
@@ -2,8 +2,10 @@
 
 #pragma once 
 
+#include "CoreMinimal.h"
 #include "Engine/DataAsset.h"
 #include "MassEntityTemplate.h"
+#include "MassEntityTemplateRegistry.h"
 #include "MassEntityConfigAsset.generated.h"
 
 class UMassEntityConfigAsset;
@@ -151,9 +153,4 @@ protected:
 	/** The config described in this asset. */
 	UPROPERTY(Category = "Entity Config", EditAnywhere)
 	FMassEntityConfig Config;
-};
-
-#if UE_ENABLE_INCLUDE_ORDER_DEPRECATED_IN_5_6
-#include "CoreMinimal.h"
-#include "MassEntityTemplateRegistry.h"
-#endif // UE_ENABLE_INCLUDE_ORDER_DEPRECATED_IN_5_6
+};
\ No newline at end of file
diff --git a/Engine/Plugins/Runtime/MassGameplay/Source/MassSpawner/Public/MassEntityEQSSpawnPointsGenerator.h b/Engine/Plugins/Runtime/MassGameplay/Source/MassSpawner/Public/MassEntityEQSSpawnPointsGenerator.h
--- a/Engine/Plugins/Runtime/MassGameplay/Source/MassSpawner/Public/MassEntityEQSSpawnPointsGenerator.h
+++ b/Engine/Plugins/Runtime/MassGameplay/Source/MassSpawner/Public/MassEntityEQSSpawnPointsGenerator.h
@@ -2,6 +2,7 @@
 
 #pragma once 
 
+#include "CoreMinimal.h"
 #include "MassEntitySpawnDataGeneratorBase.h"
 #include "EnvironmentQuery/EnvQueryTypes.h"
 #include "MassEntityEQSSpawnPointsGenerator.generated.h"
@@ -30,7 +31,3 @@ protected:
 	UPROPERTY(Category = "Query", EditAnywhere)
 	FEQSParametrizedQueryExecutionRequest EQSRequest;
 };
-
-#if UE_ENABLE_INCLUDE_ORDER_DEPRECATED_IN_5_6
-#include "CoreMinimal.h"
-#endif // UE_ENABLE_INCLUDE_ORDER_DEPRECATED_IN_5_6
diff --git a/Engine/Plugins/Runtime/MassGameplay/Source/MassSpawner/Public/MassEntitySpawnDataGeneratorBase.h b/Engine/Plugins/Runtime/MassGameplay/Source/MassSpawner/Public/MassEntitySpawnDataGeneratorBase.h
--- a/Engine/Plugins/Runtime/MassGameplay/Source/MassSpawner/Public/MassEntitySpawnDataGeneratorBase.h
+++ b/Engine/Plugins/Runtime/MassGameplay/Source/MassSpawner/Public/MassEntitySpawnDataGeneratorBase.h
@@ -2,6 +2,7 @@
 
 #pragma once 
 
+#include "CoreMinimal.h"
 #include "StructUtils/InstancedStruct.h"
 #include "MassSpawnerTypes.h"
 #include "MassEntitySpawnDataGeneratorBase.generated.h"
@@ -75,7 +76,3 @@ protected:
 	UPROPERTY()
 	uint8 bUseCustomSeed : 1;
 };
-
-#if UE_ENABLE_INCLUDE_ORDER_DEPRECATED_IN_5_6
-#include "CoreMinimal.h"
-#endif // UE_ENABLE_INCLUDE_ORDER_DEPRECATED_IN_5_6
diff --git a/Engine/Plugins/Runtime/MassGameplay/Source/MassSpawner/Public/MassEntityTemplateRegistry.h b/Engine/Plugins/Runtime/MassGameplay/Source/MassSpawner/Public/MassEntityTemplateRegistry.h
--- a/Engine/Plugins/Runtime/MassGameplay/Source/MassSpawner/Public/MassEntityTemplateRegistry.h
+++ b/Engine/Plugins/Runtime/MassGameplay/Source/MassSpawner/Public/MassEntityTemplateRegistry.h
@@ -2,7 +2,9 @@
 
 #pragma once
 
+#include "CoreMinimal.h"
 #include "UObject/Object.h"
+#include "MassCommonTypes.h"
 #include "MassEntityTemplate.h"
 #include "MassTranslator.h"
 #include "MassEntityTemplateRegistry.generated.h"
@@ -413,8 +415,3 @@ protected:
 };
 
 #undef ENSURE_SUPPORTED_TRAIT_OPERATION
-
-#if UE_ENABLE_INCLUDE_ORDER_DEPRECATED_IN_5_6
-#include "CoreMinimal.h"
-#include "MassCommonTypes.h"
-#endif // UE_ENABLE_INCLUDE_ORDER_DEPRECATED_IN_5_6
diff --git a/Engine/Plugins/Runtime/MassGameplay/Source/MassSpawner/Public/MassEntityZoneGraphSpawnPointsGenerator.h b/Engine/Plugins/Runtime/MassGameplay/Source/MassSpawner/Public/MassEntityZoneGraphSpawnPointsGenerator.h
--- a/Engine/Plugins/Runtime/MassGameplay/Source/MassSpawner/Public/MassEntityZoneGraphSpawnPointsGenerator.h
+++ b/Engine/Plugins/Runtime/MassGameplay/Source/MassSpawner/Public/MassEntityZoneGraphSpawnPointsGenerator.h
@@ -2,8 +2,10 @@
 
 #pragma once 
 
+#include "CoreMinimal.h"
 #include "MassEntitySpawnDataGeneratorBase.h"
 #include "ZoneGraphTypes.h"
+
 #include "MassEntityZoneGraphSpawnPointsGenerator.generated.h"
 
 class AZoneGraphData;
@@ -34,7 +36,3 @@ protected:
 	UPROPERTY(EditAnywhere, Category = "ZoneGraph Generator Config")
 	float MaxGap = 300;
 };
-
-#if UE_ENABLE_INCLUDE_ORDER_DEPRECATED_IN_5_6
-#include "CoreMinimal.h"
-#endif // UE_ENABLE_INCLUDE_ORDER_DEPRECATED_IN_5_6
diff --git a/Engine/Source/Developer/MassEntityTestSuite/Public/MassEntityTestFarmPlot.h b/Engine/Source/Developer/MassEntityTestSuite/Public/MassEntityTestFarmPlot.h
--- a/Engine/Source/Developer/MassEntityTestSuite/Public/MassEntityTestFarmPlot.h
+++ b/Engine/Source/Developer/MassEntityTestSuite/Public/MassEntityTestFarmPlot.h
@@ -2,12 +2,15 @@
 
 #pragma once
 
+#include "CoreMinimal.h"
 #include "GameFramework/Actor.h"
 #include "Engine/DataTable.h"
+
 #include "MassEntityTypes.h"
 #include "MassEntityManager.h"
 #include "MassCommandBuffer.h"
 #include "MassProcessor.h"
+
 #include "MassEntityTestFarmPlot.generated.h"
 
 class UStaticMesh;
@@ -246,7 +249,3 @@ public:
 	virtual void BeginPlay() override;
 	virtual void TickActor(float DeltaTime, enum ELevelTick TickType, FActorTickFunction& ThisTickFunction) override;
 };
-
-#if UE_ENABLE_INCLUDE_ORDER_DEPRECATED_IN_5_6
-#include "CoreMinimal.h"
-#endif // UE_ENABLE_INCLUDE_ORDER_DEPRECATED_IN_5_6
diff --git a/Engine/Source/Editor/MassEntityEditor/Public/MassEntityEditorModule.h b/Engine/Source/Editor/MassEntityEditor/Public/MassEntityEditorModule.h
--- a/Engine/Source/Editor/MassEntityEditor/Public/MassEntityEditorModule.h
+++ b/Engine/Source/Editor/MassEntityEditor/Public/MassEntityEditorModule.h
@@ -46,3 +46,10 @@ protected:
 	TSharedPtr<FExtensibilityManager> MenuExtensibilityManager;
 	TSharedPtr<FExtensibilityManager> ToolBarExtensibilityManager;
 };
+
+#if UE_ENABLE_INCLUDE_ORDER_DEPRECATED_IN_5_3
+#include "CoreMinimal.h"
+#include "Modules/ModuleManager.h"
+#include "AssetTypeCategories.h"
+#include "Toolkits/IToolkitHost.h"
+#endif
diff --git a/Engine/Source/Runtime/MassEntity/Public/MassEntityTypes.h b/Engine/Source/Runtime/MassEntity/Public/MassEntityTypes.h
--- a/Engine/Source/Runtime/MassEntity/Public/MassEntityTypes.h
+++ b/Engine/Source/Runtime/MassEntity/Public/MassEntityTypes.h
@@ -6,6 +6,7 @@
 #include "MassProcessingTypes.h"
 #include "StructUtils/StructArrayView.h"
 #include "Subsystems/Subsystem.h"
+#include "MassExternalSubsystemTraits.h"
 #include "StructUtils/SharedStruct.h"
 #include "MassEntityTypes.generated.h"
 #ifdef WITH_AITESTSUITE
@@ -802,7 +803,3 @@ struct FMassArchetypeCreationParams
 	/** Name to identify the archetype while debugging*/
 	FName DebugName;
 };
-
-#if UE_ENABLE_INCLUDE_ORDER_DEPRECATED_IN_5_6
-#include "MassExternalSubsystemTraits.h"
-#endif // UE_ENABLE_INCLUDE_ORDER_DEPRECATED_IN_5_6
diff --git a/Engine/Source/Runtime/MassEntity/Public/MassProcessingPhaseManager.h b/Engine/Source/Runtime/MassEntity/Public/MassProcessingPhaseManager.h
--- a/Engine/Source/Runtime/MassEntity/Public/MassProcessingPhaseManager.h
+++ b/Engine/Source/Runtime/MassEntity/Public/MassProcessingPhaseManager.h
@@ -2,6 +2,7 @@
 
 #pragma once
 
+#include "CoreMinimal.h"
 #include "UObject/UObjectGlobals.h"
 #include "UObject/Object.h"
 #include "UObject/GCObject.h"
@@ -224,7 +225,3 @@ protected:
 	EProcessorExecutionFlags ProcessorExecutionFlags = EProcessorExecutionFlags::None;
 	bool bIsAllowedToTick = false;
 };
-
-#if UE_ENABLE_INCLUDE_ORDER_DEPRECATED_IN_5_6
-#include "CoreMinimal.h"
-#endif // UE_ENABLE_INCLUDE_ORDER_DEPRECATED_IN_5_6
diff --git a/Engine/Source/Runtime/MassEntity/Public/MassProcessor.h b/Engine/Source/Runtime/MassEntity/Public/MassProcessor.h
--- a/Engine/Source/Runtime/MassEntity/Public/MassProcessor.h
+++ b/Engine/Source/Runtime/MassEntity/Public/MassProcessor.h
@@ -2,6 +2,7 @@
 
 #pragma once
 
+#include "CoreMinimal.h"
 #include "MassEntityManager.h"
 #include "MassProcessingTypes.h"
 #include "Async/TaskGraphInterfaces.h"
@@ -423,7 +424,3 @@ inline bool UMassCompositeProcessor::IsEmpty() const
 {
 	return ChildPipeline.IsEmpty();
 }
-
-#if UE_ENABLE_INCLUDE_ORDER_DEPRECATED_IN_5_6
-#include "CoreMinimal.h"
-#endif // UE_ENABLE_INCLUDE_ORDER_DEPRECATED_IN_5_6
