[Backout] - CL37034649 - Partial backout due to CIS errors
Original CL Desc
-----------------------------------------------------------------
[Backout] - CL37013364 - CIS Issue 623586
Original CL Desc
-----------------------------------------------------------------
[Mass] headers cleanup

diff --git a/Engine/Plugins/Runtime/MassGameplay/Source/MassEQS/Public/MassEQSTypes.h b/Engine/Plugins/Runtime/MassGameplay/Source/MassEQS/Public/MassEQSTypes.h
--- a/Engine/Plugins/Runtime/MassGameplay/Source/MassEQS/Public/MassEQSTypes.h
+++ b/Engine/Plugins/Runtime/MassGameplay/Source/MassEQS/Public/MassEQSTypes.h
@@ -4,7 +4,7 @@
 
 #include "IndexedHandle.h"
 #include "MassEntityTypes.h"
-
+#include "UObject/Interface.h"
 #include "MassEQSTypes.generated.h"
 
 struct FEnvQueryInstance;
