[Mass] Made MassProcessingPhaseManager to derive from TSharedFromThis in preparation for future changes.

diff --git a/Engine/Plugins/Runtime/MassGameplay/Source/MassSimulation/Private/MassSimulationSubsystem.cpp b/Engine/Plugins/Runtime/MassGameplay/Source/MassSimulation/Private/MassSimulationSubsystem.cpp
--- a/Engine/Plugins/Runtime/MassGameplay/Source/MassSimulation/Private/MassSimulationSubsystem.cpp
+++ b/Engine/Plugins/Runtime/MassGameplay/Source/MassSimulation/Private/MassSimulationSubsystem.cpp
@@ -32,8 +32,8 @@ namespace UE::Mass::Simulation
 //----------------------------------------------------------------------//
 // UMassSimulationSubsystem
 //----------------------------------------------------------------------//
-UMassSimulationSubsystem::UMassSimulationSubsystem(const FObjectInitializer& ObjectInitializer)
-	: Super()
+UMassSimulationSubsystem::UMassSimulationSubsystem(const FObjectInitializer&)
+	: PhaseManager(new FMassProcessingPhaseManager())
 {
 }
 
@@ -68,12 +68,12 @@ void UMassSimulationSubsystem::BeginDestroy()
 
 FMassProcessingPhase::FOnPhaseEvent& UMassSimulationSubsystem::GetOnProcessingPhaseStarted(const EMassProcessingPhase Phase)
 {
-	return PhaseManager.GetOnPhaseStart(Phase);
+	return PhaseManager->GetOnPhaseStart(Phase);
 }
 
 FMassProcessingPhase::FOnPhaseEvent& UMassSimulationSubsystem::GetOnProcessingPhaseFinished(const EMassProcessingPhase Phase)
 {
-	return PhaseManager.GetOnPhaseEnd(Phase);
+	return PhaseManager->GetOnPhaseEnd(Phase);
 }
 
 bool UMassSimulationSubsystem::IsDuringMassProcessing() const 
@@ -99,7 +99,7 @@ void UMassSimulationSubsystem::Deinitialize()
 	GetOnProcessingPhaseStarted(EMassProcessingPhase::PrePhysics).RemoveAll(this);
 	StopSimulation();
 
-	PhaseManager.Deinitialize();
+	PhaseManager->Deinitialize();
 	EntityManager.Reset();
 
 	Super::Deinitialize();
@@ -144,13 +144,13 @@ void UMassSimulationSubsystem::OnWorldBeginPlay(UWorld& InWorld)
 void UMassSimulationSubsystem::RegisterDynamicProcessor(UMassProcessor& Processor)
 {
 	checkf(!IsDuringMassProcessing(), TEXT("Unable to add dynamic processors to Mass during processing."));
-	PhaseManager.RegisterDynamicProcessor(Processor);
+	PhaseManager->RegisterDynamicProcessor(Processor);
 }
 
 void UMassSimulationSubsystem::UnregisterDynamicProcessor(UMassProcessor& Processor)
 {
 	checkf(!IsDuringMassProcessing(), TEXT("Unable to remove dynamic processors to Mass during processing."));
-	PhaseManager.UnregisterDynamicProcessor(Processor);
+	PhaseManager->UnregisterDynamicProcessor(Processor);
 }
 
 void UMassSimulationSubsystem::HandleSimulationTickingEnabledCVarChange(IConsoleVariable*)
@@ -206,7 +206,7 @@ void UMassSimulationSubsystem::RebuildTickPipeline()
 	}
 #endif // WITH_EDITOR
 
-	PhaseManager.Initialize(*this, ProcessingPhasesConfig, DependencyGraphFileName);
+	PhaseManager->Initialize(*this, ProcessingPhasesConfig, DependencyGraphFileName);
 }
 
 void UMassSimulationSubsystem::StartSimulation(UWorld& InWorld)
@@ -218,7 +218,7 @@ void UMassSimulationSubsystem::StartSimulation(UWorld& InWorld)
 
 	if (UE::Mass::Simulation::bSimulationTickingEnabled)
 	{
-		PhaseManager.Start(InWorld);
+		PhaseManager->Start(InWorld);
 
 		bSimulationStarted = true;
 
@@ -233,7 +233,7 @@ void UMassSimulationSubsystem::StopSimulation()
 		return;
 	}
 
-	PhaseManager.Stop();
+	PhaseManager->Stop();
 
 	bSimulationStarted = false;
 }
diff --git a/Engine/Plugins/Runtime/MassGameplay/Source/MassSimulation/Public/MassSimulationSubsystem.h b/Engine/Plugins/Runtime/MassGameplay/Source/MassSimulation/Public/MassSimulationSubsystem.h
--- a/Engine/Plugins/Runtime/MassGameplay/Source/MassSimulation/Public/MassSimulationSubsystem.h
+++ b/Engine/Plugins/Runtime/MassGameplay/Source/MassSimulation/Public/MassSimulationSubsystem.h
@@ -23,7 +23,7 @@ public:
 	
 	MASSSIMULATION_API UMassSimulationSubsystem(const FObjectInitializer& ObjectInitializer = FObjectInitializer::Get());
 
-	const FMassProcessingPhaseManager& GetPhaseManager() const { return PhaseManager; }
+	const FMassProcessingPhaseManager& GetPhaseManager() const { return PhaseManager.Get(); }
 
 	MASSSIMULATION_API FMassProcessingPhase::FOnPhaseEvent& GetOnProcessingPhaseStarted(const EMassProcessingPhase Phase);
 	MASSSIMULATION_API FMassProcessingPhase::FOnPhaseEvent& GetOnProcessingPhaseFinished(const EMassProcessingPhase Phase);
@@ -66,7 +66,7 @@ protected:
 
 	TSharedPtr<FMassEntityManager> EntityManager;
 
-	FMassProcessingPhaseManager PhaseManager;
+	TSharedRef<FMassProcessingPhaseManager> PhaseManager;
 
 	inline static FOnSimulationStarted OnSimulationStarted={};
 
diff --git a/Engine/Source/Runtime/MassEntity/Public/MassProcessingPhaseManager.h b/Engine/Source/Runtime/MassEntity/Public/MassProcessingPhaseManager.h
--- a/Engine/Source/Runtime/MassEntity/Public/MassProcessingPhaseManager.h
+++ b/Engine/Source/Runtime/MassEntity/Public/MassProcessingPhaseManager.h
@@ -122,7 +122,7 @@ struct MASSENTITY_API FMassPhaseProcessorConfigurationHelper
  * with collections of processors (via Initialize function) as well as registering arbitrary functions to be called 
  * when a particular phase starts or ends (via GetOnPhaseStart and GetOnPhaseEnd functions). 
  */
-struct MASSENTITY_API FMassProcessingPhaseManager : public FGCObject
+struct MASSENTITY_API FMassProcessingPhaseManager : public FGCObject, public TSharedFromThis<FMassProcessingPhaseManager>
 {
 public:
 	explicit FMassProcessingPhaseManager(EProcessorExecutionFlags InProcessorExecutionFlags = EProcessorExecutionFlags::None) 
@@ -131,7 +131,7 @@ public:
 	FMassProcessingPhaseManager(const FMassProcessingPhaseManager& Other) = delete;
 	FMassProcessingPhaseManager& operator=(const FMassProcessingPhaseManager& Other) = delete;
 
-	const TSharedPtr<FMassEntityManager>& GetEntityManager() { return EntityManager; }
+	const TSharedPtr<FMassEntityManager>& GetEntityManager() const { return EntityManager; }
 	FMassEntityManager& GetEntityManagerRef() { check(EntityManager); return *EntityManager.Get(); }
 
 	/** Retrieves OnPhaseStart multicast delegate's reference for a given Phase */
