[Mass] more Mass-related header cleanups.

diff --git a/Engine/Plugins/AI/MassAI/Source/MassAIBehavior/Public/Evaluators/MassZoneGraphAnnotationEvaluator.h b/Engine/Plugins/AI/MassAI/Source/MassAIBehavior/Public/Evaluators/MassZoneGraphAnnotationEvaluator.h
--- a/Engine/Plugins/AI/MassAI/Source/MassAIBehavior/Public/Evaluators/MassZoneGraphAnnotationEvaluator.h
+++ b/Engine/Plugins/AI/MassAI/Source/MassAIBehavior/Public/Evaluators/MassZoneGraphAnnotationEvaluator.h
@@ -9,7 +9,7 @@
 struct FStateTreeExecutionContext;
 struct FMassZoneGraphAnnotationFragment;
 /**
- * Evaluator to expose ZoneGraph Annotation Tags for decision making.
+ * Evaluator to expose ZoneGraph Annotation Tags for decision-making.
  */
 USTRUCT()
 struct MASSAIBEHAVIOR_API FMassZoneGraphAnnotationEvaluatorInstanceData
diff --git a/Engine/Plugins/AI/MassAI/Source/MassAIBehavior/Public/MassComponentHitTypes.h b/Engine/Plugins/AI/MassAI/Source/MassAIBehavior/Public/MassComponentHitTypes.h
--- a/Engine/Plugins/AI/MassAI/Source/MassAIBehavior/Public/MassComponentHitTypes.h
+++ b/Engine/Plugins/AI/MassAI/Source/MassAIBehavior/Public/MassComponentHitTypes.h
@@ -4,8 +4,9 @@
 
 #include "Engine/EngineTypes.h"
 #include "MassEntityTypes.h"
+#if UE_ENABLE_INCLUDE_ORDER_DEPRECATED_IN_5_6
 #include "StateTreeTypes.h"
-
+#endif // UE_ENABLE_INCLUDE_ORDER_DEPRECATED_IN_5_6
 #include "MassComponentHitTypes.generated.h"
 
 namespace UE::Mass::Signals
diff --git a/Engine/Plugins/AI/MassAI/Source/MassAIBehavior/Public/MassLookAtProcessors.h b/Engine/Plugins/AI/MassAI/Source/MassAIBehavior/Public/MassLookAtProcessors.h
--- a/Engine/Plugins/AI/MassAI/Source/MassAIBehavior/Public/MassLookAtProcessors.h
+++ b/Engine/Plugins/AI/MassAI/Source/MassAIBehavior/Public/MassLookAtProcessors.h
@@ -3,8 +3,10 @@
 #pragma once
 
 #include "MassProcessor.h"
-#include "MassCommonTypes.h"
 #include "MassSignalProcessorBase.h"
+#if UE_ENABLE_INCLUDE_ORDER_DEPRECATED_IN_5_6
+#include "MassCommonTypes.h"
+#endif // UE_ENABLE_INCLUDE_ORDER_DEPRECATED_IN_5_6
 #include "MassLookAtProcessors.generated.h"
 
 class UMassNavigationSubsystem;
diff --git a/Engine/Plugins/AI/MassAI/Source/MassAIBehavior/Public/MassStateTreeExecutionContext.h b/Engine/Plugins/AI/MassAI/Source/MassAIBehavior/Public/MassStateTreeExecutionContext.h
--- a/Engine/Plugins/AI/MassAI/Source/MassAIBehavior/Public/MassStateTreeExecutionContext.h
+++ b/Engine/Plugins/AI/MassAI/Source/MassAIBehavior/Public/MassStateTreeExecutionContext.h
@@ -4,7 +4,6 @@
 
 #include "MassEntityTypes.h"
 #include "StateTreeExecutionContext.h"
-//#include "MassStateTreeExecutionContext.generated.h"
 
 struct FMassExecutionContext;
 struct FMassEntityManager;
diff --git a/Engine/Plugins/AI/MassAI/Source/MassAIBehavior/Public/MassStateTreeProcessors.h b/Engine/Plugins/AI/MassAI/Source/MassAIBehavior/Public/MassStateTreeProcessors.h
--- a/Engine/Plugins/AI/MassAI/Source/MassAIBehavior/Public/MassStateTreeProcessors.h
+++ b/Engine/Plugins/AI/MassAI/Source/MassAIBehavior/Public/MassStateTreeProcessors.h
@@ -5,7 +5,9 @@
 #include "MassSignalProcessorBase.h"
 #include "MassStateTreeFragments.h"
 #include "MassObserverProcessor.h"
+#if UE_ENABLE_INCLUDE_ORDER_DEPRECATED_IN_5_6
 #include "MassLODTypes.h"
+#endif // UE_ENABLE_INCLUDE_ORDER_DEPRECATED_IN_5_6
 #include "MassStateTreeProcessors.generated.h"
 
 struct FMassStateTreeExecutionContext;
diff --git a/Engine/Plugins/AI/MassAI/Source/MassAIBehavior/Public/Tasks/MassClaimSmartObjectTask.h b/Engine/Plugins/AI/MassAI/Source/MassAIBehavior/Public/Tasks/MassClaimSmartObjectTask.h
--- a/Engine/Plugins/AI/MassAI/Source/MassAIBehavior/Public/Tasks/MassClaimSmartObjectTask.h
+++ b/Engine/Plugins/AI/MassAI/Source/MassAIBehavior/Public/Tasks/MassClaimSmartObjectTask.h
@@ -2,7 +2,9 @@
 
 #pragma once
 
+#if UE_ENABLE_INCLUDE_ORDER_DEPRECATED_IN_5_6
 #include "MassSmartObjectRequest.h"
+#endif // UE_ENABLE_INCLUDE_ORDER_DEPRECATED_IN_5_6
 #include "MassStateTreeTypes.h"
 #include "SmartObjectRuntime.h"
 #include "MassClaimSmartObjectTask.generated.h"
diff --git a/Engine/Plugins/AI/MassAI/Source/MassAIBehavior/Public/Tasks/MassUseSmartObjectTask.h b/Engine/Plugins/AI/MassAI/Source/MassAIBehavior/Public/Tasks/MassUseSmartObjectTask.h
--- a/Engine/Plugins/AI/MassAI/Source/MassAIBehavior/Public/Tasks/MassUseSmartObjectTask.h
+++ b/Engine/Plugins/AI/MassAI/Source/MassAIBehavior/Public/Tasks/MassUseSmartObjectTask.h
@@ -2,7 +2,9 @@
 
 #pragma once
 
+#if UE_ENABLE_INCLUDE_ORDER_DEPRECATED_IN_5_6
 #include "MassSmartObjectRequest.h"
+#endif // UE_ENABLE_INCLUDE_ORDER_DEPRECATED_IN_5_6
 #include "MassStateTreeTypes.h"
 #include "SmartObjectRuntime.h"
 #include "MassUseSmartObjectTask.generated.h"
diff --git a/Engine/Plugins/AI/MassAI/Source/MassAIDebug/Private/MassDebugStateTreeProcessor.cpp b/Engine/Plugins/AI/MassAI/Source/MassAIDebug/Private/MassDebugStateTreeProcessor.cpp
--- a/Engine/Plugins/AI/MassAI/Source/MassAIDebug/Private/MassDebugStateTreeProcessor.cpp
+++ b/Engine/Plugins/AI/MassAI/Source/MassAIDebug/Private/MassDebugStateTreeProcessor.cpp
@@ -6,6 +6,7 @@
 #include "MassStateTreeExecutionContext.h"
 #include "MassStateTreeFragments.h"
 #include "MassCommonFragments.h"
+#include "MassCommonTypes.h"
 #include "MassDebugger.h"
 #include "MassExecutionContext.h"
 #include "VisualLogger/VisualLogger.h"
diff --git a/Engine/Plugins/AI/MassAI/Source/MassNavigation/Public/Avoidance/MassAvoidanceFragments.h b/Engine/Plugins/AI/MassAI/Source/MassNavigation/Public/Avoidance/MassAvoidanceFragments.h
--- a/Engine/Plugins/AI/MassAI/Source/MassNavigation/Public/Avoidance/MassAvoidanceFragments.h
+++ b/Engine/Plugins/AI/MassAI/Source/MassNavigation/Public/Avoidance/MassAvoidanceFragments.h
@@ -1,9 +1,11 @@
-﻿// Copyright Epic Games, Inc. All Rights Reserved.
+// Copyright Epic Games, Inc. All Rights Reserved.
 
 #pragma once
 
 #include "MassEntityTypes.h"
+#if UE_ENABLE_INCLUDE_ORDER_DEPRECATED_IN_5_6
 #include "MassNavigationSubsystem.h"
+#endif // UE_ENABLE_INCLUDE_ORDER_DEPRECATED_IN_5_6
 #include "MassAvoidanceFragments.generated.h"
 
 USTRUCT()
diff --git a/Engine/Plugins/AI/MassAI/Source/MassNavigation/Public/MassNavigationTypes.h b/Engine/Plugins/AI/MassAI/Source/MassNavigation/Public/MassNavigationTypes.h
--- a/Engine/Plugins/AI/MassAI/Source/MassNavigation/Public/MassNavigationTypes.h
+++ b/Engine/Plugins/AI/MassAI/Source/MassNavigation/Public/MassNavigationTypes.h
@@ -2,9 +2,11 @@
 
 #pragma once
 
+#if UE_ENABLE_INCLUDE_ORDER_DEPRECATED_IN_5_6
 #include "MassEntityTypes.h"
 #include "MassCommonTypes.h"
 #include "RandomSequence.h"
+#endif // UE_ENABLE_INCLUDE_ORDER_DEPRECATED_IN_5_6
 #include "MassNavigationTypes.generated.h"
 
 MASSNAVIGATION_API DECLARE_LOG_CATEGORY_EXTERN(LogMassNavigation, Warning, All);
diff --git a/Engine/Plugins/AI/MassAI/Source/MassNavigation/Public/Steering/MassSteeringProcessors.h b/Engine/Plugins/AI/MassAI/Source/MassNavigation/Public/Steering/MassSteeringProcessors.h
--- a/Engine/Plugins/AI/MassAI/Source/MassNavigation/Public/Steering/MassSteeringProcessors.h
+++ b/Engine/Plugins/AI/MassAI/Source/MassNavigation/Public/Steering/MassSteeringProcessors.h
@@ -1,9 +1,11 @@
-﻿// Copyright Epic Games, Inc. All Rights Reserved.
+// Copyright Epic Games, Inc. All Rights Reserved.
 
 #pragma once
 
 #include "MassProcessor.h"
+#if UE_ENABLE_INCLUDE_ORDER_DEPRECATED_IN_5_6
 #include "MassObserverProcessor.h"
+#endif // UE_ENABLE_INCLUDE_ORDER_DEPRECATED_IN_5_6
 #include "MassSteeringProcessors.generated.h"
 
 /** 
diff --git a/Engine/Plugins/AI/MassCrowd/Source/MassCrowd/Public/MassCrowdServerRepresentationLODProcessor.h b/Engine/Plugins/AI/MassCrowd/Source/MassCrowd/Public/MassCrowdServerRepresentationLODProcessor.h
--- a/Engine/Plugins/AI/MassCrowd/Source/MassCrowd/Public/MassCrowdServerRepresentationLODProcessor.h
+++ b/Engine/Plugins/AI/MassCrowd/Source/MassCrowd/Public/MassCrowdServerRepresentationLODProcessor.h
@@ -2,8 +2,10 @@
 
 #pragma once
 
+#if UE_ENABLE_INCLUDE_ORDER_DEPRECATED_IN_5_6
 #include "MassCrowdTypes.h"
 #include "MassCrowdFragments.h"
+#endif // UE_ENABLE_INCLUDE_ORDER_DEPRECATED_IN_5_6
 #include "MassProcessor.h"
 #include "MassLODCalculator.h"
 #include "MassLODTickRateController.h"
diff --git a/Engine/Plugins/AI/MassCrowd/Source/MassCrowd/Public/MassCrowdServerRepresentationTrait.h b/Engine/Plugins/AI/MassCrowd/Source/MassCrowd/Public/MassCrowdServerRepresentationTrait.h
--- a/Engine/Plugins/AI/MassCrowd/Source/MassCrowd/Public/MassCrowdServerRepresentationTrait.h
+++ b/Engine/Plugins/AI/MassCrowd/Source/MassCrowd/Public/MassCrowdServerRepresentationTrait.h
@@ -3,10 +3,11 @@
 #pragma once 
 
 #include "MassEntityTraitBase.h"
+#if UE_ENABLE_INCLUDE_ORDER_DEPRECATED_IN_5_6
 #include "MassRepresentationTypes.h"
+#endif // UE_ENABLE_INCLUDE_ORDER_DEPRECATED_IN_5_6
 #include "MassRepresentationFragments.h"
 #include "GameFramework/Actor.h"
-
 #include "MassCrowdServerRepresentationTrait.generated.h"
 
 
diff --git a/Engine/Plugins/AI/MassCrowd/Source/MassCrowd/Public/MassCrowdSubsystem.h b/Engine/Plugins/AI/MassCrowd/Source/MassCrowd/Public/MassCrowdSubsystem.h
--- a/Engine/Plugins/AI/MassCrowd/Source/MassCrowd/Public/MassCrowdSubsystem.h
+++ b/Engine/Plugins/AI/MassCrowd/Source/MassCrowd/Public/MassCrowdSubsystem.h
@@ -2,6 +2,7 @@
 
 #pragma once
 
+#include "MassEntityTypes.h"
 #include "MassCrowdTypes.h"
 #include "ZoneGraphTypes.h"
 #include "MassSubsystemBase.h"
diff --git a/Engine/Plugins/Runtime/InstancedActors/Source/InstancedActors/Public/InstancedActorsTypes.h b/Engine/Plugins/Runtime/InstancedActors/Source/InstancedActors/Public/InstancedActorsTypes.h
--- a/Engine/Plugins/Runtime/InstancedActors/Source/InstancedActors/Public/InstancedActorsTypes.h
+++ b/Engine/Plugins/Runtime/InstancedActors/Source/InstancedActors/Public/InstancedActorsTypes.h
@@ -5,6 +5,7 @@
 #include "Logging/LogMacros.h"
 #include "GameplayTagContainer.h"
 #include "ISMPartition/ISMComponentDescriptor.h"
+#include "MassEntityTypes.h"
 #include "MassRepresentationTypes.h"
 #include "InstancedActorsIndex.h"
 #include "InstancedActorsTypes.generated.h"
diff --git a/Engine/Plugins/Runtime/MassGameplay/Source/MassActors/Public/MassAgentSubsystem.h b/Engine/Plugins/Runtime/MassGameplay/Source/MassActors/Public/MassAgentSubsystem.h
--- a/Engine/Plugins/Runtime/MassGameplay/Source/MassActors/Public/MassAgentSubsystem.h
+++ b/Engine/Plugins/Runtime/MassGameplay/Source/MassActors/Public/MassAgentSubsystem.h
@@ -2,6 +2,7 @@
 
 #pragma once
 
+#include "MassCommonTypes.h"
 #include "MassEntityTypes.h"
 #include "MassProcessingTypes.h"
 #include "MassEntityTemplate.h"
diff --git a/Engine/Plugins/Runtime/MassGameplay/Source/MassCommon/Public/MassCommonFragments.h b/Engine/Plugins/Runtime/MassGameplay/Source/MassCommon/Public/MassCommonFragments.h
--- a/Engine/Plugins/Runtime/MassGameplay/Source/MassCommon/Public/MassCommonFragments.h
+++ b/Engine/Plugins/Runtime/MassGameplay/Source/MassCommon/Public/MassCommonFragments.h
@@ -3,9 +3,6 @@
 #pragma once
 
 #include "MassEntityTypes.h"
-#if UE_ENABLE_INCLUDE_ORDER_DEPRECATED_IN_5_2
-#include "AI/Navigation/NavigationTypes.h"
-#endif //UE_ENABLE_INCLUDE_ORDER_DEPRECATED_IN_5_2
 #include "MassCommonFragments.generated.h"
 
 
diff --git a/Engine/Plugins/Runtime/MassGameplay/Source/MassCommon/Public/MassCommonUtils.h b/Engine/Plugins/Runtime/MassGameplay/Source/MassCommon/Public/MassCommonUtils.h
--- a/Engine/Plugins/Runtime/MassGameplay/Source/MassCommon/Public/MassCommonUtils.h
+++ b/Engine/Plugins/Runtime/MassGameplay/Source/MassCommon/Public/MassCommonUtils.h
@@ -2,11 +2,14 @@
 
 #pragma once
 
+#if UE_ENABLE_INCLUDE_ORDER_DEPRECATED_IN_5_6
 #include "MassCommonTypes.h"
+#endif // UE_ENABLE_INCLUDE_ORDER_DEPRECATED_IN_5_6
 #include "Containers/Array.h"
 #include "Containers/Queue.h"
 #include "MassEntityTypes.h"
 
+
 namespace UE::Mass::Utils
 {
 	/** 
diff --git a/Engine/Plugins/Runtime/MassGameplay/Source/MassGameplayDebug/Public/MassGameplayDebugTypes.h b/Engine/Plugins/Runtime/MassGameplay/Source/MassGameplayDebug/Public/MassGameplayDebugTypes.h
--- a/Engine/Plugins/Runtime/MassGameplay/Source/MassGameplayDebug/Public/MassGameplayDebugTypes.h
+++ b/Engine/Plugins/Runtime/MassGameplay/Source/MassGameplayDebug/Public/MassGameplayDebugTypes.h
@@ -6,7 +6,11 @@
 #include "Logging/LogMacros.h"
 #include "Engine/DataTable.h"
 #include "MassEntityTypes.h"
+#include "MassCommonTypes.h"
+#if UE_ENABLE_INCLUDE_ORDER_DEPRECATED_IN_5_6
+#include "CoreMinimal.h"
 #include "MassSpawnerTypes.h"
+#endif // UE_ENABLE_INCLUDE_ORDER_DEPRECATED_IN_5_6
 #include "MassGameplayDebugTypes.generated.h"
 
 
@@ -100,7 +104,3 @@ struct FAgentDebugVisualization : public FTableRowBase
 	UPROPERTY(EditAnywhere, Category = "Mass|Debug")
 	EMassEntityDebugShape WireShape = EMassEntityDebugShape::Box;
 };
-
-#if UE_ENABLE_INCLUDE_ORDER_DEPRECATED_IN_5_6
-#include "CoreMinimal.h"
-#endif // UE_ENABLE_INCLUDE_ORDER_DEPRECATED_IN_5_6
\ No newline at end of file
diff --git a/Engine/Plugins/Runtime/MassGameplay/Source/MassGameplayEditor/Public/MassTraitRepository.h b/Engine/Plugins/Runtime/MassGameplay/Source/MassGameplayEditor/Public/MassTraitRepository.h
--- a/Engine/Plugins/Runtime/MassGameplay/Source/MassGameplayEditor/Public/MassTraitRepository.h
+++ b/Engine/Plugins/Runtime/MassGameplay/Source/MassGameplayEditor/Public/MassTraitRepository.h
@@ -5,7 +5,9 @@
 #include "MassEntityTemplateRegistry.h"
 #include "EditorSubsystem.h"
 #include "MassEntitySubsystem.h"
+#if UE_ENABLE_INCLUDE_ORDER_DEPRECATED_IN_5_6
 #include "StructUtils/StructView.h"
+#endif // UE_ENABLE_INCLUDE_ORDER_DEPRECATED_IN_5_6
 #include "MassDebugger.h"
 #include "MassTraitRepository.generated.h"
 
diff --git a/Engine/Plugins/Runtime/MassGameplay/Source/MassLOD/Private/MassLODSubsystem.cpp b/Engine/Plugins/Runtime/MassGameplay/Source/MassLOD/Private/MassLODSubsystem.cpp
--- a/Engine/Plugins/Runtime/MassGameplay/Source/MassLOD/Private/MassLODSubsystem.cpp
+++ b/Engine/Plugins/Runtime/MassGameplay/Source/MassLOD/Private/MassLODSubsystem.cpp
@@ -1,6 +1,7 @@
 // Copyright Epic Games, Inc. All Rights Reserved.
 
 #include "MassLODSubsystem.h"
+#include "MassLODTypes.h"
 #include "GameFramework/PlayerController.h"
 #include "WorldPartition/WorldPartition.h"
 #include "WorldPartition/WorldPartitionStreamingSource.h"
diff --git a/Engine/Plugins/Runtime/MassGameplay/Source/MassLOD/Public/MassLODCollectorProcessor.h b/Engine/Plugins/Runtime/MassGameplay/Source/MassLOD/Public/MassLODCollectorProcessor.h
--- a/Engine/Plugins/Runtime/MassGameplay/Source/MassLOD/Public/MassLODCollectorProcessor.h
+++ b/Engine/Plugins/Runtime/MassGameplay/Source/MassLOD/Public/MassLODCollectorProcessor.h
@@ -2,12 +2,14 @@
 
 #pragma once
 
+#if UE_ENABLE_INCLUDE_ORDER_DEPRECATED_IN_5_6
 #include "MassEntityTypes.h"
+#endif // UE_ENABLE_INCLUDE_ORDER_DEPRECATED_IN_5_6
 #include "MassProcessor.h"
 #include "MassLODCollector.h"
-
 #include "MassLODCollectorProcessor.generated.h"
 
+
 struct FMassGenericCollectorLogic : public FLODDefaultLogic
 {
 	enum
diff --git a/Engine/Plugins/Runtime/MassGameplay/Source/MassLOD/Public/MassLODDistanceCollectorProcessor.h b/Engine/Plugins/Runtime/MassGameplay/Source/MassLOD/Public/MassLODDistanceCollectorProcessor.h
--- a/Engine/Plugins/Runtime/MassGameplay/Source/MassLOD/Public/MassLODDistanceCollectorProcessor.h
+++ b/Engine/Plugins/Runtime/MassGameplay/Source/MassLOD/Public/MassLODDistanceCollectorProcessor.h
@@ -2,7 +2,9 @@
 
 #pragma once
 
+#if UE_ENABLE_INCLUDE_ORDER_DEPRECATED_IN_5_6
 #include "MassEntityTypes.h"
+#endif // UE_ENABLE_INCLUDE_ORDER_DEPRECATED_IN_5_6
 #include "MassProcessor.h"
 #include "MassLODCollector.h"
 #include "MassLODLogic.h"
diff --git a/Engine/Plugins/Runtime/MassGameplay/Source/MassLOD/Public/MassLODLogic.h b/Engine/Plugins/Runtime/MassGameplay/Source/MassLOD/Public/MassLODLogic.h
--- a/Engine/Plugins/Runtime/MassGameplay/Source/MassLOD/Public/MassLODLogic.h
+++ b/Engine/Plugins/Runtime/MassGameplay/Source/MassLOD/Public/MassLODLogic.h
@@ -5,7 +5,7 @@
 #include "MassLODSubsystem.h"
 #include "Containers/StaticArray.h"
 #include "ConvexVolume.h"
-
+#include "MassLODTypes.h"
 
 #define DECLARE_CONDITIONAL_MEMBER_ACCESSORS( Condition, MemberType, MemberName ) \
 template <bool Condition, typename TemplateClass> static FORCEINLINE typename TEnableIf< Condition, MemberType>::Type Get##MemberName(TemplateClass& Obj, MemberType DefaultValue) { return Obj.MemberName; } \
diff --git a/Engine/Plugins/Runtime/MassGameplay/Source/MassLOD/Public/MassLODSubsystem.h b/Engine/Plugins/Runtime/MassGameplay/Source/MassLOD/Public/MassLODSubsystem.h
--- a/Engine/Plugins/Runtime/MassGameplay/Source/MassLOD/Public/MassLODSubsystem.h
+++ b/Engine/Plugins/Runtime/MassGameplay/Source/MassLOD/Public/MassLODSubsystem.h
@@ -3,11 +3,13 @@
 #pragma once
 
 #include "MassCommonTypes.h"
-#include "MassProcessor.h"
 #include "IndexedHandle.h"
-#include "MassLODTypes.h"
 #include "MassSubsystemBase.h"
 #include "MassExternalSubsystemTraits.h"
+#if UE_ENABLE_INCLUDE_ORDER_DEPRECATED_IN_5_6
+#include "MassLODTypes.h"
+#include "MassProcessor.h"
+#endif // UE_ENABLE_INCLUDE_ORDER_DEPRECATED_IN_5_6
 #include "MassLODSubsystem.generated.h"
 
 class UMassLODSubsystem;
diff --git a/Engine/Plugins/Runtime/MassGameplay/Source/MassLOD/Public/MassLODTickRateController.h b/Engine/Plugins/Runtime/MassGameplay/Source/MassLOD/Public/MassLODTickRateController.h
--- a/Engine/Plugins/Runtime/MassGameplay/Source/MassLOD/Public/MassLODTickRateController.h
+++ b/Engine/Plugins/Runtime/MassGameplay/Source/MassLOD/Public/MassLODTickRateController.h
@@ -5,7 +5,9 @@
 #include "MassLODLogic.h"
 #include "MassLODUtils.h"
 #include "MassCommonUtils.h"
+#if UE_ENABLE_INCLUDE_ORDER_DEPRECATED_IN_5_6
 #include "MassCommandBuffer.h"
+#endif // UE_ENABLE_INCLUDE_ORDER_DEPRECATED_IN_5_6
 
 /**
  * Helper struct to control LOD tick rate for each agent, 
diff --git a/Engine/Plugins/Runtime/MassGameplay/Source/MassLOD/Public/MassLODTypes.h b/Engine/Plugins/Runtime/MassGameplay/Source/MassLOD/Public/MassLODTypes.h
--- a/Engine/Plugins/Runtime/MassGameplay/Source/MassLOD/Public/MassLODTypes.h
+++ b/Engine/Plugins/Runtime/MassGameplay/Source/MassLOD/Public/MassLODTypes.h
@@ -2,10 +2,11 @@
 
 #pragma once
 
-#include "MassEntityTypes.h"
 #include "Containers/StaticArray.h"
+#if UE_ENABLE_INCLUDE_ORDER_DEPRECATED_IN_5_6
+#include "MassEntityTypes.h"
 #include "MassEntityManager.h"
-
+#endif // UE_ENABLE_INCLUDE_OR
 #include "MassLODTypes.generated.h"
 
 /** Debug option to allow multiple viewers per controller. Useful for testing and profiling purposes */
diff --git a/Engine/Plugins/Runtime/MassGameplay/Source/MassLOD/Public/MassLODUtils.h b/Engine/Plugins/Runtime/MassGameplay/Source/MassLOD/Public/MassLODUtils.h
--- a/Engine/Plugins/Runtime/MassGameplay/Source/MassLOD/Public/MassLODUtils.h
+++ b/Engine/Plugins/Runtime/MassGameplay/Source/MassLOD/Public/MassLODUtils.h
@@ -2,8 +2,9 @@
 
 #pragma once
 
-#include "MassEntityTypes.h"
+#if UE_ENABLE_INCLUDE_ORDER_DEPRECATED_IN_5_6
 #include "MassEntityManager.h"
+#endif // UE_ENABLE_INCLUDE_OR
 #include "MassCommonTypes.h"
 #include "MassLODFragments.h"
 #include "Logging/LogMacros.h"
diff --git a/Engine/Plugins/Runtime/MassGameplay/Source/MassMovement/Public/MassMovementTypes.h b/Engine/Plugins/Runtime/MassGameplay/Source/MassMovement/Public/MassMovementTypes.h
--- a/Engine/Plugins/Runtime/MassGameplay/Source/MassMovement/Public/MassMovementTypes.h
+++ b/Engine/Plugins/Runtime/MassGameplay/Source/MassMovement/Public/MassMovementTypes.h
@@ -1,4 +1,4 @@
-﻿// Copyright Epic Games, Inc. All Rights Reserved.
+// Copyright Epic Games, Inc. All Rights Reserved.
 
 #pragma once
 
@@ -75,5 +75,3 @@ struct MASSMOVEMENT_API FMassMovementStyleParameters
 	UPROPERTY(EditAnywhere, Category = "Movement Style")
 	TArray<FMassMovementStyleSpeedParameters> DesiredSpeeds;
 };
-
-
diff --git a/Engine/Plugins/Runtime/MassGameplay/Source/MassReplication/Public/MassClientBubbleHandler.h b/Engine/Plugins/Runtime/MassGameplay/Source/MassReplication/Public/MassClientBubbleHandler.h
--- a/Engine/Plugins/Runtime/MassGameplay/Source/MassReplication/Public/MassClientBubbleHandler.h
+++ b/Engine/Plugins/Runtime/MassGameplay/Source/MassReplication/Public/MassClientBubbleHandler.h
@@ -4,16 +4,18 @@
 
 #include "MassCommonTypes.h"
 #include "MassReplicationTypes.h"
-#include "MassEntityManager.h"
 #include "Containers/ArrayView.h"
 #include "MassClientBubbleSerializerBase.h"
 #include "MassSpawnerSubsystem.h"
-#include "MassSpawnerTypes.h"
 #include "MassReplicationFragments.h"
 #include "MassReplicationSubsystem.h"
 #include "Engine/World.h"
 #include "MassEntityTemplate.h"
 #include "MassEntityView.h"
+#if UE_ENABLE_INCLUDE_ORDER_DEPRECATED_IN_5_6
+#include "MassEntityManager.h"
+#include "MassSpawnerTypes.h"
+#endif // UE_ENABLE_INCLUDE_ORDER_DEPRECATED_IN_5_6
 
 #include "MassClientBubbleHandler.generated.h"
 
diff --git a/Engine/Plugins/Runtime/MassGameplay/Source/MassReplication/Public/MassReplicationFragments.h b/Engine/Plugins/Runtime/MassGameplay/Source/MassReplication/Public/MassReplicationFragments.h
--- a/Engine/Plugins/Runtime/MassGameplay/Source/MassReplication/Public/MassReplicationFragments.h
+++ b/Engine/Plugins/Runtime/MassGameplay/Source/MassReplication/Public/MassReplicationFragments.h
@@ -6,7 +6,7 @@
 #include "MassReplicationTypes.h"
 #include "MassLODCollector.h"
 #include "MassLODCalculator.h"
-
+#include "MassObserverProcessor.h"
 #include "MassReplicationFragments.generated.h"
 
 class AMassClientBubbleInfoBase;
diff --git a/Engine/Plugins/Runtime/MassGameplay/Source/MassReplication/Public/MassReplicationProcessor.h b/Engine/Plugins/Runtime/MassGameplay/Source/MassReplication/Public/MassReplicationProcessor.h
--- a/Engine/Plugins/Runtime/MassGameplay/Source/MassReplication/Public/MassReplicationProcessor.h
+++ b/Engine/Plugins/Runtime/MassGameplay/Source/MassReplication/Public/MassReplicationProcessor.h
@@ -4,15 +4,17 @@
 
 #include "Engine/World.h"
 #include "MassProcessor.h"
+#if UE_ENABLE_INCLUDE_ORDER_DEPRECATED_IN_5_6
 #include "MassCommonTypes.h"
+#endif // UE_ENABLE_INCLUDE_ORDER_DEPRECATED_IN_5_6
 #include "MassSimulationLOD.h"
 #include "MassReplicationTypes.h"
 #include "MassReplicationFragments.h"
 #include "MassSpawnerTypes.h"
-#include "MassProcessor.h"
 #include "MassReplicationSubsystem.h"
 #include "MassReplicationProcessor.generated.h"
 
+
 class UMassReplicationSubsystem;
 class AMassClientBubbleInfoBase;
 class UWorld;
diff --git a/Engine/Plugins/Runtime/MassGameplay/Source/MassReplication/Public/MassReplicationTransformHandlers.h b/Engine/Plugins/Runtime/MassGameplay/Source/MassReplication/Public/MassReplicationTransformHandlers.h
--- a/Engine/Plugins/Runtime/MassGameplay/Source/MassReplication/Public/MassReplicationTransformHandlers.h
+++ b/Engine/Plugins/Runtime/MassGameplay/Source/MassReplication/Public/MassReplicationTransformHandlers.h
@@ -6,8 +6,9 @@
 #include "MassEntityView.h"
 #include "MassClientBubbleHandler.h"
 #include "MassCommonFragments.h"
+#if UE_ENABLE_INCLUDE_ORDER_DEPRECATED_IN_5_6
 #include "AIHelpers.h"
-
+#endif // UE_ENABLE_INCLUDE_ORDER_DEPRECATED_IN_5_6
 #include "MassReplicationTransformHandlers.generated.h"
 
 struct FMassEntityQuery;
diff --git a/Engine/Plugins/Runtime/MassGameplay/Source/MassReplication/Public/MassReplicationTypes.h b/Engine/Plugins/Runtime/MassGameplay/Source/MassReplication/Public/MassReplicationTypes.h
--- a/Engine/Plugins/Runtime/MassGameplay/Source/MassReplication/Public/MassReplicationTypes.h
+++ b/Engine/Plugins/Runtime/MassGameplay/Source/MassReplication/Public/MassReplicationTypes.h
@@ -3,13 +3,14 @@
 #pragma once
 
 #include "MassCommonTypes.h"
-#include "MassLODSubsystem.h"
 #include "IndexedHandle.h"
 #include "MassEntityTemplate.h"
-#include "MassObserverProcessor.h"
 #include "MassLODLogic.h"
 #include "HierarchicalHashGrid2D.h"
-
+#if UE_ENABLE_INCLUDE_ORDER_DEPRECATED_IN_5_6
+#include "MassLODSubsystem.h"
+#include "MassObserverProcessor.h"
+#endif // UE_ENABLE_INCLUDE_ORDER_DEPRECATED_IN_5_6
 #include "MassReplicationTypes.generated.h"
 
 
diff --git a/Engine/Plugins/Runtime/MassGameplay/Source/MassRepresentation/Private/MassDistanceLODProcessor.cpp b/Engine/Plugins/Runtime/MassGameplay/Source/MassRepresentation/Private/MassDistanceLODProcessor.cpp
--- a/Engine/Plugins/Runtime/MassGameplay/Source/MassRepresentation/Private/MassDistanceLODProcessor.cpp
+++ b/Engine/Plugins/Runtime/MassGameplay/Source/MassRepresentation/Private/MassDistanceLODProcessor.cpp
@@ -3,6 +3,10 @@
 #include "MassDistanceLODProcessor.h"
 #include "MassRepresentationDebug.h"
 #include "MassExecutionContext.h"
+#include "MassCommonFragments.h"
+#include "MassLODFragments.h"
+#include "MassLODTypes.h"
+#include "MassRepresentationFragments.h"
 
 UMassDistanceLODProcessor::UMassDistanceLODProcessor()
 {
diff --git a/Engine/Plugins/Runtime/MassGameplay/Source/MassRepresentation/Private/MassVisualizationLODProcessor.cpp b/Engine/Plugins/Runtime/MassGameplay/Source/MassRepresentation/Private/MassVisualizationLODProcessor.cpp
--- a/Engine/Plugins/Runtime/MassGameplay/Source/MassRepresentation/Private/MassVisualizationLODProcessor.cpp
+++ b/Engine/Plugins/Runtime/MassGameplay/Source/MassRepresentation/Private/MassVisualizationLODProcessor.cpp
@@ -1,6 +1,7 @@
 // Copyright Epic Games, Inc. All Rights Reserved.
 
 #include "MassVisualizationLODProcessor.h"
+#include "MassCommonFragments.h"
 #include "MassRepresentationDebug.h"
 #include "MassExecutionContext.h"
 
diff --git a/Engine/Plugins/Runtime/MassGameplay/Source/MassRepresentation/Public/MassDistanceLODProcessor.h b/Engine/Plugins/Runtime/MassGameplay/Source/MassRepresentation/Public/MassDistanceLODProcessor.h
--- a/Engine/Plugins/Runtime/MassGameplay/Source/MassRepresentation/Public/MassDistanceLODProcessor.h
+++ b/Engine/Plugins/Runtime/MassGameplay/Source/MassRepresentation/Public/MassDistanceLODProcessor.h
@@ -2,11 +2,12 @@
 
 #pragma once
 
-#include "MassCommonFragments.h"
-#include "MassRepresentationFragments.h"
 #include "MassProcessor.h"
+#if UE_ENABLE_INCLUDE_ORDER_DEPRECATED_IN_5_6
+#include "MassRepresentationFragments.h"
+#include "MassCommonFragments.h"
 #include "MassLODCalculator.h"
-
+#endif // UE_ENABLE_INCLUDE_ORDER_DEPRECATED_IN_5_6
 #include "MassDistanceLODProcessor.generated.h"
 
 
diff --git a/Engine/Plugins/Runtime/MassGameplay/Source/MassRepresentation/Public/MassRepresentationFragments.h b/Engine/Plugins/Runtime/MassGameplay/Source/MassRepresentation/Public/MassRepresentationFragments.h
--- a/Engine/Plugins/Runtime/MassGameplay/Source/MassRepresentation/Public/MassRepresentationFragments.h
+++ b/Engine/Plugins/Runtime/MassGameplay/Source/MassRepresentation/Public/MassRepresentationFragments.h
@@ -4,14 +4,16 @@
 
 #include "MassLODTypes.h"
 #include "MassEntityTypes.h"
+#if UE_ENABLE_INCLUDE_ORDER_DEPRECATED_IN_5_6
 #include "MassEntityManager.h"
+#endif // UE_ENABLE_INCLUDE_ORDER_DEPRECATED_IN_5_6
 #include "MassActorSpawnerSubsystem.h"
 #include "MassRepresentationTypes.h"
 #include "MassRepresentationActorManagement.h"
 #include "MassLODCalculator.h"
-
 #include "MassRepresentationFragments.generated.h"
 
+
 class UMassRepresentationSubsystem;
 class UMassRepresentationActorManagement;
 
diff --git a/Engine/Plugins/Runtime/MassGameplay/Source/MassRepresentation/Public/MassRepresentationTypes.h b/Engine/Plugins/Runtime/MassGameplay/Source/MassRepresentation/Public/MassRepresentationTypes.h
--- a/Engine/Plugins/Runtime/MassGameplay/Source/MassRepresentation/Public/MassRepresentationTypes.h
+++ b/Engine/Plugins/Runtime/MassGameplay/Source/MassRepresentation/Public/MassRepresentationTypes.h
@@ -3,12 +3,14 @@
 #pragma once
 
 #include "MassLODTypes.h"
+#include "MassEntityTypes.h"
 #include "Engine/DataTable.h"
 #include "Misc/MTAccessDetector.h"
 #include "InstanceDataTypes.h"
 #include "Experimental/Containers/RobinHoodHashTable.h"
 #include "UObject/ObjectKey.h"
 #include "Components/InstancedStaticMeshComponent.h"
+#include "Math/MathFwd.h"
 #include "MassRepresentationTypes.generated.h"
 
 class UMaterialInterface;
diff --git a/Engine/Plugins/Runtime/MassGameplay/Source/MassRepresentation/Public/MassRepresentationUtils.h b/Engine/Plugins/Runtime/MassGameplay/Source/MassRepresentation/Public/MassRepresentationUtils.h
--- a/Engine/Plugins/Runtime/MassGameplay/Source/MassRepresentation/Public/MassRepresentationUtils.h
+++ b/Engine/Plugins/Runtime/MassGameplay/Source/MassRepresentation/Public/MassRepresentationUtils.h
@@ -3,10 +3,13 @@
 #pragma once
 
 #include "MassEntityTypes.h"
+#include "MassLODFragments.h"
+#if UE_ENABLE_INCLUDE_ORDER_DEPRECATED_IN_5_6
 #include "MassEntityManager.h"
 #include "MassCommonTypes.h"
 #include "MassRepresentationTypes.h"
-#include "MassLODFragments.h"
+#endif // UE_ENABLE_INCLUDE_ORDER_DEPRECATED_IN_5_6
+
 
 struct FMassCommandBuffer;
 
diff --git a/Engine/Plugins/Runtime/MassGameplay/Source/MassRepresentation/Public/MassVisualizationComponent.h b/Engine/Plugins/Runtime/MassGameplay/Source/MassRepresentation/Public/MassVisualizationComponent.h
--- a/Engine/Plugins/Runtime/MassGameplay/Source/MassRepresentation/Public/MassVisualizationComponent.h
+++ b/Engine/Plugins/Runtime/MassGameplay/Source/MassRepresentation/Public/MassVisualizationComponent.h
@@ -3,11 +3,11 @@
 #pragma once
 
 #include "MassRepresentationTypes.h"
-#include "MassCommonTypes.h"
 #include "Components/ActorComponent.h"
-#include "MassRepresentationTypes.h"
 #include "Misc/MTAccessDetector.h"
-
+#if UE_ENABLE_INCLUDE_ORDER_DEPRECATED_IN_5_6
+#include "MassCommonTypes.h"
+#endif // UE_ENABLE_INCLUDE_ORDER_DEPRECATED_IN_5_6
 #include "MassVisualizationComponent.generated.h"
 
 
diff --git a/Engine/Plugins/Runtime/MassGameplay/Source/MassRepresentation/Public/MassVisualizationLODProcessor.h b/Engine/Plugins/Runtime/MassGameplay/Source/MassRepresentation/Public/MassVisualizationLODProcessor.h
--- a/Engine/Plugins/Runtime/MassGameplay/Source/MassRepresentation/Public/MassVisualizationLODProcessor.h
+++ b/Engine/Plugins/Runtime/MassGameplay/Source/MassRepresentation/Public/MassVisualizationLODProcessor.h
@@ -2,11 +2,12 @@
 
 #pragma once
 
-#include "MassCommonFragments.h"
 #include "MassRepresentationFragments.h"
 #include "MassProcessor.h"
+#if UE_ENABLE_INCLUDE_ORDER_DEPRECATED_IN_5_6
+#include "MassCommonFragments.h"
 #include "MassLODCalculator.h"
-
+#endif // UE_ENABLE_INCLUDE_ORDER_DEPRECATED_IN_5_6
 #include "MassVisualizationLODProcessor.generated.h"
 
 
diff --git a/Engine/Plugins/Runtime/MassGameplay/Source/MassSmartObjects/Public/MassSmartObjectFragments.h b/Engine/Plugins/Runtime/MassGameplay/Source/MassSmartObjects/Public/MassSmartObjectFragments.h
--- a/Engine/Plugins/Runtime/MassGameplay/Source/MassSmartObjects/Public/MassSmartObjectFragments.h
+++ b/Engine/Plugins/Runtime/MassGameplay/Source/MassSmartObjects/Public/MassSmartObjectFragments.h
@@ -4,7 +4,9 @@
 
 #include "MassEntityTypes.h"
 #include "MassSmartObjectTypes.h"
+#if UE_ENABLE_INCLUDE_ORDER_DEPRECATED_IN_5_6
 #include "MassSmartObjectRequest.h"
+#endif // UE_ENABLE_INCLUDE_ORDER_DEPRECATED_IN_5_6
 #include "SmartObjectRuntime.h"
 #include "MassSmartObjectFragments.generated.h"
 
diff --git a/Engine/Plugins/Runtime/MassGameplay/Source/MassSpawner/Private/MassEntityTemplate.cpp b/Engine/Plugins/Runtime/MassGameplay/Source/MassSpawner/Private/MassEntityTemplate.cpp
--- a/Engine/Plugins/Runtime/MassGameplay/Source/MassSpawner/Private/MassEntityTemplate.cpp
+++ b/Engine/Plugins/Runtime/MassGameplay/Source/MassSpawner/Private/MassEntityTemplate.cpp
@@ -1,6 +1,7 @@
 // Copyright Epic Games, Inc. All Rights Reserved.
 
 #include "MassEntityTemplate.h"
+#include "MassCommonTypes.h"
 #include "VisualLogger/VisualLoggerTypes.h"
 #include "MassDebugger.h"
 #include "Algo/Find.h"
diff --git a/Engine/Plugins/Runtime/MassGameplay/Source/MassSpawner/Private/MassEntityTemplateRegistry.cpp b/Engine/Plugins/Runtime/MassGameplay/Source/MassSpawner/Private/MassEntityTemplateRegistry.cpp
--- a/Engine/Plugins/Runtime/MassGameplay/Source/MassSpawner/Private/MassEntityTemplateRegistry.cpp
+++ b/Engine/Plugins/Runtime/MassGameplay/Source/MassSpawner/Private/MassEntityTemplateRegistry.cpp
@@ -1,11 +1,11 @@
 // Copyright Epic Games, Inc. All Rights Reserved.
 
 #include "MassEntityTemplateRegistry.h"
+#include "MassCommonTypes.h"
 #include "MassSpawnerTypes.h"
 #include "MassEntityManager.h"
 #include "Engine/World.h"
 #include "VisualLogger/VisualLogger.h"
-#include "MassSpawnerSubsystem.h"
 #include "MassEntityTypes.h"
 #include "MassEntityTraitBase.h"
 
diff --git a/Engine/Plugins/Runtime/MassGameplay/Source/MassSpawner/Public/MassEntityTemplate.h b/Engine/Plugins/Runtime/MassGameplay/Source/MassSpawner/Public/MassEntityTemplate.h
--- a/Engine/Plugins/Runtime/MassGameplay/Source/MassSpawner/Public/MassEntityTemplate.h
+++ b/Engine/Plugins/Runtime/MassGameplay/Source/MassSpawner/Public/MassEntityTemplate.h
@@ -2,13 +2,15 @@
 
 #pragma once
 
-#include "MassProcessingTypes.h"
 #include "MassEntityManager.h"
-#include "MassCommonTypes.h"
 #include "MassTranslator.h"
 #include "Templates/SharedPointer.h"
 #include "Misc/Guid.h"
 #include "Hash/CityHash.h"
+#if UE_ENABLE_INCLUDE_ORDER_DEPRECATED_IN_5_6
+#include "MassCommonTypes.h"
+#include "MassProcessingTypes.h"
+#endif // UE_ENABLE_INCLUDE_ORDER_DEPRECATED_IN_5_6
 #include "MassEntityTemplate.generated.h"
 
 
diff --git a/Engine/Plugins/Runtime/MassGameplay/Source/MassSpawner/Public/MassEntityTraitBase.h b/Engine/Plugins/Runtime/MassGameplay/Source/MassSpawner/Public/MassEntityTraitBase.h
--- a/Engine/Plugins/Runtime/MassGameplay/Source/MassSpawner/Public/MassEntityTraitBase.h
+++ b/Engine/Plugins/Runtime/MassGameplay/Source/MassSpawner/Public/MassEntityTraitBase.h
@@ -2,6 +2,9 @@
 
 #pragma once 
 
+#if UE_ENABLE_INCLUDE_ORDER_DEPRECATED_IN_5_5
+#include "CoreMinimal.h"
+#endif // UE_ENABLE_INCLUDE_ORDER_DEPRECATED_IN_5_6
 #include "MassEntityTraitBase.generated.h"
 
 struct FMassEntityTemplateBuildContext;
@@ -82,7 +85,3 @@ private:
 	static FOnNewTraitType OnNewTraitTypeEvent;
 #endif // WITH_EDITOR
 };
-
-#if UE_ENABLE_INCLUDE_ORDER_DEPRECATED_IN_5_5
-#include "CoreMinimal.h"
-#endif
diff --git a/Engine/Plugins/Runtime/MassGameplay/Source/MassSpawner/Public/MassEntityZoneGraphSpawnPointsGenerator.h b/Engine/Plugins/Runtime/MassGameplay/Source/MassSpawner/Public/MassEntityZoneGraphSpawnPointsGenerator.h
--- a/Engine/Plugins/Runtime/MassGameplay/Source/MassSpawner/Public/MassEntityZoneGraphSpawnPointsGenerator.h
+++ b/Engine/Plugins/Runtime/MassGameplay/Source/MassSpawner/Public/MassEntityZoneGraphSpawnPointsGenerator.h
@@ -2,6 +2,9 @@
 
 #pragma once 
 
+#if UE_ENABLE_INCLUDE_ORDER_DEPRECATED_IN_5_5
+#include "CoreMinimal.h"
+#endif // UE_ENABLE_INCLUDE_ORDER_DEPRECATED_IN_5_6
 #include "MassEntitySpawnDataGeneratorBase.h"
 #include "ZoneGraphTypes.h"
 #include "MassEntityZoneGraphSpawnPointsGenerator.generated.h"
@@ -34,7 +37,3 @@ protected:
 	UPROPERTY(EditAnywhere, Category = "ZoneGraph Generator Config")
 	float MaxGap = 300;
 };
-
-#if UE_ENABLE_INCLUDE_ORDER_DEPRECATED_IN_5_6
-#include "CoreMinimal.h"
-#endif // UE_ENABLE_INCLUDE_ORDER_DEPRECATED_IN_5_6
diff --git a/Engine/Plugins/Runtime/MassGameplay/Source/MassSpawner/Public/MassSpawnerSubsystem.h b/Engine/Plugins/Runtime/MassGameplay/Source/MassSpawner/Public/MassSpawnerSubsystem.h
--- a/Engine/Plugins/Runtime/MassGameplay/Source/MassSpawner/Public/MassSpawnerSubsystem.h
+++ b/Engine/Plugins/Runtime/MassGameplay/Source/MassSpawner/Public/MassSpawnerSubsystem.h
@@ -3,7 +3,9 @@
 #pragma once
 
 #include "MassEntityManager.h"
+#if UE_ENABLE_INCLUDE_ORDER_DEPRECATED_IN_5_5
 #include "StructUtils/InstancedStruct.h"
+#endif // UE_ENABLE_INCLUDE_ORDER_DEPRECATED_IN_5_6
 #include "MassSubsystemBase.h"
 #include "MassEntityTemplateRegistry.h"
 #include "MassSpawnerSubsystem.generated.h"
diff --git a/Engine/Source/Runtime/MassEntity/Internal/MassArchetypeData.h b/Engine/Source/Runtime/MassEntity/Internal/MassArchetypeData.h
--- a/Engine/Source/Runtime/MassEntity/Internal/MassArchetypeData.h
+++ b/Engine/Source/Runtime/MassEntity/Internal/MassArchetypeData.h
@@ -2,14 +2,17 @@
 
 #pragma once
 
-#include "MassEntityManager.h"
 #include "MassArchetypeTypes.h"
-
+#if UE_ENABLE_INCLUDE_ORDER_DEPRECATED_IN_5_6
+#include "MassEntityManager.h"
+#endif // UE_ENABLE_INCLUDE_ORDER_DEPRECATED_IN_5_6
 
 struct FMassEntityQuery;
 struct FMassExecutionContext;
 class FOutputDevice;
 struct FMassArchetypeEntityCollection;
+struct FMassFragmentRequirementDescription;
+struct FMassFragmentRequirements;
 
 namespace UE::Mass
 {
diff --git a/Engine/Source/Runtime/MassEntity/Private/MassArchetypeData.cpp b/Engine/Source/Runtime/MassEntity/Private/MassArchetypeData.cpp
--- a/Engine/Source/Runtime/MassEntity/Private/MassArchetypeData.cpp
+++ b/Engine/Source/Runtime/MassEntity/Private/MassArchetypeData.cpp
@@ -6,7 +6,7 @@
 #include "MassEntitySettings.h"
 #include "MassDebugger.h"
 #include "Misc/StringBuilder.h"
-
+#include "MassRequirements.h"
 
 DECLARE_CYCLE_STAT(TEXT("Mass Archetype BatchAdd"), STAT_Mass_ArchetypeBatchAdd, STATGROUP_Mass);
 
diff --git a/Engine/Source/Runtime/MassEntity/Private/MassCommandBuffer.cpp b/Engine/Source/Runtime/MassEntity/Private/MassCommandBuffer.cpp
--- a/Engine/Source/Runtime/MassEntity/Private/MassCommandBuffer.cpp
+++ b/Engine/Source/Runtime/MassEntity/Private/MassCommandBuffer.cpp
@@ -2,7 +2,7 @@
 
 #include "MassCommandBuffer.h"
 #include "Containers/AnsiString.h"
-#include "MassObserverManager.h"
+#include "MassEntityManager.h"
 #include "MassEntityUtils.h"
 #include "HAL/IConsoleManager.h"
 #include "ProfilingDebugging/CsvProfiler.h"
diff --git a/Engine/Source/Runtime/MassEntity/Private/MassExecutionContext.cpp b/Engine/Source/Runtime/MassEntity/Private/MassExecutionContext.cpp
--- a/Engine/Source/Runtime/MassEntity/Private/MassExecutionContext.cpp
+++ b/Engine/Source/Runtime/MassEntity/Private/MassExecutionContext.cpp
@@ -2,6 +2,7 @@
 
 #include "MassExecutionContext.h"
 #include "MassArchetypeData.h"
+#include "MassEntityManager.h"
 
 
 //-----------------------------------------------------------------------------
diff --git a/Engine/Source/Runtime/MassEntity/Public/MassCommandBuffer.h b/Engine/Source/Runtime/MassEntity/Public/MassCommandBuffer.h
--- a/Engine/Source/Runtime/MassEntity/Public/MassCommandBuffer.h
+++ b/Engine/Source/Runtime/MassEntity/Public/MassCommandBuffer.h
@@ -3,11 +3,16 @@
 #pragma once
 
 #include "MassEntityTypes.h"
+#if UE_ENABLE_INCLUDE_ORDER_DEPRECATED_IN_5_6
 #include "MassEntityManager.h"
+#endif // UE_ENABLE_INCLUDE_ORDER_DEPRECATED_IN_5_6
 #include "Misc/MTAccessDetector.h"
 #include "MassEntityUtils.h"
 #include "MassCommands.h"
 
+
+struct FMassEntitymanager;
+
 //@TODO: Consider debug information in case there is an assert when replaying the command buffer
 // (e.g., which system added the command, or even file/line number in development builds for the specific call via a macro)
 
diff --git a/Engine/Source/Runtime/MassEntity/Public/MassDebugger.h b/Engine/Source/Runtime/MassEntity/Public/MassDebugger.h
--- a/Engine/Source/Runtime/MassEntity/Public/MassDebugger.h
+++ b/Engine/Source/Runtime/MassEntity/Public/MassDebugger.h
@@ -8,8 +8,10 @@
 #include "MassEntityQuery.h"
 #include "MassProcessor.h"
 #include "Misc/SpinLock.h"
-#include "StructUtils/InstancedStruct.h"
 #include "Logging/TokenizedMessage.h"
+#if UE_ENABLE_INCLUDE_ORDER_DEPRECATED_IN_5_6
+#include "StructUtils/InstancedStruct.h"
+#endif // UE_ENABLE_INCLUDE_ORDER_DEPRECATED_IN_5_6
 
 class FOutputDevice;
 class UMassProcessor;
diff --git a/Engine/Source/Runtime/MassEntity/Public/MassEntityQuery.h b/Engine/Source/Runtime/MassEntity/Public/MassEntityQuery.h
--- a/Engine/Source/Runtime/MassEntity/Public/MassEntityQuery.h
+++ b/Engine/Source/Runtime/MassEntity/Public/MassEntityQuery.h
@@ -7,7 +7,9 @@
 #include "Containers/UnrealString.h"
 #include "MassEntityTypes.h"
 #include "MassArchetypeTypes.h"
+#if UE_ENABLE_INCLUDE_ORDER_DEPRECATED_IN_5_6
 #include "MassExternalSubsystemTraits.h"
+#endif // UE_ENABLE_INCLUDE_ORDER_DEPRECATED_IN_5_6
 #include "MassRequirements.h"
 #include "MassEntityQuery.generated.h"
 
diff --git a/Engine/Source/Runtime/MassEntity/Public/MassEntitySettings.h b/Engine/Source/Runtime/MassEntity/Public/MassEntitySettings.h
--- a/Engine/Source/Runtime/MassEntity/Public/MassEntitySettings.h
+++ b/Engine/Source/Runtime/MassEntity/Public/MassEntitySettings.h
@@ -5,7 +5,9 @@
 #include "MassSettings.h"
 #include "MassProcessingPhaseManager.h"
 #include "MassProcessor.h"
+#if UE_ENABLE_INCLUDE_ORDER_DEPRECATED_IN_5_6
 #include "StructUtils/InstancedStruct.h"
+#endif // UE_ENABLE_INCLUDE_ORDER_DEPRECATED_IN_5_6
 #include "MassEntitySettings.generated.h"
 
 
diff --git a/Engine/Source/Runtime/MassEntity/Public/MassEntityView.h b/Engine/Source/Runtime/MassEntity/Public/MassEntityView.h
--- a/Engine/Source/Runtime/MassEntity/Public/MassEntityView.h
+++ b/Engine/Source/Runtime/MassEntity/Public/MassEntityView.h
@@ -3,7 +3,9 @@
 #pragma once
 
 #include "MassArchetypeTypes.h"
+#if UE_ENABLE_INCLUDE_ORDER_DEPRECATED_IN_5_6
 #include "StructUtils/InstancedStruct.h"
+#endif // UE_ENABLE_INCLUDE_ORDER_DEPRECATED_IN_5_6
 #include "MassEntityView.generated.h"
 
 
diff --git a/Engine/Source/Runtime/MassEntity/Public/MassExecutionContext.h b/Engine/Source/Runtime/MassEntity/Public/MassExecutionContext.h
--- a/Engine/Source/Runtime/MassEntity/Public/MassExecutionContext.h
+++ b/Engine/Source/Runtime/MassEntity/Public/MassExecutionContext.h
@@ -3,7 +3,9 @@
 #pragma once
 
 #include "MassEntityTypes.h"
+#if UE_ENABLE_INCLUDE_ORDER_DEPRECATED_IN_5_6
 #include "MassExternalSubsystemTraits.h"
+#endif // UE_ENABLE_INCLUDE_ORDER_DEPRECATED_IN_5_6
 #include "MassEntityQuery.h"
 #include "MassSubsystemAccess.h"
 
diff --git a/Engine/Source/Runtime/MassEntity/Public/MassProcessor.h b/Engine/Source/Runtime/MassEntity/Public/MassProcessor.h
--- a/Engine/Source/Runtime/MassEntity/Public/MassProcessor.h
+++ b/Engine/Source/Runtime/MassEntity/Public/MassProcessor.h
@@ -2,6 +2,9 @@
 
 #pragma once
 
+#if UE_ENABLE_INCLUDE_ORDER_DEPRECATED_IN_5_6
+#include "CoreMinimal.h"
+#endif // UE_ENABLE_INCLUDE_ORDER_DEPRECATED_IN_5_6
 #include "MassEntityManager.h"
 #include "MassProcessingTypes.h"
 #include "Async/TaskGraphInterfaces.h"
@@ -12,6 +15,7 @@
 struct FMassProcessingPhaseConfig;
 class UMassCompositeProcessor;
 struct FMassDebugger;
+struct FMassEntityQuery;
 
 enum class EProcessorCompletionStatus : uint8
 {
@@ -423,7 +427,3 @@ inline bool UMassCompositeProcessor::IsEmpty() const
 {
 	return ChildPipeline.IsEmpty();
 }
-
-#if UE_ENABLE_INCLUDE_ORDER_DEPRECATED_IN_5_6
-#include "CoreMinimal.h"
-#endif // UE_ENABLE_INCLUDE_ORDER_DEPRECATED_IN_5_6
