[MassGameplay] Fixed stationary visualization traits failing to render entities:
- Added additional entity template validation to LOD Collector traits - they'll verify whether the relevant LOD Collector processor is enabled by default.
- Fixed basic Stationary Visualization traits not working properly due to missing FMassStationaryISMSwitcherProcessorTag
- Added FMassStaticMeshInstanceVisualizationMeshDesc.bRequiresExternalInstanceIDTracking to influence the value of FMassISMCSharedData.bRequiresExternalInstanceIDTracking -- missing that was the reason MassVisualizationComponent was treating stationary entities as movable and, in consequence failing, to render them.

diff --git a/Engine/Plugins/Runtime/MassGameplay/Source/MassLOD/Private/MassLODTrait.cpp b/Engine/Plugins/Runtime/MassGameplay/Source/MassLOD/Private/MassLODTrait.cpp
--- a/Engine/Plugins/Runtime/MassGameplay/Source/MassLOD/Private/MassLODTrait.cpp
+++ b/Engine/Plugins/Runtime/MassGameplay/Source/MassLOD/Private/MassLODTrait.cpp
@@ -6,8 +6,42 @@
 #include "MassCommonFragments.h"
 #include "MassLODFragments.h"
 #include "MassEntityUtils.h"
+#include "MassLODCollectorProcessor.h"
+#include "MassLODDistanceCollectorProcessor.h"
 
 
+namespace UE::Mass::Private
+{
+	bool ValidateRequiredProcessor(gsl::not_null<const UMassEntityTraitBase*> Trait, TSubclassOf<UMassProcessor> ProcessorClass)
+	{
+		check(ProcessorClass);
+
+		// @todo make this processor class configurable. Could be something like:
+		// * every processor declared "identifying tag", a unique combination that would identify given processor class
+		// * traits can add "required processors"
+		// * tags of required processors are added automatically
+		// * we can also report required processors that are marked as not auto-added.
+		//	-- downside of that last, similar to the one bellow, is that disabling a given processor might be deliberate and it can still be enabled at runtime
+		const UMassProcessor* LODCollectorProcessor = ProcessorClass->GetDefaultObject<UMassProcessor>();
+		if (LODCollectorProcessor)
+		{
+			if (LODCollectorProcessor->ShouldAutoAddToGlobalList() == false)
+			{
+				CA_ASSUME(LODCollectorProcessor->GetClass());
+				UE_LOG(LogMassLOD, Error, TEXT("Using %s trait while required processor %s is not enabled by default")
+					, *Trait->GetName(), *LODCollectorProcessor->GetClass()->GetName());
+				return false;
+			}
+
+			return true;
+		}
+		checkf(LODCollectorProcessor, TEXT("Failed to fetch CDO from class %s"), *ProcessorClass->GetName());
+		return false;
+	}
+}
+
+using UE::Mass::Private::ValidateRequiredProcessor;
+
 //-----------------------------------------------------------------------------
 // UMassLODCollectorTrait
 //-----------------------------------------------------------------------------
@@ -18,6 +52,17 @@ void UMassLODCollectorTrait::BuildTemplate(FMassEntityTemplateBuildContext& Buil
 	BuildContext.RequireFragment<FTransformFragment>();
 }
 
+bool UMassLODCollectorTrait::ValidateTemplate(const FMassEntityTemplateBuildContext& BuildContext, const UWorld& World, FAdditionalTraitRequirements& OutTraitRequirements) const
+{
+	// if enabled, we require MassLODCollectorProcessor to be enabled.
+	if (bTestCollectorProcessor && !ValidateRequiredProcessor(this, UMassLODCollectorProcessor::StaticClass()))
+	{
+		return false;
+	}
+
+	return Super::ValidateTemplate(BuildContext, World, OutTraitRequirements);
+}
+
 //-----------------------------------------------------------------------------
 // UMassDistanceLODCollectorTrait
 //-----------------------------------------------------------------------------
@@ -28,6 +73,17 @@ void UMassDistanceLODCollectorTrait::BuildTemplate(FMassEntityTemplateBuildConte
 	BuildContext.RequireFragment<FTransformFragment>();
 }
 
+bool UMassDistanceLODCollectorTrait::ValidateTemplate(const FMassEntityTemplateBuildContext& BuildContext, const UWorld& World, FAdditionalTraitRequirements& OutTraitRequirements) const
+{
+	// if enabled, we require UMassLODDistanceCollectorProcessor to be enabled.
+	if (bTestCollectorProcessor && !ValidateRequiredProcessor(this, UMassLODDistanceCollectorProcessor::StaticClass()))
+	{
+		return false;
+	}
+
+	return Super::ValidateTemplate(BuildContext, World, OutTraitRequirements);
+}
+
 //-----------------------------------------------------------------------------
 // UMassSimulationLODTrait
 //-----------------------------------------------------------------------------
diff --git a/Engine/Plugins/Runtime/MassGameplay/Source/MassLOD/Public/MassLODTrait.h b/Engine/Plugins/Runtime/MassGameplay/Source/MassLOD/Public/MassLODTrait.h
--- a/Engine/Plugins/Runtime/MassGameplay/Source/MassLOD/Public/MassLODTrait.h
+++ b/Engine/Plugins/Runtime/MassGameplay/Source/MassLOD/Public/MassLODTrait.h
@@ -12,7 +12,13 @@ class MASSLOD_API UMassLODCollectorTrait : public UMassEntityTraitBase
 	GENERATED_BODY()
 
 public:
+	/** Whether we should verify that the LOD collector processor associated with this trait is enabled by default*/
+	UPROPERTY(Category="LOD", EditAnywhere, meta=(AdvancedDisplay), config)
+	bool bTestCollectorProcessor = true;
+
+protected:
 	virtual void BuildTemplate(FMassEntityTemplateBuildContext& BuildContext, const UWorld& World) const override;
+	virtual bool ValidateTemplate(const FMassEntityTemplateBuildContext& BuildContext, const UWorld& World, FAdditionalTraitRequirements& OutTraitRequirements) const override;
 };
 
 // Simplest version of UMassLODCollectorTrait that will ensure collection strictly based on Distance from Viewer
@@ -22,7 +28,13 @@ class MASSLOD_API UMassDistanceLODCollectorTrait : public UMassEntityTraitBase
 	GENERATED_BODY()
 
 public:
+	/** Whether we should verify that the LOD collector processor associated with this trait is enabled by default*/
+	UPROPERTY(Category="LOD", EditAnywhere, meta=(AdvancedDisplay), config)
+	bool bTestCollectorProcessor = true;
+
+protected:
 	virtual void BuildTemplate(FMassEntityTemplateBuildContext& BuildContext, const UWorld& World) const override;
+	virtual bool ValidateTemplate(const FMassEntityTemplateBuildContext& BuildContext, const UWorld& World, FAdditionalTraitRequirements& OutTraitRequirements) const override;
 };
 
 UCLASS(meta = (DisplayName = "SimulationLOD"))
diff --git a/Engine/Plugins/Runtime/MassGameplay/Source/MassRepresentation/Private/MassStationaryDistanceVisualizationTrait.cpp b/Engine/Plugins/Runtime/MassGameplay/Source/MassRepresentation/Private/MassStationaryDistanceVisualizationTrait.cpp
--- a/Engine/Plugins/Runtime/MassGameplay/Source/MassRepresentation/Private/MassStationaryDistanceVisualizationTrait.cpp
+++ b/Engine/Plugins/Runtime/MassGameplay/Source/MassRepresentation/Private/MassStationaryDistanceVisualizationTrait.cpp
@@ -2,7 +2,9 @@
 
 #include "MassStationaryDistanceVisualizationTrait.h"
 #include "MassEntityTemplateRegistry.h"
+#include "MassStationaryISMSwitcherProcessor.h"
 #include "VisualLogger/VisualLogger.h"
+#include "StationaryVisualizationTraitUtils.h"
 
 
 UMassStationaryDistanceVisualizationTrait::UMassStationaryDistanceVisualizationTrait(const FObjectInitializer& ObjectInitializer)
@@ -12,17 +14,11 @@ UMassStationaryDistanceVisualizationTrait::UMassStationaryDistanceVisualizationT
 
 void UMassStationaryDistanceVisualizationTrait::BuildTemplate(FMassEntityTemplateBuildContext& BuildContext, const UWorld& World) const
 {
-	bool bIssuesFound = false;
-	for (FMassStaticMeshInstanceVisualizationMeshDesc& MeshDesc : StaticMeshInstanceDesc.Meshes)
-	{
-		bIssuesFound = bIssuesFound || (MeshDesc.Mobility != EComponentMobility::Stationary);
-		MeshDesc.Mobility = EComponentMobility::Stationary;
-	}
-	UE_CVLOG_UELOG(bIssuesFound, this, LogMass, Log, TEXT("%s some Meshes' mobility has been set to non-Stationary. Theese settings will be overridden. "), *GetPathName());
+	UE::Mass::LOD::Private::SetUpStationaryVisualizationTrait(*this, BuildContext, StaticMeshInstanceDesc);
 
 	Super::BuildTemplate(BuildContext, World);
 
-	BuildContext.AddTag<FMassStaticRepresentationTag>();
+	BuildContext.RequireTag<FMassCollectDistanceLODViewerInfoTag>();
 }
 
 #if WITH_EDITOR
diff --git a/Engine/Plugins/Runtime/MassGameplay/Source/MassRepresentation/Private/MassStationaryVisualizationTrait.cpp b/Engine/Plugins/Runtime/MassGameplay/Source/MassRepresentation/Private/MassStationaryVisualizationTrait.cpp
--- a/Engine/Plugins/Runtime/MassGameplay/Source/MassRepresentation/Private/MassStationaryVisualizationTrait.cpp
+++ b/Engine/Plugins/Runtime/MassGameplay/Source/MassRepresentation/Private/MassStationaryVisualizationTrait.cpp
@@ -2,9 +2,30 @@
 
 #include "MassStationaryVisualizationTrait.h"
 #include "MassEntityTemplateRegistry.h"
+#include "MassStationaryISMSwitcherProcessor.h"
 #include "VisualLogger/VisualLogger.h"
 
 
+namespace UE::Mass::LOD::Private
+{
+	void SetUpStationaryVisualizationTrait(const UMassEntityTraitBase& Trait, FMassEntityTemplateBuildContext& BuildContext, FStaticMeshInstanceVisualizationDesc& StaticMeshInstanceDesc)
+	{
+		bool bIssuesFound = false;
+		for (FMassStaticMeshInstanceVisualizationMeshDesc& MeshDesc : StaticMeshInstanceDesc.Meshes)
+		{
+			bIssuesFound = bIssuesFound || (MeshDesc.Mobility != EComponentMobility::Stationary);
+			MeshDesc.Mobility = EComponentMobility::Stationary;
+			MeshDesc.bRequiresExternalInstanceIDTracking = true;
+		}
+		UE_CVLOG_UELOG(bIssuesFound, &Trait, LogMass, Log, TEXT("%s some Meshes' mobility has been set to non-Stationary. These settings will be overridden."), *Trait.GetPathName());
+		BuildContext.AddTag<FMassStaticRepresentationTag>();
+		BuildContext.AddTag<FMassStationaryISMSwitcherProcessorTag>();
+	}
+}
+
+//-----------------------------------------------------------------------------
+// UMassStationaryVisualizationTrait
+//-----------------------------------------------------------------------------
 UMassStationaryVisualizationTrait::UMassStationaryVisualizationTrait(const FObjectInitializer& ObjectInitializer)
 {
 	bAllowServerSideVisualization = true;
@@ -12,17 +33,11 @@ UMassStationaryVisualizationTrait::UMassStationaryVisualizationTrait(const FObje
 
 void UMassStationaryVisualizationTrait::BuildTemplate(FMassEntityTemplateBuildContext& BuildContext, const UWorld& World) const
 {
-	bool bIssuesFound = false;
-	for (FMassStaticMeshInstanceVisualizationMeshDesc& MeshDesc : StaticMeshInstanceDesc.Meshes)
-	{
-		bIssuesFound = bIssuesFound || (MeshDesc.Mobility != EComponentMobility::Stationary);
-		MeshDesc.Mobility = EComponentMobility::Stationary;
-	}
-	UE_CVLOG_UELOG(bIssuesFound, this, LogMass, Log, TEXT("%s some Meshes' mobility has been set to non-Stationary. Theese settings will be overridden. "), *GetPathName());
+	UE::Mass::LOD::Private::SetUpStationaryVisualizationTrait(*this, BuildContext, StaticMeshInstanceDesc);
 
 	Super::BuildTemplate(BuildContext, World);
 
-	BuildContext.AddTag<FMassStaticRepresentationTag>();
+	BuildContext.RequireTag<FMassCollectLODViewerInfoTag>();
 }
 
 #if WITH_EDITOR
diff --git a/Engine/Plugins/Runtime/MassGameplay/Source/MassRepresentation/Private/MassVisualizationComponent.cpp b/Engine/Plugins/Runtime/MassGameplay/Source/MassRepresentation/Private/MassVisualizationComponent.cpp
--- a/Engine/Plugins/Runtime/MassGameplay/Source/MassRepresentation/Private/MassVisualizationComponent.cpp
+++ b/Engine/Plugins/Runtime/MassGameplay/Source/MassRepresentation/Private/MassVisualizationComponent.cpp
@@ -262,7 +262,7 @@ void UMassVisualizationComponent::ConstructStaticMeshComponents()
 
 				if (SharedData == nullptr)
 				{
-					SharedData = &ISMCSharedData.Add(ISMC, FMassISMCSharedData(ISMC));
+					SharedData = &ISMCSharedData.Add(ISMC, FMassISMCSharedData(ISMC, MeshDesc.bRequiresExternalInstanceIDTracking));
 				}
 				else
 				{
diff --git a/Engine/Plugins/Runtime/MassGameplay/Source/MassRepresentation/Private/StationaryVisualizationTraitUtils.h b/Engine/Plugins/Runtime/MassGameplay/Source/MassRepresentation/Private/StationaryVisualizationTraitUtils.h
new file mode 100644
--- /dev/null
+++ b/Engine/Plugins/Runtime/MassGameplay/Source/MassRepresentation/Private/StationaryVisualizationTraitUtils.h
@@ -0,0 +1,12 @@
+// Copyright Epic Games, Inc. All Rights Reserved.
+
+#pragma once
+
+class UMassEntityTraitBase;
+struct FMassEntityTemplateBuildContext;
+struct FStaticMeshInstanceVisualizationDesc;
+
+namespace UE::Mass::LOD::Private
+{
+	void SetUpStationaryVisualizationTrait(const UMassEntityTraitBase& Trait, FMassEntityTemplateBuildContext& BuildContext, FStaticMeshInstanceVisualizationDesc& StaticMeshInstanceDesc);
+}
\ No newline at end of file
diff --git a/Engine/Plugins/Runtime/MassGameplay/Source/MassRepresentation/Public/MassRepresentationTypes.h b/Engine/Plugins/Runtime/MassGameplay/Source/MassRepresentation/Public/MassRepresentationTypes.h
--- a/Engine/Plugins/Runtime/MassGameplay/Source/MassRepresentation/Public/MassRepresentationTypes.h
+++ b/Engine/Plugins/Runtime/MassGameplay/Source/MassRepresentation/Public/MassRepresentationTypes.h
@@ -73,6 +73,13 @@ struct MASSREPRESENTATION_API FMassStaticMeshInstanceVisualizationMeshDesc
 	/** Controls whether the ISM can cast shadow or not */
 	UPROPERTY(EditAnywhere, Category = "Mass|Visual")
 	bool bCastShadows = false;
+
+	/**
+	 * Determines the value of corresponding FMassISMCSharedData.bRequiresExternalInstanceIDTracking.
+	 * Recommended enabled for stationary meshes.
+	 * @see FMassISMCSharedData.bRequiresExternalInstanceIDTracking for details.
+	 */
+	bool bRequiresExternalInstanceIDTracking = false;
 	
 	/** Controls the mobility of the ISM */
 	EComponentMobility::Type Mobility = EComponentMobility::Movable;
diff --git a/Engine/Plugins/Runtime/MassGameplay/Source/MassSpawner/Public/MassEntityTraitBase.h b/Engine/Plugins/Runtime/MassGameplay/Source/MassSpawner/Public/MassEntityTraitBase.h
--- a/Engine/Plugins/Runtime/MassGameplay/Source/MassSpawner/Public/MassEntityTraitBase.h
+++ b/Engine/Plugins/Runtime/MassGameplay/Source/MassSpawner/Public/MassEntityTraitBase.h
@@ -16,7 +16,7 @@ struct FMassEntityTemplateBuildContext;
  * For example, a fragment can be added based on a referenced asset, or some memory hungry settings can be
  * cached and just and index stored on a fragment.
  */
-UCLASS(Abstract, BlueprintType, EditInlineNew, CollapseCategories)
+UCLASS(Abstract, BlueprintType, EditInlineNew, CollapseCategories, config = Mass, defaultconfig)
 class MASSSPAWNER_API UMassEntityTraitBase : public UObject
 {
 	GENERATED_BODY()
