[MassGameplay] Fixed Mass Visualization issue ignoring FStaticMeshInstanceVisualizationDesc.TransformOffset when comparing two instances. As a result any entity using given mesh would use the same representation index, regardless of requested mesh transform offset.

diff --git a/Engine/Plugins/Runtime/MassGameplay/Source/MassRepresentation/Public/MassRepresentationTypes.h b/Engine/Plugins/Runtime/MassGameplay/Source/MassRepresentation/Public/MassRepresentationTypes.h
--- a/Engine/Plugins/Runtime/MassGameplay/Source/MassRepresentation/Public/MassRepresentationTypes.h
+++ b/Engine/Plugins/Runtime/MassGameplay/Source/MassRepresentation/Public/MassRepresentationTypes.h
@@ -145,7 +145,7 @@ struct FStaticMeshInstanceVisualizationDesc : public FTableRowBase
 
 	bool operator==(const FStaticMeshInstanceVisualizationDesc& Other) const
 	{
-		return Meshes == Other.Meshes;
+		return Meshes == Other.Meshes && bUseTransformOffset == Other.bUseTransformOffset && TransformOffset.Equals(Other.TransformOffset);
 	}
 
 	void Reset()
