[MassGamplay] Fixed UMassActorSpawnerSubsystem::SpawnOrRetrieveFromPool logic to take into consideration that actors in the pool can get nulled-out due to being GCed.

diff --git a/Engine/Plugins/Runtime/MassGameplay/Source/MassActors/Private/MassActorSpawnerSubsystem.cpp b/Engine/Plugins/Runtime/MassGameplay/Source/MassActors/Private/MassActorSpawnerSubsystem.cpp
--- a/Engine/Plugins/Runtime/MassGameplay/Source/MassActors/Private/MassActorSpawnerSubsystem.cpp
+++ b/Engine/Plugins/Runtime/MassGameplay/Source/MassActors/Private/MassActorSpawnerSubsystem.cpp
@@ -228,23 +228,30 @@ ESpawnRequestStatus UMassActorSpawnerSubsystem::SpawnOrRetrieveFromPool(FConstSt
 
 		if (Pool && Pool->Num() > 0)
 		{
-			AActor* PooledActor = (*Pool)[0];
-			Pool->RemoveAt(0);
-			--NumActorPooled;
-			PooledActor->SetActorHiddenInGame(false);
-			PooledActor->SetActorTransform(SpawnRequest.Transform, false, nullptr, ETeleportType::ResetPhysics);
-
-			IMassActorPoolableInterface::Execute_PrepareForGame(PooledActor);
+			AActor* PooledActor = nullptr;
+			do
+			{
+				PooledActor = Pool->Pop(EAllowShrinking::No);
+				--NumActorPooled;
+			} while (IsValid(PooledActor) == false && Pool->Num() > 0);
 
-			if (UMassAgentComponent* AgentComp = PooledActor->FindComponentByClass<UMassAgentComponent>())
+			if (IsValid(PooledActor))
 			{
-				// normally this function gets called from UMassAgentComponent::OnRegister. We need to call it manually here
-				// since we're bringing this actor our of a pool.
-				AgentComp->RegisterWithAgentSubsystem();
-			}
+				PooledActor->SetActorHiddenInGame(false);
+				PooledActor->SetActorTransform(SpawnRequest.Transform, false, nullptr, ETeleportType::ResetPhysics);
 
-			OutSpawnedActor = PooledActor;
-			return ESpawnRequestStatus::Succeeded;
+				IMassActorPoolableInterface::Execute_PrepareForGame(PooledActor);
+
+				if (UMassAgentComponent* AgentComp = PooledActor->FindComponentByClass<UMassAgentComponent>())
+				{
+					// normally this function gets called from UMassAgentComponent::OnRegister. We need to call it manually here
+					// since we're bringing this actor our of a pool.
+					AgentComp->RegisterWithAgentSubsystem();
+				}
+
+				OutSpawnedActor = PooledActor;
+				return ESpawnRequestStatus::Succeeded;
+			}
 		}
 	}
 
diff --git a/Engine/Plugins/Runtime/MassGameplay/Source/MassActors/Public/MassActorSpawnerSubsystem.h b/Engine/Plugins/Runtime/MassGameplay/Source/MassActors/Public/MassActorSpawnerSubsystem.h
--- a/Engine/Plugins/Runtime/MassGameplay/Source/MassActors/Public/MassActorSpawnerSubsystem.h
+++ b/Engine/Plugins/Runtime/MassGameplay/Source/MassActors/Public/MassActorSpawnerSubsystem.h
@@ -283,6 +283,7 @@ protected:
 
 	bool bActorPoolingEnabled = true;
 
+	// @todo should use FObjetKey for the map key.
 	TMap<TSubclassOf<AActor>, TArray<TObjectPtr<AActor>>> PooledActors;
 
 	FMassEntityHandleManager_ActorSpawnRequest SpawnRequestHandleManager;
