[MassGameplay] Fixed duplicating MassAgentComponent while dupicating Blueprints not updating hosted Entity Config's GUID (resulting in two BP classes with Mass Agent components seemingly pointing to the same config, which doesn't need to be true).

diff --git a/Engine/Plugins/Runtime/MassGameplay/Source/MassActors/Private/MassAgentComponent.cpp b/Engine/Plugins/Runtime/MassGameplay/Source/MassActors/Private/MassAgentComponent.cpp
--- a/Engine/Plugins/Runtime/MassGameplay/Source/MassActors/Private/MassAgentComponent.cpp
+++ b/Engine/Plugins/Runtime/MassGameplay/Source/MassActors/Private/MassAgentComponent.cpp
@@ -62,6 +62,12 @@ void UMassAgentComponent::PostLoad()
 		bAutoRegister = bAutoRegisterInEditorMode || GetWorld()->IsGameWorld();
 	}
 }
+
+void UMassAgentComponent::PostDuplicate(const bool bDuplicateForPIE)
+{
+	Super::PostDuplicate(bDuplicateForPIE);
+	EntityConfig.PostDuplicate(bDuplicateForPIE);
+}
 #endif // WITH_EDITOR
 
 void UMassAgentComponent::OnRegister()
diff --git a/Engine/Plugins/Runtime/MassGameplay/Source/MassActors/Public/MassAgentComponent.h b/Engine/Plugins/Runtime/MassGameplay/Source/MassActors/Public/MassAgentComponent.h
--- a/Engine/Plugins/Runtime/MassGameplay/Source/MassActors/Public/MassAgentComponent.h
+++ b/Engine/Plugins/Runtime/MassGameplay/Source/MassActors/Public/MassAgentComponent.h
@@ -53,6 +53,7 @@ protected:
 #if WITH_EDITOR
 	virtual void PostInitProperties() override;
 	virtual void PostLoad() override;
+	virtual void PostDuplicate(const bool bDuplicateForPIE) override;
 #endif // WITH_EDITOR
 
 public:
