[MassGameplay] non unity MassEQS fixes

#rnx

diff --git a/Engine/Plugins/Runtime/MassGameplay/Source/MassEQS/Private/Generators/MassEnvQueryGenerator_MassEntityHandles.cpp b/Engine/Plugins/Runtime/MassGameplay/Source/MassEQS/Private/Generators/MassEnvQueryGenerator_MassEntityHandles.cpp
--- a/Engine/Plugins/Runtime/MassGameplay/Source/MassEQS/Private/Generators/MassEnvQueryGenerator_MassEntityHandles.cpp
+++ b/Engine/Plugins/Runtime/MassGameplay/Source/MassEQS/Private/Generators/MassEnvQueryGenerator_MassEntityHandles.cpp
@@ -1,13 +1,11 @@
 // Copyright Epic Games, Inc. All Rights Reserved.
 
 #include "Generators/MassEnvQueryGenerator_MassEntityHandles.h"
-
 #include "EnvironmentQuery/Contexts/EnvQueryContext_Querier.h"
 #include "Items/EnvQueryItemType_MassEntityHandle.h"
 #include "MassEQSSubsystem.h"
 #include "MassEQSUtils.h"
 
-#include "MassEntityTypes.h"
 
 UMassEnvQueryGenerator_MassEntityHandles::UMassEnvQueryGenerator_MassEntityHandles(const FObjectInitializer& ObjectInitializer)
 	: Super(ObjectInitializer)
diff --git a/Engine/Plugins/Runtime/MassGameplay/Source/MassEQS/Private/Items/EnvQueryItemType_MassEntityHandle.cpp b/Engine/Plugins/Runtime/MassGameplay/Source/MassEQS/Private/Items/EnvQueryItemType_MassEntityHandle.cpp
--- a/Engine/Plugins/Runtime/MassGameplay/Source/MassEQS/Private/Items/EnvQueryItemType_MassEntityHandle.cpp
+++ b/Engine/Plugins/Runtime/MassGameplay/Source/MassEQS/Private/Items/EnvQueryItemType_MassEntityHandle.cpp
@@ -1,10 +1,8 @@
 // Copyright Epic Games, Inc. All Rights Reserved.
 
 #include "Items/EnvQueryItemType_MassEntityHandle.h"
-
 #include "EnvironmentQuery/EnvQueryTypes.h"
 
-#include "MassEntitySubsystem.h"
 
 UEnvQueryItemType_MassEntityHandle::UEnvQueryItemType_MassEntityHandle()
 {
diff --git a/Engine/Plugins/Runtime/MassGameplay/Source/MassEQS/Private/Processors/MassEnvQueryGeneratorProcessor_MassEntityHandles.cpp b/Engine/Plugins/Runtime/MassGameplay/Source/MassEQS/Private/Processors/MassEnvQueryGeneratorProcessor_MassEntityHandles.cpp
--- a/Engine/Plugins/Runtime/MassGameplay/Source/MassEQS/Private/Processors/MassEnvQueryGeneratorProcessor_MassEntityHandles.cpp
+++ b/Engine/Plugins/Runtime/MassGameplay/Source/MassEQS/Private/Processors/MassEnvQueryGeneratorProcessor_MassEntityHandles.cpp
@@ -1,18 +1,14 @@
 // Copyright Epic Games, Inc. All Rights Reserved.
 
 #include "Processors/MassEnvQueryGeneratorProcessor_MassEntityHandles.h"
-
 #include "MassExecutionContext.h"
 #include "MassCommonFragments.h"
 #include "MassObserverRegistry.h"
-
 #include "MassEQSUtils.h"
 #include "MassEQSSubsystem.h"
-
-#include "EnvironmentQuery/EnvQueryTypes.h"
-#include "Items/EnvQueryItemType_MassEntityHandle.h"
 #include "Generators/MassEnvQueryGenerator_MassEntityHandles.h"
 
+
 UMassEnvQueryGeneratorProcessor_MassEntityHandles::UMassEnvQueryGeneratorProcessor_MassEntityHandles()
 	: EntityQuery(*this)
 {
diff --git a/Engine/Plugins/Runtime/MassGameplay/Source/MassEQS/Private/Processors/MassEnvQueryTestProcessor_MassEntityTags.cpp b/Engine/Plugins/Runtime/MassGameplay/Source/MassEQS/Private/Processors/MassEnvQueryTestProcessor_MassEntityTags.cpp
--- a/Engine/Plugins/Runtime/MassGameplay/Source/MassEQS/Private/Processors/MassEnvQueryTestProcessor_MassEntityTags.cpp
+++ b/Engine/Plugins/Runtime/MassGameplay/Source/MassEQS/Private/Processors/MassEnvQueryTestProcessor_MassEntityTags.cpp
@@ -1,17 +1,13 @@
 // Copyright Epic Games, Inc. All Rights Reserved.
 
 #include "Processors/MassEnvQueryTestProcessor_MassEntityTags.h"
-
 #include "Tests/MassEnvQueryTest_MassEntityTags.h"
-
 #include "MassEQSSubsystem.h"
-#include "MassEQSTypes.h"
 #include "MassEQSUtils.h"
 #include "MassEQS.h"
-
-#include "MassCommonFragments.h"
 #include "MassExecutionContext.h"
 
+
 namespace Mass::EQS::Utils
 {
 
diff --git a/Engine/Plugins/Runtime/MassGameplay/Source/MassEQS/Private/Tests/MassEnvQueryTest_MassEntityTags.cpp b/Engine/Plugins/Runtime/MassGameplay/Source/MassEQS/Private/Tests/MassEnvQueryTest_MassEntityTags.cpp
--- a/Engine/Plugins/Runtime/MassGameplay/Source/MassEQS/Private/Tests/MassEnvQueryTest_MassEntityTags.cpp
+++ b/Engine/Plugins/Runtime/MassGameplay/Source/MassEQS/Private/Tests/MassEnvQueryTest_MassEntityTags.cpp
@@ -1,16 +1,10 @@
 // Copyright Epic Games, Inc. All Rights Reserved.
 
 #include "Tests/MassEnvQueryTest_MassEntityTags.h"
-
 #include "Items/EnvQueryItemType_MassEntityHandle.h"
-
 #include "MassEQSUtils.h"
-#include "MassEQSTypes.h"
 #include "MassEQSSubsystem.h"
 
-#include "MassCommandBuffer.h"
-#include "MassEntityManager.h"
-#include "MassEntitySubsystem.h"
 
 UMassEnvQueryTest_MassEntityTags::UMassEnvQueryTest_MassEntityTags(const FObjectInitializer& ObjectInitializer)
 	: Super(ObjectInitializer)
diff --git a/Engine/Plugins/Runtime/MassGameplay/Source/MassEQS/Public/Tests/MassEnvQueryTest_MassEntityTags.h b/Engine/Plugins/Runtime/MassGameplay/Source/MassEQS/Public/Tests/MassEnvQueryTest_MassEntityTags.h
--- a/Engine/Plugins/Runtime/MassGameplay/Source/MassEQS/Public/Tests/MassEnvQueryTest_MassEntityTags.h
+++ b/Engine/Plugins/Runtime/MassGameplay/Source/MassEQS/Public/Tests/MassEnvQueryTest_MassEntityTags.h
@@ -3,9 +3,10 @@
 #pragma once
 
 #include "Tests/MassEnvQueryTest.h"
+#include "StructUtils/InstancedStruct.h"
 #include "MassEnvQueryTest_MassEntityTags.generated.h"
 
-struct FInstancedStruct;
+
 struct FMassEntityHandle;
 
 /** Different modes that this Test can be run in */
