[MassGameplay] relaxed conditions when setting up the entity in UMassAgentComponent::SetEntityHandleInternal - we used to do it only when `IsNetSimulating` and I believe that was an unnecessary limitation.

#rnx

diff --git a/Engine/Plugins/Runtime/MassGameplay/Source/MassActors/Private/MassAgentComponent.cpp b/Engine/Plugins/Runtime/MassGameplay/Source/MassActors/Private/MassAgentComponent.cpp
--- a/Engine/Plugins/Runtime/MassGameplay/Source/MassActors/Private/MassAgentComponent.cpp
+++ b/Engine/Plugins/Runtime/MassGameplay/Source/MassActors/Private/MassAgentComponent.cpp
@@ -196,32 +196,29 @@ void UMassAgentComponent::SetEntityHandleInternal(const FMassEntityHandle NewHan
 	// Sync up with mass
 	if (EntitySubsystem)
 	{
-		if (IsNetSimulating())
-		{
-			const FMassEntityView EntityView(EntitySubsystem->GetEntityManager(), AgentHandle);
+		const FMassEntityView EntityView(EntitySubsystem->GetEntityManager(), AgentHandle);
 
-			// @todo Find a way to add these initialization into either translator initializer or adding new fragments
-			// Make sure to fetch the fragment after any release, as that action can move the entity around into new archetype and 
-			// by the same fact change the references to the fragments.
-			if (FMassActorFragment* ActorInfo = EntityView.GetFragmentDataPtr<FMassActorFragment>())
-			{
-				checkf(!ActorInfo->IsValid(), TEXT("Expecting ActorInfo fragment to be null"));
-				ActorInfo->SetAndUpdateHandleMap(AgentHandle, GetOwner(), !IsNetSimulating()/*bIsOwnedByMass*/);
-			}
+		// @todo Find a way to add these initialization into either translator initializer or adding new fragments
+		// Make sure to fetch the fragment after any release, as that action can move the entity around into new archetype and 
+		// by the same fact change the references to the fragments.
+		if (FMassActorFragment* ActorInfo = EntityView.GetFragmentDataPtr<FMassActorFragment>())
+		{
+			checkf(!ActorInfo->IsValid(), TEXT("Expecting ActorInfo fragment to be null"));
+			ActorInfo->SetAndUpdateHandleMap(AgentHandle, GetOwner(), !IsNetSimulating()/*bIsOwnedByMass*/);
+		}
 
-			// Initialize location of the replicated actor to match the mass replicated one
-			if (const FTransformFragment* TransformFragment = EntityView.GetFragmentDataPtr<FTransformFragment>())
-			{
-				GetOwner()->SetActorTransform(TransformFragment->GetTransform(), /*bSweep*/false, /*OutSweepHitResult*/nullptr, ETeleportType::TeleportPhysics);
-			}
+		// Initialize location of the replicated actor to match the mass replicated one
+		if (const FTransformFragment* TransformFragment = EntityView.GetFragmentDataPtr<FTransformFragment>())
+		{
+			GetOwner()->SetActorTransform(TransformFragment->GetTransform(), /*bSweep*/false, /*OutSweepHitResult*/nullptr, ETeleportType::TeleportPhysics);
+		}
 
-			// Initialize velocity of the replicated actor to match the mass replicated one
-			if (const FMassVelocityFragment* Velocity = EntityView.GetFragmentDataPtr<FMassVelocityFragment>())
+		// Initialize velocity of the replicated actor to match the mass replicated one
+		if (const FMassVelocityFragment* Velocity = EntityView.GetFragmentDataPtr<FMassVelocityFragment>())
+		{
+			if (UCharacterMovementComponent* MovementComp = GetOwner()->FindComponentByClass<UCharacterMovementComponent>())
 			{
-				if (UCharacterMovementComponent* MovementComp = GetOwner()->FindComponentByClass<UCharacterMovementComponent>())
-				{
-					MovementComp->Velocity = Velocity->Value;
-				}
+				MovementComp->Velocity = Velocity->Value;
 			}
 		}
 	}
