Fix crash when debugging Mass entities (inverted range).

diff --git a/Engine/Plugins/Runtime/MassGameplay/Source/MassGameplayDebug/Private/MassGameplayDebugTypes.cpp b/Engine/Plugins/Runtime/MassGameplay/Source/MassGameplayDebug/Private/MassGameplayDebugTypes.cpp
--- a/Engine/Plugins/Runtime/MassGameplay/Source/MassGameplayDebug/Private/MassGameplayDebugTypes.cpp
+++ b/Engine/Plugins/Runtime/MassGameplay/Source/MassGameplayDebug/Private/MassGameplayDebugTypes.cpp
@@ -13,8 +13,8 @@ namespace UE::Mass::Debug
 
 void GetDebugEntitiesAndLocations(const FMassEntityManager& EntityManager, TArray<FMassEntityHandle>& OutEntities, TArray<FVector>& OutLocations)
 {
-	int32 DebugEntityEnd, DebugEntityBegin;
-	if (GetDebugEntitiesRange(DebugEntityEnd, DebugEntityBegin) == false)
+	int32 DebugEntityBegin, DebugEntityEnd;
+	if (GetDebugEntitiesRange(DebugEntityBegin, DebugEntityEnd) == false)
 	{
 		return;
 	}
