[MassGameplay] Modified UMassRepresentationSubsystem::GetOrSpawnActorFromTemplate to actually check spawn handle's validity (not only whether it's set) before using it to retrieve data

diff --git a/Engine/Plugins/Runtime/MassGameplay/Source/MassActors/Public/MassActorSpawnerSubsystem.h b/Engine/Plugins/Runtime/MassGameplay/Source/MassActors/Public/MassActorSpawnerSubsystem.h
--- a/Engine/Plugins/Runtime/MassGameplay/Source/MassActors/Public/MassActorSpawnerSubsystem.h
+++ b/Engine/Plugins/Runtime/MassGameplay/Source/MassActors/Public/MassActorSpawnerSubsystem.h
@@ -208,6 +208,18 @@ public:
 		return SpawnRequestHandleManager.IsValidHandle(SpawnRequestHandle);
 	}
 
+	/**
+	 * Checks whether the given handle has already been released.
+	 * Note that the function will return `false` for invalid handles - the function is intended
+	 * to check whether a given "valid" handle that the caller has access to has already been released.
+	 */
+	bool IsHandleReleased(const FMassActorSpawnRequestHandle SpawnRequestHandle) const
+	{
+		return IsSpawnRequestHandleValid(SpawnRequestHandle) == false
+			&& SpawnRequestHandleManager.GetHandles().IsValidIndex(SpawnRequestHandle.GetIndex())
+			&& SpawnRequestHandleManager.GetHandles()[SpawnRequestHandle.GetIndex()].IsValid() == false;
+	}
+
 	/**
 	 * Destroy an actor 
 	 * @param Actor to destroy
diff --git a/Engine/Plugins/Runtime/MassGameplay/Source/MassRepresentation/Private/MassRepresentationSubsystem.cpp b/Engine/Plugins/Runtime/MassGameplay/Source/MassRepresentation/Private/MassRepresentationSubsystem.cpp
--- a/Engine/Plugins/Runtime/MassGameplay/Source/MassRepresentation/Private/MassRepresentationSubsystem.cpp
+++ b/Engine/Plugins/Runtime/MassGameplay/Source/MassRepresentation/Private/MassRepresentationSubsystem.cpp
@@ -111,7 +111,10 @@ AActor* UMassRepresentationSubsystem::GetOrSpawnActorFromTemplate(const FMassEnt
 	check(ActorSpawnerSubsystem);
 	const TSubclassOf<AActor> TemplateToSpawn = TemplateActors[TemplateActorIndex].Actor;
 
-	if (InOutSpawnRequestHandle.IsValid())
+	if (InOutSpawnRequestHandle.IsValid()
+		&& ensureMsgf(ActorSpawnerSubsystem->IsSpawnRequestHandleValid(InOutSpawnRequestHandle)
+			, TEXT("Given FMassActorSpawnRequestHandle(%d,sn:%d) %s"), InOutSpawnRequestHandle.GetIndex(), InOutSpawnRequestHandle.GetSerialNumber()
+			, ActorSpawnerSubsystem->IsHandleReleased(InOutSpawnRequestHandle) ? TEXT("has already been RELEASED") : TEXT("is OUTDATED")))
 	{
 		FMassActorSpawnRequest& SpawnRequest = ActorSpawnerSubsystem->GetMutableSpawnRequest<FMassActorSpawnRequest>(InOutSpawnRequestHandle);
 		// Check if this existing spawn request is matching the template actor
diff --git a/Engine/Plugins/Runtime/MassGameplay/Source/MassRepresentation/Public/MassRepresentationSubsystem.h b/Engine/Plugins/Runtime/MassGameplay/Source/MassRepresentation/Public/MassRepresentationSubsystem.h
--- a/Engine/Plugins/Runtime/MassGameplay/Source/MassRepresentation/Public/MassRepresentationSubsystem.h
+++ b/Engine/Plugins/Runtime/MassGameplay/Source/MassRepresentation/Public/MassRepresentationSubsystem.h
@@ -103,6 +103,7 @@ public:
 	 * @param ActorPreSpawnDelegate is an optional delegate called before the spawning of an actor
 	 * @param ActorPostSpawnDelegate is an optional delegate called once the actor is spawned
 	 * @return The spawned actor from the template actor type if ready
+	 * @todo should be renamed to GetOrRequestSpawnActorFromTemplate
 	 */
 	AActor* GetOrSpawnActorFromTemplate(const FMassEntityHandle MassAgent, const FTransform& Transform, const int16 TemplateActorIndex, FMassActorSpawnRequestHandle& InOutSpawnRequestHandle, float Priority = MAX_FLT,
 		FMassActorPreSpawnDelegate ActorPreSpawnDelegate = FMassActorPreSpawnDelegate(), FMassActorPostSpawnDelegate ActorPostSpawnDelegate = FMassActorPostSpawnDelegate());
