[Mass] Deprecated use of gsl::not_null in Mass API

#rnx

diff --git a/Engine/Plugins/Runtime/InstancedActors/Source/InstancedActors/Private/InstancedActorsManager.cpp b/Engine/Plugins/Runtime/InstancedActors/Source/InstancedActors/Private/InstancedActorsManager.cpp
--- a/Engine/Plugins/Runtime/InstancedActors/Source/InstancedActors/Private/InstancedActorsManager.cpp
+++ b/Engine/Plugins/Runtime/InstancedActors/Source/InstancedActors/Private/InstancedActorsManager.cpp
@@ -1570,7 +1570,7 @@ void AInstancedActorsManager::CreateISMComponents(const FInstancedActorsVisualiz
 	}
 }
 
-void AInstancedActorsManager::SetUpEditorPreviewISMComponent(gsl::not_null<UInstancedStaticMeshComponent*> ISMComponent)
+void AInstancedActorsManager::SetUpEditorPreviewISMComponent(TNotNull<UInstancedStaticMeshComponent*> ISMComponent)
 {
 	ISMComponent->SetCollisionEnabled(ECollisionEnabled::QueryOnly);
 	ISMComponent->bSelectable = true;
diff --git a/Engine/Plugins/Runtime/InstancedActors/Source/InstancedActors/Public/InstancedActorsManager.h b/Engine/Plugins/Runtime/InstancedActors/Source/InstancedActors/Public/InstancedActorsManager.h
--- a/Engine/Plugins/Runtime/InstancedActors/Source/InstancedActors/Public/InstancedActorsManager.h
+++ b/Engine/Plugins/Runtime/InstancedActors/Source/InstancedActors/Public/InstancedActorsManager.h
@@ -301,7 +301,7 @@ protected:
 #endif
 
 	/** Configures given ISMComponent for editor-preview purposes */
-	static void SetUpEditorPreviewISMComponent(gsl::not_null<UInstancedStaticMeshComponent*> ISMComponent);
+	static void SetUpEditorPreviewISMComponent(TNotNull<UInstancedStaticMeshComponent*> ISMComponent);
 
 	UPROPERTY(Transient)
 	TObjectPtr<UInstancedActorsSubsystem> InstancedActorSubsystem;
diff --git a/Engine/Plugins/Runtime/MassGameplay/Source/MassLOD/Private/MassLODTrait.cpp b/Engine/Plugins/Runtime/MassGameplay/Source/MassLOD/Private/MassLODTrait.cpp
--- a/Engine/Plugins/Runtime/MassGameplay/Source/MassLOD/Private/MassLODTrait.cpp
+++ b/Engine/Plugins/Runtime/MassGameplay/Source/MassLOD/Private/MassLODTrait.cpp
@@ -2,7 +2,6 @@
 
 #include "MassLODTrait.h"
 #include "MassEntityTemplateRegistry.h"
-#include "Engine/World.h"
 #include "MassCommonFragments.h"
 #include "MassLODFragments.h"
 #include "MassEntityUtils.h"
@@ -12,7 +11,7 @@
 
 namespace UE::Mass::Private
 {
-	bool ValidateRequiredProcessor(gsl::not_null<const UMassEntityTraitBase*> Trait, TSubclassOf<UMassProcessor> ProcessorClass)
+	bool ValidateRequiredProcessor(const TNotNull<const UMassEntityTraitBase*> Trait, const TSubclassOf<UMassProcessor>& ProcessorClass)
 	{
 		check(ProcessorClass);
 
diff --git a/Engine/Source/Runtime/MassEntity/Private/MassProcessor.cpp b/Engine/Source/Runtime/MassEntity/Private/MassProcessor.cpp
--- a/Engine/Source/Runtime/MassEntity/Private/MassProcessor.cpp
+++ b/Engine/Source/Runtime/MassEntity/Private/MassProcessor.cpp
@@ -122,7 +122,7 @@ UMassProcessor::UMassProcessor()
 {
 }
 
-void UMassProcessor::CallInitialize(const gsl::not_null<UObject*> Owner, const TSharedRef<FMassEntityManager>& EntityManager)
+void UMassProcessor::CallInitialize(const TNotNull<UObject*> Owner, const TSharedRef<FMassEntityManager>& EntityManager)
 {
 	if (ensure(HasAnyFlags(RF_ClassDefaultObject) == false 
 		&& GetClass()->HasAnyClassFlags(CLASS_Abstract) == false))
diff --git a/Engine/Source/Runtime/MassEntity/Public/MassExecutionContext.h b/Engine/Source/Runtime/MassEntity/Public/MassExecutionContext.h
--- a/Engine/Source/Runtime/MassEntity/Public/MassExecutionContext.h
+++ b/Engine/Source/Runtime/MassEntity/Public/MassExecutionContext.h
@@ -72,7 +72,7 @@ private:
 
 	struct FQueryTransientRuntime
 	{
-		gsl::not_null<FMassEntityQuery*> Query;
+		TNotNull<FMassEntityQuery*> Query;
 		FMassExternalSubsystemBitSet ConstSubsystemsBitSet;
 		FMassExternalSubsystemBitSet MutableSubsystemsBitSet;
 	};
diff --git a/Engine/Source/Runtime/MassEntity/Public/MassProcessor.h b/Engine/Source/Runtime/MassEntity/Public/MassProcessor.h
--- a/Engine/Source/Runtime/MassEntity/Public/MassProcessor.h
+++ b/Engine/Source/Runtime/MassEntity/Public/MassProcessor.h
@@ -63,7 +63,7 @@ public:
 	bool IsInitialized() const;
 
 	/** Calls InitializeInternal and handles initialization bookkeeping. */
-	void CallInitialize(const gsl::not_null<UObject*> Owner, const TSharedRef<FMassEntityManager>& EntityManager);
+	void CallInitialize(const TNotNull<UObject*> Owner, const TSharedRef<FMassEntityManager>& EntityManager);
 	
 	virtual FGraphEventRef DispatchProcessorTasks(const TSharedPtr<FMassEntityManager>& EntityManager, FMassExecutionContext& ExecutionContext, const FGraphEventArray& Prerequisites = FGraphEventArray());
 
