Fix incorrectly declared static inline delegate OnSimulationStarted in Mass Simulation subsystem, which would create an instance of the delegate for each module because it is inline instead of having one static instance of the variable, which is what we really want.

#rnx

diff --git a/Engine/Plugins/Runtime/MassGameplay/Source/MassSimulation/Private/MassSimulationSubsystem.cpp b/Engine/Plugins/Runtime/MassGameplay/Source/MassSimulation/Private/MassSimulationSubsystem.cpp
--- a/Engine/Plugins/Runtime/MassGameplay/Source/MassSimulation/Private/MassSimulationSubsystem.cpp
+++ b/Engine/Plugins/Runtime/MassGameplay/Source/MassSimulation/Private/MassSimulationSubsystem.cpp
@@ -14,6 +14,8 @@
 
 DEFINE_LOG_CATEGORY(LogMassSim);
 
+UMassSimulationSubsystem::FOnSimulationStarted UMassSimulationSubsystem::OnSimulationStarted;
+
 namespace UE::Mass::Simulation
 {
 	bool bDoEntityCompaction = true;
diff --git a/Engine/Plugins/Runtime/MassGameplay/Source/MassSimulation/Public/MassSimulationSubsystem.h b/Engine/Plugins/Runtime/MassGameplay/Source/MassSimulation/Public/MassSimulationSubsystem.h
--- a/Engine/Plugins/Runtime/MassGameplay/Source/MassSimulation/Public/MassSimulationSubsystem.h
+++ b/Engine/Plugins/Runtime/MassGameplay/Source/MassSimulation/Public/MassSimulationSubsystem.h
@@ -27,7 +27,7 @@ public:
 
 	MASSSIMULATION_API FMassProcessingPhase::FOnPhaseEvent& GetOnProcessingPhaseStarted(const EMassProcessingPhase Phase);
 	MASSSIMULATION_API FMassProcessingPhase::FOnPhaseEvent& GetOnProcessingPhaseFinished(const EMassProcessingPhase Phase);
-	static FOnSimulationStarted& GetOnSimulationStarted() { return OnSimulationStarted; }
+	MASSSIMULATION_API static FOnSimulationStarted& GetOnSimulationStarted() { return OnSimulationStarted; }
 
 	MASSSIMULATION_API void RegisterDynamicProcessor(UMassProcessor& Processor);
 	MASSSIMULATION_API void UnregisterDynamicProcessor(UMassProcessor& Processor);
@@ -68,7 +68,7 @@ protected:
 
 	TSharedRef<FMassProcessingPhaseManager> PhaseManager;
 
-	inline static FOnSimulationStarted OnSimulationStarted={};
+	static FOnSimulationStarted OnSimulationStarted;
 
 	UPROPERTY()
 	FMassRuntimePipeline RuntimePipeline;
