[MassGameplay] Fixed USmartObjectZoneAnnotations::DebugDraw crashing when no USmartObjectSubsystem can be found.
Said crash encountered in testing project in Client PIE

#rnx

diff --git a/Engine/Plugins/Runtime/MassGameplay/Source/MassSmartObjects/Private/SmartObjectZoneAnnotations.cpp b/Engine/Plugins/Runtime/MassGameplay/Source/MassSmartObjects/Private/SmartObjectZoneAnnotations.cpp
--- a/Engine/Plugins/Runtime/MassGameplay/Source/MassSmartObjects/Private/SmartObjectZoneAnnotations.cpp
+++ b/Engine/Plugins/Runtime/MassGameplay/Source/MassSmartObjects/Private/SmartObjectZoneAnnotations.cpp
@@ -180,7 +180,7 @@ void USmartObjectZoneAnnotations::TickAnnotation(const float DeltaTime, FZoneGra
 void USmartObjectZoneAnnotations::DebugDraw(FZoneGraphAnnotationSceneProxy* DebugProxy)
 {
 	UZoneGraphSubsystem* ZoneGraph = UWorld::GetSubsystem<UZoneGraphSubsystem>(GetWorld());
-	if (ZoneGraph == nullptr)
+	if (ZoneGraph == nullptr || SmartObjectSubsystem == nullptr)
 	{
 		return;
 	}
