[MassGameplay] Changed how FMassBehaviorEntityContext interacts with FMassEntityView

Used to copy, now it moves the view instance. This change is in preparation to restrict users capability of storing FMassEntityView.

diff --git a/Engine/Plugins/Runtime/MassGameplay/Source/MassSmartObjects/Private/MassSmartObjectHandler.cpp b/Engine/Plugins/Runtime/MassGameplay/Source/MassSmartObjects/Private/MassSmartObjectHandler.cpp
--- a/Engine/Plugins/Runtime/MassGameplay/Source/MassSmartObjects/Private/MassSmartObjectHandler.cpp
+++ b/Engine/Plugins/Runtime/MassGameplay/Source/MassSmartObjects/Private/MassSmartObjectHandler.cpp
@@ -49,8 +49,7 @@ namespace UE::Mass::SmartObject
 		UMassSignalSubsystem* SignalSubsystem = Payload.SignalSubsystem.Get();
 		if (SmartObjectSubsystem != nullptr && EntityManager != nullptr && SignalSubsystem != nullptr && EntityManager->IsEntityActive(Payload.Entity))
 		{
-			const FMassEntityView EntityView(*EntityManager, Payload.Entity);
-			const FMassBehaviorEntityContext Context(EntityView, *SmartObjectSubsystem);
+			const FMassBehaviorEntityContext Context(FMassEntityView(*EntityManager, Payload.Entity), *SmartObjectSubsystem);
 
 			StopInteraction(EntityManager->Defer(), Context, ClaimHandle, EMassSmartObjectInteractionStatus::Aborted);
 			if (EntityManager->IsProcessing() == false)
@@ -208,8 +207,7 @@ bool FMassSmartObjectHandler::StartUsingSmartObject(
 	User.InteractionHandle = ClaimHandle;
 
 	// Activate behavior
-	const FMassEntityView EntityView(ExecutionContext.GetEntityManagerChecked(), Entity);
-	const FMassBehaviorEntityContext Context(EntityView, SmartObjectSubsystem);
+	const FMassBehaviorEntityContext Context(FMassEntityView(ExecutionContext.GetEntityManagerChecked(), Entity), SmartObjectSubsystem);
 	BehaviorDefinition->Activate(ExecutionContext.Defer(), Context);
 
 	return true;
@@ -265,8 +263,7 @@ void FMassSmartObjectHandler::StopUsingSmartObject(const FMassEntityHandle Entit
 	}
 #endif
 
-	const FMassEntityView EntityView(ExecutionContext.GetEntityManagerChecked(), Entity);
-	const FMassBehaviorEntityContext Context(EntityView, SmartObjectSubsystem);
+	const FMassBehaviorEntityContext Context(FMassEntityView(ExecutionContext.GetEntityManagerChecked(), Entity), SmartObjectSubsystem);
 	UE::Mass::SmartObject::StopInteraction(ExecutionContext.Defer(), Context, User.InteractionHandle, NewStatus);
 }
 
@@ -290,8 +287,7 @@ void FMassSmartObjectHandler::ReleaseSmartObject(const FMassEntityHandle Entity,
 	}
 #endif
 
-	const FMassEntityView EntityView(ExecutionContext.GetEntityManagerChecked(), Entity);
-	const FMassBehaviorEntityContext Context(EntityView, SmartObjectSubsystem);
+	const FMassBehaviorEntityContext Context(FMassEntityView(ExecutionContext.GetEntityManagerChecked(), Entity), SmartObjectSubsystem);
 
 	SmartObjectSubsystem.UnregisterSlotInvalidationCallback(ClaimHandle);
 
diff --git a/Engine/Plugins/Runtime/MassGameplay/Source/MassSmartObjects/Public/MassSmartObjectBehaviorDefinition.h b/Engine/Plugins/Runtime/MassGameplay/Source/MassSmartObjects/Public/MassSmartObjectBehaviorDefinition.h
--- a/Engine/Plugins/Runtime/MassGameplay/Source/MassSmartObjects/Public/MassSmartObjectBehaviorDefinition.h
+++ b/Engine/Plugins/Runtime/MassGameplay/Source/MassSmartObjects/Public/MassSmartObjectBehaviorDefinition.h
@@ -21,8 +21,8 @@ struct MASSSMARTOBJECTS_API FMassBehaviorEntityContext
 {
 	FMassBehaviorEntityContext() = delete;
 
-	FMassBehaviorEntityContext(const FMassEntityView InEntityView, USmartObjectSubsystem& InSubsystem)
-		: EntityView(InEntityView)
+	FMassBehaviorEntityContext(FMassEntityView&& InEntityView, USmartObjectSubsystem& InSubsystem)
+		: EntityView(MoveTemp(InEntityView))
 		, SmartObjectSubsystem(InSubsystem)
 	{}
 
