[Mass] Slightly changing some calls into FMassEntityManager::GetOrCreateSharedFragment() to ensure template argument deduction uses the proper overload. Those that are passing in a struct directly, instead of initializer arguments, should ideally call the overload that uses it directly for hashing/caching, rather than instantiating a new struct and copying

diff --git a/Engine/Plugins/AI/MassCrowd/Source/MassCrowd/Private/MassCrowdServerRepresentationTrait.cpp b/Engine/Plugins/AI/MassCrowd/Source/MassCrowd/Private/MassCrowdServerRepresentationTrait.cpp
--- a/Engine/Plugins/AI/MassCrowd/Source/MassCrowd/Private/MassCrowdServerRepresentationTrait.cpp
+++ b/Engine/Plugins/AI/MassCrowd/Source/MassCrowd/Private/MassCrowdServerRepresentationTrait.cpp
@@ -42,7 +42,7 @@ void UMassCrowdServerRepresentationTrait::BuildTemplate(FMassEntityTemplateBuild
 
 	FMassRepresentationSubsystemSharedFragment SubsystemSharedFragment;
 	SubsystemSharedFragment.RepresentationSubsystem = RepresentationSubsystem;
-	FSharedStruct SubsystemFragment = EntityManager.GetOrCreateSharedFragment<FMassRepresentationSubsystemSharedFragment>(SubsystemSharedFragment);
+	FSharedStruct SubsystemFragment = EntityManager.GetOrCreateSharedFragment(SubsystemSharedFragment);
 	BuildContext.AddSharedFragment(SubsystemFragment);
 
 	FConstSharedStruct ParamsFragment = EntityManager.GetOrCreateConstSharedFragment(Params);
diff --git a/Engine/Plugins/Runtime/InstancedActors/Source/InstancedActors/Private/InstancedActorsData.cpp b/Engine/Plugins/Runtime/InstancedActors/Source/InstancedActors/Private/InstancedActorsData.cpp
--- a/Engine/Plugins/Runtime/InstancedActors/Source/InstancedActors/Private/InstancedActorsData.cpp
+++ b/Engine/Plugins/Runtime/InstancedActors/Source/InstancedActors/Private/InstancedActorsData.cpp
@@ -519,7 +519,7 @@ void UInstancedActorsData::DespawnEntities()
 
 	FInstancedActorsDataSharedFragment ManagerFragment;
 	ManagerFragment.InstanceData = this;
-	FSharedStruct SharedFragmentInstance = MassEntityManager.GetOrCreateSharedFragment<FInstancedActorsDataSharedFragment>(ManagerFragment);
+	FSharedStruct SharedFragmentInstance = MassEntityManager.GetOrCreateSharedFragment(ManagerFragment);
 
 	FInstancedActorsDataSharedFragment* AsSharedManagerFragment = SharedFragmentInstance.GetPtr<FInstancedActorsDataSharedFragment>();
 	if (ensure(AsSharedManagerFragment) && AsSharedManagerFragment->InstanceData == this)
diff --git a/Engine/Plugins/Runtime/InstancedActors/Source/InstancedActors/Private/InstancedActorsVisualizationTrait.cpp b/Engine/Plugins/Runtime/InstancedActors/Source/InstancedActors/Private/InstancedActorsVisualizationTrait.cpp
--- a/Engine/Plugins/Runtime/InstancedActors/Source/InstancedActors/Private/InstancedActorsVisualizationTrait.cpp
+++ b/Engine/Plugins/Runtime/InstancedActors/Source/InstancedActors/Private/InstancedActorsVisualizationTrait.cpp
@@ -93,7 +93,7 @@ void UInstancedActorsVisualizationTrait::BuildTemplate(FMassEntityTemplateBuildC
 
 	FInstancedActorsDataSharedFragment ManagerSharedFragment;
 	ManagerSharedFragment.InstanceData = InstanceData;
-	FSharedStruct SubsystemFragment = EntityManager.GetOrCreateSharedFragment<FInstancedActorsDataSharedFragment>(ManagerSharedFragment);
+	FSharedStruct SubsystemFragment = EntityManager.GetOrCreateSharedFragment(ManagerSharedFragment);
 
 	FInstancedActorsDataSharedFragment* AsShared = SubsystemFragment.GetPtr<FInstancedActorsDataSharedFragment>();
 	if (ensure(AsShared))
diff --git a/Engine/Plugins/Runtime/MassGameplay/Source/MassRepresentation/Private/MassDistanceVisualizationTrait.cpp b/Engine/Plugins/Runtime/MassGameplay/Source/MassRepresentation/Private/MassDistanceVisualizationTrait.cpp
--- a/Engine/Plugins/Runtime/MassGameplay/Source/MassRepresentation/Private/MassDistanceVisualizationTrait.cpp
+++ b/Engine/Plugins/Runtime/MassGameplay/Source/MassRepresentation/Private/MassDistanceVisualizationTrait.cpp
@@ -59,7 +59,7 @@ void UMassDistanceVisualizationTrait::BuildTemplate(FMassEntityTemplateBuildCont
 
 	FMassRepresentationSubsystemSharedFragment SubsystemSharedFragment;
 	SubsystemSharedFragment.RepresentationSubsystem = RepresentationSubsystem;
-	FSharedStruct SubsystemFragment = EntityManager.GetOrCreateSharedFragment<FMassRepresentationSubsystemSharedFragment>(SubsystemSharedFragment);
+	FSharedStruct SubsystemFragment = EntityManager.GetOrCreateSharedFragment(SubsystemSharedFragment);
 	BuildContext.AddSharedFragment(SubsystemFragment);
 
 	if (!Params.RepresentationActorManagementClass && !BuildContext.IsInspectingData())
diff --git a/Engine/Plugins/Runtime/MassGameplay/Source/MassRepresentation/Private/MassVisualizationTrait.cpp b/Engine/Plugins/Runtime/MassGameplay/Source/MassRepresentation/Private/MassVisualizationTrait.cpp
--- a/Engine/Plugins/Runtime/MassGameplay/Source/MassRepresentation/Private/MassVisualizationTrait.cpp
+++ b/Engine/Plugins/Runtime/MassGameplay/Source/MassRepresentation/Private/MassVisualizationTrait.cpp
@@ -78,7 +78,7 @@ void UMassVisualizationTrait::BuildTemplate(FMassEntityTemplateBuildContext& Bui
 
 	FMassRepresentationSubsystemSharedFragment SubsystemSharedFragment;
 	SubsystemSharedFragment.RepresentationSubsystem = RepresentationSubsystem;
-	FSharedStruct SubsystemFragment = EntityManager.GetOrCreateSharedFragment<FMassRepresentationSubsystemSharedFragment>(SubsystemSharedFragment);
+	FSharedStruct SubsystemFragment = EntityManager.GetOrCreateSharedFragment(SubsystemSharedFragment);
 	BuildContext.AddSharedFragment(SubsystemFragment);
 
 	if (!Params.RepresentationActorManagementClass)
