[Mass] Added a mutable PhaseManager getter to UMassSimulationSubsystem

#rnx

diff --git a/Engine/Plugins/Runtime/MassGameplay/Source/MassSimulation/Public/MassSimulationSubsystem.h b/Engine/Plugins/Runtime/MassGameplay/Source/MassSimulation/Public/MassSimulationSubsystem.h
--- a/Engine/Plugins/Runtime/MassGameplay/Source/MassSimulation/Public/MassSimulationSubsystem.h
+++ b/Engine/Plugins/Runtime/MassGameplay/Source/MassSimulation/Public/MassSimulationSubsystem.h
@@ -23,7 +23,8 @@ public:
 	
 	MASSSIMULATION_API UMassSimulationSubsystem(const FObjectInitializer& ObjectInitializer = FObjectInitializer::Get());
 
-	const FMassProcessingPhaseManager& GetPhaseManager() const { return PhaseManager.Get(); }
+	const FMassProcessingPhaseManager& GetPhaseManager() const;
+	FMassProcessingPhaseManager& GetMutablePhaseManager();
 
 	MASSSIMULATION_API FMassProcessingPhase::FOnPhaseEvent& GetOnProcessingPhaseStarted(const EMassProcessingPhase Phase);
 	MASSSIMULATION_API FMassProcessingPhase::FOnPhaseEvent& GetOnProcessingPhaseFinished(const EMassProcessingPhase Phase);
@@ -32,7 +33,7 @@ public:
 	MASSSIMULATION_API void RegisterDynamicProcessor(UMassProcessor& Processor);
 	MASSSIMULATION_API void UnregisterDynamicProcessor(UMassProcessor& Processor);
 
-	bool IsSimulationStarted() const { return bSimulationStarted; }
+	bool IsSimulationStarted() const;
 
 	/** @return whether hosted EntityManager is currently, actively being used for processing purposes. Equivalent to calling FMassEntityManager.IsProcessing() */
 	MASSSIMULATION_API bool IsDuringMassProcessing() const;
@@ -84,3 +85,21 @@ protected:
 	FDelegateHandle MassEntitySettingsChangeHandle;
 #endif // WITH_EDITOR
 };
+
+//-----------------------------------------------------------------------------
+// INLINES
+//-----------------------------------------------------------------------------
+inline const FMassProcessingPhaseManager& UMassSimulationSubsystem::GetPhaseManager() const
+{
+	return PhaseManager.Get();
+}
+
+inline FMassProcessingPhaseManager& UMassSimulationSubsystem::GetMutablePhaseManager()
+{
+	return PhaseManager.Get();
+}
+
+inline bool UMassSimulationSubsystem::IsSimulationStarted() const
+{
+	return bSimulationStarted;
+}
