[MassGameplay] Fixed UMassSimulationSubsystem's editor-time event handles still being active after the subsystem has been Deinitialize-ed.
The bug resulted in UMassSimulationSubsystem::StartSimulation being called upon PIE end even if the given subsystem instance has already been deinitialized (for example due to the editor world's reinitialization, like after geometry rebuilding).

diff --git a/Engine/Plugins/Runtime/MassGameplay/Source/MassSimulation/Private/MassSimulationSubsystem.cpp b/Engine/Plugins/Runtime/MassGameplay/Source/MassSimulation/Private/MassSimulationSubsystem.cpp
--- a/Engine/Plugins/Runtime/MassGameplay/Source/MassSimulation/Private/MassSimulationSubsystem.cpp
+++ b/Engine/Plugins/Runtime/MassGameplay/Source/MassSimulation/Private/MassSimulationSubsystem.cpp
@@ -3,7 +3,6 @@
 #include "MassSimulationSubsystem.h"
 #include "MassEntityManager.h"
 #include "MassEntitySubsystem.h"
-#include "MassExecutor.h"
 #include "MassSimulationSettings.h"
 #include "VisualLogger/VisualLogger.h"
 #include "MassEntitySettings.h"
@@ -41,23 +40,7 @@ UMassSimulationSubsystem::UMassSimulationSubsystem(const FObjectInitializer&)
 
 void UMassSimulationSubsystem::BeginDestroy()
 {
-#if WITH_EDITOR
-	if (PieBeginEventHandle.IsValid())
-	{
-		FEditorDelegates::BeginPIE.Remove(PieBeginEventHandle);
-	}
-	if (PieEndedEventHandle.IsValid())
-	{
-		FEditorDelegates::PrePIEEnded.Remove(PieEndedEventHandle);
-	}
-	if (MassEntitySettingsChangeHandle.IsValid())
-	{
-		if (UMassEntitySettings* Settings = GetMutableDefault<UMassEntitySettings>())
-		{
-			Settings->GetOnSettingsChange().Remove(MassEntitySettingsChangeHandle);
-		}
-	}
-#endif //  WITH_EDITOR
+	ReleaseEventHandles();
 
 	if (bSimulationStarted)
 	{
@@ -109,6 +92,8 @@ void UMassSimulationSubsystem::Deinitialize()
 	PhaseManager->Deinitialize();
 	EntityManager.Reset();
 
+	ReleaseEventHandles();
+
 	Super::Deinitialize();
 }
 
@@ -285,4 +270,28 @@ void UMassSimulationSubsystem::OnMassEntitySettingsChange(const FPropertyChanged
 {
 	RebuildTickPipeline();
 }
-#endif // WITH_EDITOR
\ No newline at end of file
+#endif // WITH_EDITOR
+
+void UMassSimulationSubsystem::ReleaseEventHandles()
+{
+#if WITH_EDITOR
+	if (PieBeginEventHandle.IsValid())
+	{
+		FEditorDelegates::BeginPIE.Remove(PieBeginEventHandle);
+		PieBeginEventHandle.Reset();
+	}
+	if (PieEndedEventHandle.IsValid())
+	{
+		FEditorDelegates::PrePIEEnded.Remove(PieEndedEventHandle);
+		PieEndedEventHandle.Reset();
+	}
+	if (MassEntitySettingsChangeHandle.IsValid())
+	{
+		if (UMassEntitySettings* Settings = GetMutableDefault<UMassEntitySettings>())
+		{
+			Settings->GetOnSettingsChange().Remove(MassEntitySettingsChangeHandle);
+		}
+		MassEntitySettingsChangeHandle.Reset();
+	}
+#endif // WITH_EDITOR
+}
diff --git a/Engine/Plugins/Runtime/MassGameplay/Source/MassSimulation/Public/MassSimulationSubsystem.h b/Engine/Plugins/Runtime/MassGameplay/Source/MassSimulation/Public/MassSimulationSubsystem.h
--- a/Engine/Plugins/Runtime/MassGameplay/Source/MassSimulation/Public/MassSimulationSubsystem.h
+++ b/Engine/Plugins/Runtime/MassGameplay/Source/MassSimulation/Public/MassSimulationSubsystem.h
@@ -63,6 +63,12 @@ protected:
 	MASSSIMULATION_API void OnMassEntitySettingsChange(const FPropertyChangedEvent& PropertyChangedEvent);
 #endif // WITH_EDITOR
 
+	/**
+	 * Called when it's time to clean up all the delegate handles.
+	 * Override to add more clean up, but make sure to call the super-implementation.
+	 */
+	MASSSIMULATION_API virtual void ReleaseEventHandles();
+
 protected:
 
 	TSharedPtr<FMassEntityManager> EntityManager;
