[MassGameplay] Improvement to UMassSimulationSubsystem's clearing of event-handlers, now the users won't need to remember to call super implementation in their own overrides.

diff --git a/Engine/Plugins/Runtime/MassGameplay/Source/MassSimulation/Private/MassSimulationSubsystem.cpp b/Engine/Plugins/Runtime/MassGameplay/Source/MassSimulation/Private/MassSimulationSubsystem.cpp
--- a/Engine/Plugins/Runtime/MassGameplay/Source/MassSimulation/Private/MassSimulationSubsystem.cpp
+++ b/Engine/Plugins/Runtime/MassGameplay/Source/MassSimulation/Private/MassSimulationSubsystem.cpp
@@ -294,4 +294,6 @@ void UMassSimulationSubsystem::ReleaseEventHandles()
 		MassEntitySettingsChangeHandle.Reset();
 	}
 #endif // WITH_EDITOR
+
+	OnReleasingEventHandles();
 }
diff --git a/Engine/Plugins/Runtime/MassGameplay/Source/MassSimulation/Public/MassSimulationSubsystem.h b/Engine/Plugins/Runtime/MassGameplay/Source/MassSimulation/Public/MassSimulationSubsystem.h
--- a/Engine/Plugins/Runtime/MassGameplay/Source/MassSimulation/Public/MassSimulationSubsystem.h
+++ b/Engine/Plugins/Runtime/MassGameplay/Source/MassSimulation/Public/MassSimulationSubsystem.h
@@ -63,11 +63,11 @@ protected:
 	MASSSIMULATION_API void OnMassEntitySettingsChange(const FPropertyChangedEvent& PropertyChangedEvent);
 #endif // WITH_EDITOR
 
-	/**
-	 * Called when it's time to clean up all the delegate handles.
-	 * Override to add more clean up, but make sure to call the super-implementation.
-	 */
-	MASSSIMULATION_API virtual void ReleaseEventHandles();
+	/** Called when it's time to clean up all the delegate handles. Override OnReleasingEventHandles to add more handles cleanup */
+	MASSSIMULATION_API void ReleaseEventHandles();
+	
+	/** Override to add more handle cleanup. Will get automatically called by ReleaseEventHandles */
+	MASSSIMULATION_API virtual void OnReleasingEventHandles() {}
 
 protected:
 
