[SmartObjectZoneAnnotation] added custom FActorComponentInstanceData to store data that is considered modified by the UCS and not generically saved during RerunConstructionScripts (i.e., SmartObjectAnnotationDataArray)

diff --git a/Engine/Plugins/Runtime/MassGameplay/Source/MassSmartObjects/Private/SmartObjectZoneAnnotations.cpp b/Engine/Plugins/Runtime/MassGameplay/Source/MassSmartObjects/Private/SmartObjectZoneAnnotations.cpp
--- a/Engine/Plugins/Runtime/MassGameplay/Source/MassSmartObjects/Private/SmartObjectZoneAnnotations.cpp
+++ b/Engine/Plugins/Runtime/MassGameplay/Source/MassSmartObjects/Private/SmartObjectZoneAnnotations.cpp
@@ -5,7 +5,6 @@
 #include "SmartObjectComponent.h"
 #include "SmartObjectSubsystem.h"
 #include "ZoneGraphAnnotationSubsystem.h"
-#include "ZoneGraphAnnotationTypes.h"
 #include "ZoneGraphDelegates.h"
 #include "ZoneGraphQuery.h"
 #include "ZoneGraphSubsystem.h"
@@ -106,7 +105,7 @@ void USmartObjectZoneAnnotations::PreZoneGraphDataRemoved(const AZoneGraphData&
 
 	FSmartObjectAnnotationData& Data = SmartObjectAnnotationDataArray[Index];
 
-	// We use precomputed data for runtime so we only mark it as not longer used
+	// We use precomputed data for runtime so we only mark it as no longer used
 	if (ZoneGraphData.GetWorld()->IsGameWorld())
 	{
 		Data.DataHandle = {};
@@ -147,6 +146,12 @@ TOptional<FSmartObjectLaneLocation> USmartObjectZoneAnnotations::GetSmartObjectL
 	return SmartObjectLaneLocation;
 }
 
+void USmartObjectZoneAnnotations::ApplyComponentInstanceData(FSmartObjectZoneAnnotationsInstanceData* InstanceData)
+{
+	check(InstanceData);
+	SmartObjectAnnotationDataArray = InstanceData->SmartObjectAnnotations;
+}
+
 void USmartObjectZoneAnnotations::TickAnnotation(const float DeltaTime, FZoneGraphAnnotationTagContainer& BehaviorTagContainer)
 {
 	if (!BehaviorTag.IsValid())
@@ -176,6 +181,11 @@ void USmartObjectZoneAnnotations::TickAnnotation(const float DeltaTime, FZoneGra
 #endif // UE_ENABLE_DEBUG_DRAWING
 }
 
+TStructOnScope<FActorComponentInstanceData> USmartObjectZoneAnnotations::GetComponentInstanceData() const
+{
+	return MakeStructOnScope<FActorComponentInstanceData, FSmartObjectZoneAnnotationsInstanceData>(this);
+}
+
 #if UE_ENABLE_DEBUG_DRAWING
 void USmartObjectZoneAnnotations::DebugDraw(FZoneGraphAnnotationSceneProxy* DebugProxy)
 {
diff --git a/Engine/Plugins/Runtime/MassGameplay/Source/MassSmartObjects/Public/SmartObjectZoneAnnotations.h b/Engine/Plugins/Runtime/MassGameplay/Source/MassSmartObjects/Public/SmartObjectZoneAnnotations.h
--- a/Engine/Plugins/Runtime/MassGameplay/Source/MassSmartObjects/Public/SmartObjectZoneAnnotations.h
+++ b/Engine/Plugins/Runtime/MassGameplay/Source/MassSmartObjects/Public/SmartObjectZoneAnnotations.h
@@ -7,6 +7,7 @@
 #include "ZoneGraphTypes.h"
 #include "SmartObjectZoneAnnotations.generated.h"
 
+struct FSmartObjectZoneAnnotationsInstanceData;
 class AZoneGraphData;
 class USmartObjectSubsystem;
 
@@ -96,12 +97,20 @@ class MASSSMARTOBJECTS_API USmartObjectZoneAnnotations : public UZoneGraphAnnota
 public:
 	const FSmartObjectAnnotationData* GetAnnotationData(FZoneGraphDataHandle DataHandle) const;
 	TOptional<FSmartObjectLaneLocation> GetSmartObjectLaneLocation(const FZoneGraphDataHandle DataHandle, const FSmartObjectHandle SmartObjectHandle) const;
+	void ApplyComponentInstanceData(FSmartObjectZoneAnnotationsInstanceData* InstanceData);
+
+	TConstArrayView<FSmartObjectAnnotationData> GetSmartObjectAnnotations() const
+	{
+		return SmartObjectAnnotationDataArray;
+	}
 
 protected:
 	virtual void PostSubsystemsInitialized() override;
 	virtual FZoneGraphTagMask GetAnnotationTags() const override;
 	virtual void TickAnnotation(const float DeltaTime, FZoneGraphAnnotationTagContainer& BehaviorTagContainer) override;
 
+	virtual TStructOnScope<FActorComponentInstanceData> GetComponentInstanceData() const override;
+
 	virtual void PostZoneGraphDataAdded(const AZoneGraphData& ZoneGraphData) override;
 	virtual void PreZoneGraphDataRemoved(const AZoneGraphData& ZoneGraphData) override;
 
@@ -143,3 +152,37 @@ protected:
 	UPROPERTY(Transient)
 	TObjectPtr<USmartObjectSubsystem> SmartObjectSubsystem = nullptr;
 };
+
+/** Used to store data that is considered modified by the UCS and not generically saved during RerunConstructionScripts */
+USTRUCT()
+struct FSmartObjectZoneAnnotationsInstanceData : public FActorComponentInstanceData
+{
+	GENERATED_BODY()
+
+public:
+	FSmartObjectZoneAnnotationsInstanceData() = default;
+	FSmartObjectZoneAnnotationsInstanceData(const USmartObjectZoneAnnotations* SourceComponent)
+		: FActorComponentInstanceData(SourceComponent)
+		, SmartObjectAnnotations(SourceComponent->GetSmartObjectAnnotations())
+	{}
+
+	virtual ~FSmartObjectZoneAnnotationsInstanceData() = default;
+
+	virtual bool ContainsData() const override
+	{
+		return SmartObjectAnnotations.Num() || Super::ContainsData();
+	}
+
+	virtual void ApplyToComponent(UActorComponent* Component, const ECacheApplyPhase CacheApplyPhase) override
+	{
+		Super::ApplyToComponent(Component, CacheApplyPhase);
+
+		if (CacheApplyPhase == ECacheApplyPhase::PostUserConstructionScript)
+		{
+			CastChecked<USmartObjectZoneAnnotations>(Component)->ApplyComponentInstanceData(this);
+		}
+	}
+
+	UPROPERTY()
+	TArray<FSmartObjectAnnotationData> SmartObjectAnnotations;
+};
