[MassGameplay] Made entity-spawning functions of MassSpawnerSubsystem return the creation context so that the call-sites have some control over when the accumulated commands and observers get executed.

This new functionality is utilized in `UMassAgentSubsystem::HandlePendingInitialization` where entity views are being created with an assumption that the entities being processed had not changed their archetypes just aftert spawning. This assumption can be invalidated with entity creation observers change entity's composition with commands that get flushes as soon as the creation context gets released.

diff --git a/Engine/Plugins/Runtime/MassGameplay/Source/MassActors/Private/MassAgentSubsystem.cpp b/Engine/Plugins/Runtime/MassGameplay/Source/MassActors/Private/MassAgentSubsystem.cpp
--- a/Engine/Plugins/Runtime/MassGameplay/Source/MassActors/Private/MassAgentSubsystem.cpp
+++ b/Engine/Plugins/Runtime/MassGameplay/Source/MassActors/Private/MassAgentSubsystem.cpp
@@ -16,7 +16,6 @@
 #include "Engine/NetDriver.h"
 #include "MassEntityUtils.h"
 
-
 namespace FMassAgentSubsystemHelper
 {
 
@@ -305,40 +304,45 @@ void UMassAgentSubsystem::HandlePendingInitialization()
 	check(SimulationSystem);
 	check(SimulationSystem->IsDuringMassProcessing() == false);
 
-	for (TTuple<FMassEntityTemplateID, FMassAgentInitializationQueue>& Data : PendingAgentEntities)
 	{
-		const FMassEntityTemplateID EntityTemplateID = Data.Get<0>();
-		const TSharedRef<FMassEntityTemplate>* EntityTemplatePtr = SpawnerSystem->GetTemplateRegistryInstance().FindTemplateFromTemplateID(EntityTemplateID);
-		check(EntityTemplatePtr);
-		const FMassEntityTemplate& EntityTemplate = EntityTemplatePtr->Get();
+		// this block is scoped so that CreationContext gets released by the end
+		TSharedPtr<FMassEntityManager::FEntityCreationContext> CreationContext;
+
+		for (TTuple<FMassEntityTemplateID, FMassAgentInitializationQueue>& Data : PendingAgentEntities)
+		{
+			const FMassEntityTemplateID EntityTemplateID = Data.Get<0>();
+			const TSharedRef<FMassEntityTemplate>* EntityTemplatePtr = SpawnerSystem->GetTemplateRegistryInstance().FindTemplateFromTemplateID(EntityTemplateID);
+			check(EntityTemplatePtr);
+			const FMassEntityTemplate& EntityTemplate = EntityTemplatePtr->Get();
 		
-		TArray<TObjectPtr<UMassAgentComponent>>& AgentComponents = Data.Get<1>().AgentComponents;
-		const int32 NewEntityCount = AgentComponents.Num();
+			TArray<TObjectPtr<UMassAgentComponent>>& AgentComponents = Data.Get<1>().AgentComponents;
+			const int32 NewEntityCount = AgentComponents.Num();
 		
-		if (NewEntityCount <= 0)
-		{
-			// this case is perfectly fine, all agents registered and unregistered before we processed the queue
-			continue;
-		}
+			if (NewEntityCount <= 0)
+			{
+				// this case is perfectly fine, all agents registered and unregistered before we processed the queue
+				continue;
+			}
 		
-		TArray<FMassEntityHandle> Entities;
-		SpawnerSystem->SpawnEntities(EntityTemplate, NewEntityCount, Entities);
-		check(Entities.Num() == NewEntityCount);
+			TArray<FMassEntityHandle> Entities;
+			CreationContext = SpawnerSystem->SpawnEntities(EntityTemplate, NewEntityCount, Entities);
+			check(Entities.Num() == NewEntityCount);
 
-		if (EntityTemplate.GetObjectFragmentInitializers().Num())
-		{
-			const TConstArrayView<FMassEntityTemplateData::FObjectFragmentInitializerFunction> ObjectFragmentInitializers = EntityTemplate.GetObjectFragmentInitializers();
-
-			for (int AgentIndex = 0; AgentIndex < Entities.Num(); ++AgentIndex)
+			if (EntityTemplate.GetObjectFragmentInitializers().Num())
 			{
-				FMassEntityView EntityView(EntityTemplate.GetArchetype(), Entities[AgentIndex]);
-				FMassAgentSubsystemHelper::InitializeAgentComponentFragments(*AgentComponents[AgentIndex], EntityView, EMassTranslationDirection::ActorToMass, ObjectFragmentInitializers);
+				const TConstArrayView<FMassEntityTemplateData::FObjectFragmentInitializerFunction> ObjectFragmentInitializers = EntityTemplate.GetObjectFragmentInitializers();
+
+				for (int AgentIndex = 0; AgentIndex < Entities.Num(); ++AgentIndex)
+				{
+					FMassEntityView EntityView(EntityTemplate.GetArchetype(), Entities[AgentIndex]);
+					FMassAgentSubsystemHelper::InitializeAgentComponentFragments(*AgentComponents[AgentIndex], EntityView, EMassTranslationDirection::ActorToMass, ObjectFragmentInitializers);
+				}
 			}
-		}
 
-		for (int AgentIndex = 0; AgentIndex < Entities.Num(); ++AgentIndex)
-		{		
-			AgentComponents[AgentIndex]->SetEntityHandle(Entities[AgentIndex]);
+			for (int AgentIndex = 0; AgentIndex < Entities.Num(); ++AgentIndex)
+			{		
+				AgentComponents[AgentIndex]->SetEntityHandle(Entities[AgentIndex]);
+			}
 		}
 	}
 
diff --git a/Engine/Plugins/Runtime/MassGameplay/Source/MassSpawner/Private/MassSpawnerSubsystem.cpp b/Engine/Plugins/Runtime/MassGameplay/Source/MassSpawner/Private/MassSpawnerSubsystem.cpp
--- a/Engine/Plugins/Runtime/MassGameplay/Source/MassSpawner/Private/MassSpawnerSubsystem.cpp
+++ b/Engine/Plugins/Runtime/MassGameplay/Source/MassSpawner/Private/MassSpawnerSubsystem.cpp
@@ -47,7 +47,7 @@ void UMassSpawnerSubsystem::Deinitialize()
 	Super::Deinitialize();
 }
 
-void UMassSpawnerSubsystem::SpawnEntities(const FMassEntityTemplate& EntityTemplate, const uint32 NumberToSpawn, TArray<FMassEntityHandle>& OutEntities)
+TSharedPtr<FMassEntityManager::FEntityCreationContext> UMassSpawnerSubsystem::SpawnEntities(const FMassEntityTemplate& EntityTemplate, const uint32 NumberToSpawn, TArray<FMassEntityHandle>& OutEntities)
 {
 	check(EntityManager);
 	check(EntityTemplate.IsValid());
@@ -55,20 +55,20 @@ void UMassSpawnerSubsystem::SpawnEntities(const FMassEntityTemplate& EntityTempl
 	if (NumberToSpawn == 0)
 	{
 		UE_VLOG(this, LogMassSpawner, Warning, TEXT("Trying to spawn 0 entities. This would cause inefficiency. Bailing out with result FALSE."));
-		return;
+		return {};
 	}
 
-	DoSpawning(EntityTemplate, NumberToSpawn, FStructView(), TSubclassOf<UMassProcessor>(), OutEntities);
+	return DoSpawning(EntityTemplate, NumberToSpawn, FStructView(), TSubclassOf<UMassProcessor>(), OutEntities);
 }
 
-void UMassSpawnerSubsystem::SpawnEntities(FMassEntityTemplateID TemplateID, const uint32 NumberToSpawn, FConstStructView SpawnData, TSubclassOf<UMassProcessor> InitializerClass, TArray<FMassEntityHandle>& OutEntities)
+TSharedPtr<FMassEntityManager::FEntityCreationContext> UMassSpawnerSubsystem::SpawnEntities(FMassEntityTemplateID TemplateID, const uint32 NumberToSpawn, FConstStructView SpawnData, TSubclassOf<UMassProcessor> InitializerClass, TArray<FMassEntityHandle>& OutEntities)
 {
 	check(TemplateID.IsValid());
 
 	const TSharedRef<FMassEntityTemplate>* EntityTemplate = TemplateRegistryInstance.FindTemplateFromTemplateID(TemplateID);
 	checkf(EntityTemplate, TEXT("SpawnEntities: TemplateID must have been registered!"));
 
-	DoSpawning(EntityTemplate->Get(), NumberToSpawn, SpawnData, InitializerClass, OutEntities);
+	return DoSpawning(EntityTemplate->Get(), NumberToSpawn, SpawnData, InitializerClass, OutEntities);
 }
 
 void UMassSpawnerSubsystem::DestroyEntities(TConstArrayView<FMassEntityHandle> Entities)
@@ -112,7 +112,7 @@ UMassProcessor* UMassSpawnerSubsystem::GetSpawnDataInitializer(TSubclassOf<UMass
 	return *Initializer;
 }
 
-void UMassSpawnerSubsystem::DoSpawning(const FMassEntityTemplate& EntityTemplate, const int32 NumToSpawn, FConstStructView SpawnData, TSubclassOf<UMassProcessor> InitializerClass, TArray<FMassEntityHandle>& OutEntities)
+TSharedPtr<FMassEntityManager::FEntityCreationContext> UMassSpawnerSubsystem::DoSpawning(const FMassEntityTemplate& EntityTemplate, const int32 NumToSpawn, FConstStructView SpawnData, TSubclassOf<UMassProcessor> InitializerClass, TArray<FMassEntityHandle>& OutEntities)
 {
 	check(EntityManager);
 	check(EntityTemplate.GetArchetype().IsValid());
@@ -121,7 +121,7 @@ void UMassSpawnerSubsystem::DoSpawning(const FMassEntityTemplate& EntityTemplate
 	if (NumToSpawn <= 0)
 	{
 		UE_VLOG(this, LogMassSpawner, Warning, TEXT("%s: Trying to spawn %d entities. Ignoring."), ANSI_TO_TCHAR(__FUNCTION__), NumToSpawn);
-		return;
+		return {};
 	}
 
 	LLM_SCOPE_BYNAME(TEXT("Mass/Spawner"))
@@ -151,7 +151,11 @@ void UMassSpawnerSubsystem::DoSpawning(const FMassEntityTemplate& EntityTemplate
 	}
 
 	OutEntities.Append(MoveTemp(SpawnedEntities));
+
 	// 4. "OnEntitiesCreated" notifies will be sent out once the CreationContext gets destroyed (via its destructor).
+	// The caller can postpone this moment keeping the returned CreationContext alive as long as needed.
+
+	return CreationContext;
 }
 
 const FMassEntityTemplate* UMassSpawnerSubsystem::GetMassEntityTemplate(FMassEntityTemplateID TemplateID) const
diff --git a/Engine/Plugins/Runtime/MassGameplay/Source/MassSpawner/Public/MassSpawnerSubsystem.h b/Engine/Plugins/Runtime/MassGameplay/Source/MassSpawner/Public/MassSpawnerSubsystem.h
--- a/Engine/Plugins/Runtime/MassGameplay/Source/MassSpawner/Public/MassSpawnerSubsystem.h
+++ b/Engine/Plugins/Runtime/MassGameplay/Source/MassSpawner/Public/MassSpawnerSubsystem.h
@@ -11,6 +11,7 @@
 #include "MassSpawnerSubsystem.generated.h"
 
 struct FMassEntityManager;
+struct FEntityCreationContext;
 struct FMassEntityTemplate;
 struct FInstancedStruct;
 struct FStructView;
@@ -30,10 +31,11 @@ public:
 	 *  @param EntityTemplate template to use for spawning entities
 	 *  @param NumberToSpawn number of entities to spawn
 	 *  @param OutEntities where the IDs of created entities get added. Note that the contents of OutEntities get overridden by the function.
-	 *  @return true if spawning was successful, false otherwise. In case of failure see logs for more details. */
-	void SpawnEntities(const FMassEntityTemplate& EntityTemplate, const uint32 NumberToSpawn, TArray<FMassEntityHandle>& OutEntities);
+	 *  @return shared pointer to the entity creation context, that, one released, will cause all the accumulated observers and commands executed
+	 */ 
+	TSharedPtr<FMassEntityManager::FEntityCreationContext> SpawnEntities(const FMassEntityTemplate& EntityTemplate, const uint32 NumberToSpawn, TArray<FMassEntityHandle>& OutEntities);
 
-	void SpawnEntities(FMassEntityTemplateID TemplateID, const uint32 NumberToSpawn, FConstStructView SpawnData, TSubclassOf<UMassProcessor> InitializerClass, TArray<FMassEntityHandle>& OutEntities);
+	TSharedPtr<FMassEntityManager::FEntityCreationContext> SpawnEntities(FMassEntityTemplateID TemplateID, const uint32 NumberToSpawn, FConstStructView SpawnData, TSubclassOf<UMassProcessor> InitializerClass, TArray<FMassEntityHandle>& OutEntities);
 
 	void DestroyEntities(TConstArrayView<FMassEntityHandle> Entities);
 
@@ -54,7 +56,7 @@ protected:
 	virtual void Deinitialize() override;
 	// UWorldSubsystem END
 
-	void DoSpawning(const FMassEntityTemplate& EntityTemplate, const int32 NumToSpawn, FConstStructView SpawnData, TSubclassOf<UMassProcessor> InitializerClass, TArray<FMassEntityHandle>& OutEntities);
+	TSharedPtr<FMassEntityManager::FEntityCreationContext> DoSpawning(const FMassEntityTemplate& EntityTemplate, const int32 NumToSpawn, FConstStructView SpawnData, TSubclassOf<UMassProcessor> InitializerClass, TArray<FMassEntityHandle>& OutEntities);
 
 	UMassProcessor* GetSpawnDataInitializer(TSubclassOf<UMassProcessor> InitializerClass);
 
