[Mass Navigation] Addition of LogMassMoverDivergence. It allows to track divergences between the mass entity requested velocity and the mover result (if any).

diff --git a/Engine/Plugins/Runtime/MassGameplay/Source/MassMovement/Public/MassMovementFragments.h b/Engine/Plugins/Runtime/MassGameplay/Source/MassMovement/Public/MassMovementFragments.h
--- a/Engine/Plugins/Runtime/MassGameplay/Source/MassMovement/Public/MassMovementFragments.h
+++ b/Engine/Plugins/Runtime/MassGameplay/Source/MassMovement/Public/MassMovementFragments.h
@@ -4,8 +4,8 @@
 
 #if UE_ENABLE_INCLUDE_ORDER_DEPRECATED_IN_5_6
 #include "CoreMinimal.h"
-#include "MassCommonTypes.h"
 #endif // UE_ENABLE_INCLUDE_ORDER_DEPRECATED_IN_5_6
+#include "MassCommonTypes.h"
 #include "MassEntityTypes.h"
 #include "MassMovementTypes.h"
 #include "RandomSequence.h"
@@ -17,6 +17,10 @@ struct MASSMOVEMENT_API FMassVelocityFragment : public FMassFragment
 	GENERATED_BODY()
 
 	FVector Value = FVector::ZeroVector;
+
+#if WITH_MASSGAMEPLAY_DEBUG
+	FVector DebugPreviousValue = FVector::ZeroVector;
+#endif	
 };
 
 USTRUCT()
