PR #12696: Add ability to Pause and Resume a Mass Simulation

diff --git a/Engine/Plugins/Runtime/MassGameplay/Source/MassSimulation/Private/MassSimulationSubsystem.cpp b/Engine/Plugins/Runtime/MassGameplay/Source/MassSimulation/Private/MassSimulationSubsystem.cpp
--- a/Engine/Plugins/Runtime/MassGameplay/Source/MassSimulation/Private/MassSimulationSubsystem.cpp
+++ b/Engine/Plugins/Runtime/MassGameplay/Source/MassSimulation/Private/MassSimulationSubsystem.cpp
@@ -228,6 +228,34 @@ void UMassSimulationSubsystem::StopSimulation()
 	bSimulationStarted = false;
 }
 
+void UMassSimulationSubsystem::PauseSimulation()
+{
+	if (bSimulationPaused == true)
+	{
+		return;
+	}
+
+	bSimulationPaused = true;
+
+	PhaseManager->Pause();
+
+	OnSimulationPaused.Broadcast(this);
+}
+
+void UMassSimulationSubsystem::ResumeSimulation()
+{
+	if (bSimulationPaused == false)
+	{
+		return;
+	}
+
+	bSimulationPaused = false;
+
+	PhaseManager->Resume();
+
+	OnSimulationResumed.Broadcast(this);
+}
+
 void UMassSimulationSubsystem::OnProcessingPhaseStarted(const float DeltaSeconds, const EMassProcessingPhase Phase) const
 {
 	switch (Phase)
diff --git a/Engine/Plugins/Runtime/MassGameplay/Source/MassSimulation/Public/MassSimulationSubsystem.h b/Engine/Plugins/Runtime/MassGameplay/Source/MassSimulation/Public/MassSimulationSubsystem.h
--- a/Engine/Plugins/Runtime/MassGameplay/Source/MassSimulation/Public/MassSimulationSubsystem.h
+++ b/Engine/Plugins/Runtime/MassGameplay/Source/MassSimulation/Public/MassSimulationSubsystem.h
@@ -20,6 +20,7 @@ class UMassSimulationSubsystem : public UMassSubsystemBase
 	GENERATED_BODY()
 public:
 	DECLARE_MULTICAST_DELEGATE_OneParam(FOnSimulationStarted, UWorld* /*World*/);
+	DECLARE_MULTICAST_DELEGATE_OneParam(FOnSimulationPauseEvent, TNotNull<UMassSimulationSubsystem*> /*this*/);
 	
 	MASSSIMULATION_API UMassSimulationSubsystem(const FObjectInitializer& ObjectInitializer = FObjectInitializer::Get());
 
@@ -30,6 +31,9 @@ public:
 	MASSSIMULATION_API FMassProcessingPhase::FOnPhaseEvent& GetOnProcessingPhaseFinished(const EMassProcessingPhase Phase);
 	MASSSIMULATION_API static FOnSimulationStarted& GetOnSimulationStarted();
 
+	FOnSimulationPauseEvent& GetOnSimulationPaused();
+	FOnSimulationPauseEvent& GetOnSimulationResumed();
+
 	MASSSIMULATION_API void RegisterDynamicProcessor(UMassProcessor& Processor);
 	MASSSIMULATION_API void UnregisterDynamicProcessor(UMassProcessor& Processor);
 
@@ -41,6 +45,22 @@ public:
 	/** Starts/stops simulation ticking for all worlds, based on new `mass.SimulationTickingEnabled` cvar value */
 	MASSSIMULATION_API static void HandleSimulationTickingEnabledCVarChange(IConsoleVariable*);
 
+	/**
+	 * Determine if this Simulation is currently paused.
+	 * 
+	 * While paused, phases will transition as usual, but processors
+	 * will not be executed.
+	 * 
+	 * @return True if this Simulation is currently paused; else False
+	 */
+	bool IsSimulationPaused() const;
+
+	/** Pause the simulation from executing processors during phase ticks */
+	MASSSIMULATION_API void PauseSimulation();
+
+	/** Resume the simulation executing processors during phase ticks */
+	MASSSIMULATION_API void ResumeSimulation();
+
 protected:
 	// UWorldSubsystem BEGIN
 	MASSSIMULATION_API virtual void Initialize(FSubsystemCollectionBase& Collection) override;
@@ -77,12 +97,16 @@ protected:
 
 	static MASSSIMULATION_API FOnSimulationStarted OnSimulationStarted;
 
+	FOnSimulationPauseEvent OnSimulationPaused;
+	FOnSimulationPauseEvent OnSimulationResumed;
+
 	UPROPERTY()
 	FMassRuntimePipeline RuntimePipeline;
 
 	float CurrentDeltaSeconds = 0.f;
 	bool bTickInProgress = false;
 	bool bSimulationStarted = false;
+	bool bSimulationPaused = false;
 
 #if WITH_EDITOR
 	FDelegateHandle PieBeginEventHandle;
@@ -95,6 +119,21 @@ protected:
 //-----------------------------------------------------------------------------
 // INLINES
 //-----------------------------------------------------------------------------
+inline UMassSimulationSubsystem::FOnSimulationPauseEvent& UMassSimulationSubsystem::GetOnSimulationPaused()
+{
+	return OnSimulationPaused;
+}
+
+inline UMassSimulationSubsystem::FOnSimulationPauseEvent& UMassSimulationSubsystem::GetOnSimulationResumed()
+{
+	return OnSimulationResumed;
+}
+
+inline bool UMassSimulationSubsystem::IsSimulationPaused() const
+{
+	return bSimulationPaused;
+}
+
 inline const FMassProcessingPhaseManager& UMassSimulationSubsystem::GetPhaseManager() const
 {
 	return PhaseManager.Get();
diff --git a/Engine/Source/Runtime/MassEntity/Private/MassProcessingPhaseManager.cpp b/Engine/Source/Runtime/MassEntity/Private/MassProcessingPhaseManager.cpp
--- a/Engine/Source/Runtime/MassEntity/Private/MassProcessingPhaseManager.cpp
+++ b/Engine/Source/Runtime/MassEntity/Private/MassProcessingPhaseManager.cpp
@@ -79,7 +79,7 @@ void FMassProcessingPhase::ExecuteTick(float DeltaTime, ELevelTick TickType, ENa
 
 	bIsDuringMassProcessing = true;
 
-	if (bRunInParallelMode)
+	if (bRunInParallelMode && PhaseManager->IsPaused() == false)
 	{
 		bool bWorkRequested = false;
 		if (PhaseProcessor->IsEmpty() == false)
@@ -103,7 +103,10 @@ void FMassProcessingPhase::ExecuteTick(float DeltaTime, ELevelTick TickType, ENa
 	}
 	else
 	{
-		UE::Mass::Executor::Run(*PhaseProcessor, Context);
+		if (PhaseManager->IsPaused() == false)
+		{
+			UE::Mass::Executor::Run(*PhaseProcessor, Context);
+		}
 
 		{
 			LLM_SCOPE_BYNAME(TEXT("Mass/PhaseEndDelegate"));
@@ -463,11 +466,47 @@ void FMassProcessingPhaseManager::Stop()
 	}
 }
 
+void FMassProcessingPhaseManager::Pause()
+{
+	check(IsInGameThread());
+
+	if (bIsPaused == false)
+	{
+		bIsPauseTogglePending = true;
+
+		UE_VLOG_UELOG(Owner.Get(), LogMass, Log, TEXT("Scheduling Pause for next FrameEnd phase"));
+	}
+}
+
+void FMassProcessingPhaseManager::Resume()
+{
+	check(IsInGameThread());
+
+	if (bIsPaused == true)
+	{
+		bIsPauseTogglePending = true;
+
+		UE_VLOG_UELOG(Owner.Get(), LogMass, Log, TEXT("Scheduling Resume for next PrePhysics phase"));
+	}
+}
+
 void FMassProcessingPhaseManager::OnPhaseStart(FMassProcessingPhase& Phase)
 {
 	ensure(CurrentPhase == EMassProcessingPhase::MAX);
 	CurrentPhase = Phase.Phase;
 
+	const int32 PhaseAsInt = int32(Phase.Phase);
+
+	// The VERY FIRST thing we do in the first phase is to change the Pause state if needed.
+	// This way any code that depends on knowing the pause state (if any) gets consistent results.
+	if (bIsPauseTogglePending && bIsPaused == true && PhaseAsInt == 0)
+	{
+		bIsPaused = false;
+		bIsPauseTogglePending = false;
+
+		UE_VLOG_UELOG(Owner.Get(), LogMass, Log, TEXT("Phase Processing is now Resumed"));
+	}
+
 	// switch between parallel and single-thread versions only after a given batch of processing has been wrapped up	
 	if (Phase.IsConfiguredForParallelMode() != UE::Mass::Tweakables::bFullyParallel)
 	{
@@ -481,8 +520,6 @@ void FMassProcessingPhaseManager::OnPhaseStart(FMassProcessingPhase& Phase)
 		}
 	}
 
-	const int32 PhaseAsInt = int32(Phase.Phase);
-
 	if (PendingDynamicProcessors[PhaseAsInt].IsEmpty() == false)
 	{
 		HandlePendingDynamicProcessorOperations(PhaseAsInt);
@@ -532,6 +569,19 @@ void FMassProcessingPhaseManager::OnPhaseEnd(FMassProcessingPhase& Phase)
 {
 	ensure(CurrentPhase == Phase.Phase);
 	CurrentPhase = EMassProcessingPhase::MAX;
+
+	// The VERY LAST thing we do in FrameEnd is change the Pause state if needed.
+	// This way any code that depends on knowing the pause state (if any) gets consistent results.
+	if (bIsPauseTogglePending && bIsPaused == false
+		&& Phase.Phase == EMassProcessingPhase::FrameEnd)
+	{
+		bIsPaused = true;
+		bIsPauseTogglePending = false;
+
+#if WITH_MASSENTITY_DEBUG
+		UE_VLOG_UELOG(Owner.Get(), LogMass, Log, TEXT("Phase Processing is now Paused"));
+#endif // WITH_MASSENTITY_DEBUG
+	}
 }
 
 FString FMassProcessingPhaseManager::GetName() const
diff --git a/Engine/Source/Runtime/MassEntity/Public/MassProcessingPhaseManager.h b/Engine/Source/Runtime/MassEntity/Public/MassProcessingPhaseManager.h
--- a/Engine/Source/Runtime/MassEntity/Public/MassProcessingPhaseManager.h
+++ b/Engine/Source/Runtime/MassEntity/Public/MassProcessingPhaseManager.h
@@ -177,6 +177,27 @@ public:
 	void Stop();
 	bool IsRunning() const { return EntityManager.IsValid(); }
 
+	/**
+	 * Determine if this Phase Manager is currently paused.
+	 * 
+	 * While paused, phases will transition as usual, but processors
+	 * will not be executed.
+	 * 
+	 * @return True if this PhaseManager is currently paused; else False
+	 */
+	bool IsPaused() const;
+
+	/**
+	 * Pause this phase manager at the earliest opportunity (on next FrameEnd phase end).
+	 * This allows the current phase cycle to complete before the pause takes effect.
+	 */
+	void Pause();
+
+	/**
+	 * Unpause this phase manager at the earliest opportunity (on next PrePhysics phase start).
+	 */
+	void Resume();
+
 	FString GetName() const;
 
 	/** Registers a dynamic processor. This needs to be a fully formed processor and will be slotted in during the next tick. */
@@ -260,6 +281,9 @@ protected:
 	EProcessorExecutionFlags ProcessorExecutionFlags = EProcessorExecutionFlags::None;
 	bool bIsAllowedToTick = false;
 
+	bool bIsPaused = false;
+	bool bIsPauseTogglePending = false;
+
 #if WITH_MASSENTITY_DEBUG
 	FDelegateHandle OnDebugEntityManagerInitializedHandle;
 	FDelegateHandle OnDebugEntityManagerDeinitializedHandle;
@@ -293,3 +317,7 @@ inline TConstArrayView<FMassProcessingPhaseManager::FPhaseGraphBuildState> FMass
 }
 #endif // WITH_MASSENTITY_DEBUG
 
+inline bool FMassProcessingPhaseManager::IsPaused() const 
+{
+	return bIsPaused; 
+}
