PR #12876: [Mass Entity] Fix a crash caused by capturing "this" in a lambda expression

diff --git a/Engine/Plugins/Runtime/MassGameplay/Source/MassSpawner/Private/MassSpawner.cpp b/Engine/Plugins/Runtime/MassGameplay/Source/MassSpawner/Private/MassSpawner.cpp
--- a/Engine/Plugins/Runtime/MassGameplay/Source/MassSpawner/Private/MassSpawner.cpp
+++ b/Engine/Plugins/Runtime/MassGameplay/Source/MassSpawner/Private/MassSpawner.cpp
@@ -356,7 +356,7 @@ void AMassSpawner::DoSpawning()
 	if (AssetsToLoad.Num())
 	{
 		FStreamableManager& StreamableManager = UAssetManager::GetStreamableManager();
-		StreamingHandle = StreamableManager.RequestAsyncLoad(AssetsToLoad, GenerateSpawningPoints);
+		StreamingHandle = StreamableManager.RequestAsyncLoad(AssetsToLoad, FStreamableDelegate::CreateWeakLambda(this, GenerateSpawningPoints));
 	}
 	else
 	{
