PR #12791: MassLOD: Use Correct `bLocalViewersOnly` Setting

The change makes listen-servers behave the same as dedicated servers when considering Mass LOD Viewers (i.e. considering all viewers, not just the local ones). The change also makes Standalone games use the same path, which doesn't make any difference since there are no non-local viewers in that scenario.

diff --git a/Engine/Plugins/Runtime/MassGameplay/Source/MassLOD/Private/MassLODCollectorProcessor.cpp b/Engine/Plugins/Runtime/MassGameplay/Source/MassLOD/Private/MassLODCollectorProcessor.cpp
--- a/Engine/Plugins/Runtime/MassGameplay/Source/MassLOD/Private/MassLODCollectorProcessor.cpp
+++ b/Engine/Plugins/Runtime/MassGameplay/Source/MassLOD/Private/MassLODCollectorProcessor.cpp
@@ -87,7 +87,7 @@ void UMassLODCollectorProcessor::Execute(FMassEntityManager& EntityManager, FMas
 
 	UWorld* World = EntityManager.GetWorld();
 	check(World);
-	if (World->IsNetMode(NM_DedicatedServer))
+	if (World->GetNetMode() != NM_Client)
 	{
 		ExecuteInternal<false/*bLocalViewersOnly*/>(EntityManager, Context);
 	}
diff --git a/Engine/Plugins/Runtime/MassGameplay/Source/MassLOD/Private/MassLODDistanceCollectorProcessor.cpp b/Engine/Plugins/Runtime/MassGameplay/Source/MassLOD/Private/MassLODDistanceCollectorProcessor.cpp
--- a/Engine/Plugins/Runtime/MassGameplay/Source/MassLOD/Private/MassLODDistanceCollectorProcessor.cpp
+++ b/Engine/Plugins/Runtime/MassGameplay/Source/MassLOD/Private/MassLODDistanceCollectorProcessor.cpp
@@ -100,7 +100,7 @@ void UMassLODDistanceCollectorProcessor::Execute(FMassEntityManager& EntityManag
 
 	UWorld* World = EntityManager.GetWorld();
 	check(World);
-	if (World->IsNetMode(NM_DedicatedServer))
+	if (World->GetNetMode() != NM_Client)
 	{
 		ExecuteInternal<false/*bLocalViewersOnly*/>(EntityManager, Context);
 	}
