[MassGameplay] Made UMassSceneComponentLocationToMassTranslator gamethread-only to sidestep data races when the engine updates components' transforms.

diff --git a/Engine/Plugins/Runtime/MassGameplay/Source/MassActors/Private/Translators/MassSceneComponentLocationTranslator.cpp b/Engine/Plugins/Runtime/MassGameplay/Source/MassActors/Private/Translators/MassSceneComponentLocationTranslator.cpp
--- a/Engine/Plugins/Runtime/MassGameplay/Source/MassActors/Private/Translators/MassSceneComponentLocationTranslator.cpp
+++ b/Engine/Plugins/Runtime/MassGameplay/Source/MassActors/Private/Translators/MassSceneComponentLocationTranslator.cpp
@@ -14,6 +14,7 @@ UMassSceneComponentLocationToMassTranslator::UMassSceneComponentLocationToMassTr
 	ExecutionFlags = (int32)EProcessorExecutionFlags::AllNetModes;
 	ExecutionOrder.ExecuteInGroup = UE::Mass::ProcessorGroupNames::SyncWorldToMass;
 	RequiredTags.Add<FMassSceneComponentLocationCopyToMassTag>();
+	bRequiresGameThreadExecution = true;
 }
 
 void UMassSceneComponentLocationToMassTranslator::ConfigureQueries(const TSharedRef<FMassEntityManager>& EntityManager)
