[Mass] Mass debug drawing improvements/changes - added a cvar, setting `mass.debug.DrawAllEntities 1` will cause a color box drawn on every entity containing a FTransformFragment. The color used will be the MassArchetypeData.DebugColor

diff --git a/Engine/Plugins/Runtime/MassGameplay/Source/MassGameplayDebug/Private/DebugVisLocationProcessor.cpp b/Engine/Plugins/Runtime/MassGameplay/Source/MassGameplayDebug/Private/DebugVisLocationProcessor.cpp
--- a/Engine/Plugins/Runtime/MassGameplay/Source/MassGameplayDebug/Private/DebugVisLocationProcessor.cpp
+++ b/Engine/Plugins/Runtime/MassGameplay/Source/MassGameplayDebug/Private/DebugVisLocationProcessor.cpp
@@ -8,15 +8,33 @@
 #include "MassEntityManager.h"
 #include "MassExecutionContext.h"
 #include "MassMovementFragments.h"
+#include "MassDebugDrawHelpers.h"
+
+
+namespace UE::Mass::Debug
+{
+	bool bDebugDrawAllEntities = false;
+
+	namespace
+	{
+		FAutoConsoleVariableRef AnonymousCVars[] = {
+			{TEXT("mass.debug.DrawAllEntities"), bDebugDrawAllEntities
+				, TEXT("When enabled will debug-draw debug shapes marking all entities that have a TransformFragment")
+				, ECVF_Cheat},
+		};
+	}
+}
 
 //----------------------------------------------------------------------//
 // UDebugVisLocationProcessor
 //----------------------------------------------------------------------//
 UDebugVisLocationProcessor::UDebugVisLocationProcessor()
 	: EntityQuery(*this)
+	, AllLocationEntitiesQuery(*this)
 {
 	ExecutionOrder.ExecuteAfter.Add(UE::Mass::ProcessorGroupNames::SyncWorldToMass);
 	bRequiresGameThreadExecution = true; // due to UMassDebuggerSubsystem access
+	QueryBasedPruning = EMassQueryBasedPruning::Never;
 }
 
 void UDebugVisLocationProcessor::ConfigureQueries(const TSharedRef<FMassEntityManager>& EntityManager)
@@ -25,6 +43,8 @@ void UDebugVisLocationProcessor::ConfigureQueries(const TSharedRef<FMassEntityMa
 	EntityQuery.AddRequirement<FSimDebugVisFragment>(EMassFragmentAccess::ReadOnly);
 	EntityQuery.AddTagRequirement<FMassDebuggableTag>(EMassFragmentPresence::All);
 	EntityQuery.AddSubsystemRequirement<UMassDebuggerSubsystem>(EMassFragmentAccess::ReadWrite);
+
+	AllLocationEntitiesQuery.AddRequirement<FTransformFragment>(EMassFragmentAccess::ReadOnly);
 }
 
 void UDebugVisLocationProcessor::Execute(FMassEntityManager& EntityManager, FMassExecutionContext& Context)
@@ -32,41 +52,77 @@ void UDebugVisLocationProcessor::Execute(FMassEntityManager& EntityManager, FMas
 #if WITH_EDITORONLY_DATA
 	QUICK_SCOPE_CYCLE_COUNTER(DebugVisLocationProcessor_Run);
 
-	EntityQuery.ForEachEntityChunk(Context, [this](FMassExecutionContext& Context)
+	if (UE::Mass::Debug::bDebugDrawAllEntities == false)
 	{
-		UMassDebuggerSubsystem& Debugger = Context.GetMutableSubsystemChecked<UMassDebuggerSubsystem>();
-		UMassDebugVisualizationComponent* Visualizer = Debugger.GetVisualizationComponent();
-		check(Visualizer);
-		TArrayView<UHierarchicalInstancedStaticMeshComponent* const> VisualDataISMCs = Visualizer->GetVisualDataISMCs();
-		if (VisualDataISMCs.Num() > 0)
-		{
-			const TConstArrayView<FTransformFragment> LocationList = Context.GetFragmentView<FTransformFragment>();
-			const TConstArrayView<FSimDebugVisFragment> DebugVisList = Context.GetFragmentView<FSimDebugVisFragment>();
-
-			for (FMassExecutionContext::FEntityIterator EntityIt = Context.CreateEntityIterator(); EntityIt; ++EntityIt)
+		EntityQuery.ForEachEntityChunk(Context, [this](FMassExecutionContext& Context)
 			{
-				const FSimDebugVisFragment& VisualComp = DebugVisList[EntityIt];
-
-				// @todo: remove this code once the asset is exported with correct alignment SM_Mannequin.uasset
-				FTransform SMTransform = LocationList[EntityIt].GetTransform();
-				FQuat FromEngineToSM(FVector::UpVector, -HALF_PI);
-				SMTransform.SetRotation(FromEngineToSM * SMTransform.GetRotation());
-
-				VisualDataISMCs[VisualComp.VisualType]->UpdateInstanceTransform(VisualComp.InstanceIndex, SMTransform, true);
-			}
-		}
-		else
+				UMassDebuggerSubsystem& Debugger = Context.GetMutableSubsystemChecked<UMassDebuggerSubsystem>();
+				UMassDebugVisualizationComponent* Visualizer = Debugger.GetVisualizationComponent();
+				check(Visualizer);
+				TArrayView<UHierarchicalInstancedStaticMeshComponent* const> VisualDataISMCs = Visualizer->GetVisualDataISMCs();
+				if (VisualDataISMCs.Num() > 0)
+				{
+					const int32 NumEntities = Context.GetNumEntities();
+					const TConstArrayView<FTransformFragment> LocationList = Context.GetFragmentView<FTransformFragment>();
+					const TConstArrayView<FSimDebugVisFragment> DebugVisList = Context.GetFragmentView<FSimDebugVisFragment>();
+
+					for (int32 i = 0; i < NumEntities; ++i)
+					{
+						const FSimDebugVisFragment& VisualComp = DebugVisList[i];
+
+						// @todo: remove this code once the asset is exported with correct alignment SM_Mannequin.uasset
+						FTransform SMTransform = LocationList[i].GetTransform();
+						FQuat FromEngineToSM(FVector::UpVector, -HALF_PI);
+						SMTransform.SetRotation(FromEngineToSM * SMTransform.GetRotation());
+
+						VisualDataISMCs[VisualComp.VisualType]->UpdateInstanceTransform(VisualComp.InstanceIndex, SMTransform, true);
+					}
+				}
+				else
+				{
+					UE_LOG(LogMassDebug, Log, TEXT("UDebugVisLocationProcessor: Trying to update InstanceStaticMeshes while none created. Check your debug visualization setup"));
+				}
+			});
+
+		UMassDebuggerSubsystem* Debugger = UWorld::GetSubsystem<UMassDebuggerSubsystem>(EntityManager.GetWorld());
+		if (ensure(Debugger))
 		{
-			UE_LOG(LogMassDebug, Log, TEXT("UDebugVisLocationProcessor: Trying to update InstanceStaticMeshes while none created. Check your debug visualization setup"));
+			Debugger->GetVisualizationComponent()->DirtyVisuals();
 		}
-	});
-
-	UMassDebuggerSubsystem* Debugger = UWorld::GetSubsystem<UMassDebuggerSubsystem>(EntityManager.GetWorld());
-	if (ensure(Debugger))
-	{
-		Debugger->GetVisualizationComponent()->DirtyVisuals();
 	}
 #endif // WITH_EDITORONLY_DATA
+	if (UE::Mass::Debug::bDebugDrawAllEntities)
+	{
+		if (UWorld* World = GetWorld())
+		{
+			UE::Mass::Debug::FLineBatcher LineBatcher = UE::Mass::Debug::FLineBatcher::MakeLineBatcher(World);
+
+			AllLocationEntitiesQuery.ForEachEntityChunk(Context, [&LineBatcher](const FMassExecutionContext& Context)
+				{
+					const FColor ArchetypeColor =
+#if WITH_MASSENTITY_DEBUG
+						Context.DebugGetArchetypeColor();
+#else
+						FColor::Green;
+#endif // WITH_MASSENTITY_DEBUG
+
+					const FVector DebugDrawBoxExtent(20.);
+					constexpr float ArrowLength = 30.f;
+
+					const TConstArrayView<FTransformFragment> LocationList = Context.GetFragmentView<FTransformFragment>();
+
+					for (const FTransformFragment& LocationFragment : LocationList)
+					{
+						LineBatcher.DrawSolidBox(LocationFragment.GetTransform().GetLocation(), DebugDrawBoxExtent, ArchetypeColor);
+						if (LocationFragment.GetTransform().GetRotation().IsIdentity() == false)
+						{
+							
+							LineBatcher.DrawArrow(LocationFragment.GetTransform(), ArrowLength, ArchetypeColor);
+						}
+					}
+				});
+		}
+	}
 }
 
 //----------------------------------------------------------------------//
diff --git a/Engine/Plugins/Runtime/MassGameplay/Source/MassGameplayDebug/Public/DebugVisLocationProcessor.h b/Engine/Plugins/Runtime/MassGameplay/Source/MassGameplayDebug/Public/DebugVisLocationProcessor.h
--- a/Engine/Plugins/Runtime/MassGameplay/Source/MassGameplayDebug/Public/DebugVisLocationProcessor.h
+++ b/Engine/Plugins/Runtime/MassGameplay/Source/MassGameplayDebug/Public/DebugVisLocationProcessor.h
@@ -21,6 +21,7 @@ protected:
 	virtual void Execute(FMassEntityManager& EntityManager, FMassExecutionContext& Context) override;
 
 	FMassEntityQuery EntityQuery;
+	FMassEntityQuery AllLocationEntitiesQuery;
 };
 
 //----------------------------------------------------------------------//
diff --git a/Engine/Plugins/Runtime/MassGameplay/Source/MassGameplayDebug/Public/MassDebugDrawHelpers.h b/Engine/Plugins/Runtime/MassGameplay/Source/MassGameplayDebug/Public/MassDebugDrawHelpers.h
new file mode 100644
--- /dev/null
+++ b/Engine/Plugins/Runtime/MassGameplay/Source/MassGameplayDebug/Public/MassDebugDrawHelpers.h
@@ -0,0 +1,53 @@
+// Copyright Epic Games, Inc. All Rights Reserved.
+
+#pragma once
+
+#include "Misc/NotNull.h"
+#include "Components/LineBatchComponent.h"
+#include "Engine/World.h"
+
+
+namespace UE::Mass::Debug
+{
+	struct FLineBatcher
+	{
+		static FLineBatcher MakeLineBatcher(const UWorld* InWorld, bool bPersistentLines = false, float LifeTime = -1.f)
+		{
+			return FLineBatcher(InWorld 
+					? (( bPersistentLines || (LifeTime > 0.f)) 
+						? InWorld->GetLineBatcher(UWorld::ELineBatcherType::WorldPersistent)
+						: InWorld->GetLineBatcher(UWorld::ELineBatcherType::World))
+					: nullptr
+				, LifeTime);
+		}
+
+		FLineBatcher(TNotNull<ULineBatchComponent*> InLineBatcherInstance, float InLifeTime = -1.f)
+			: LineBatcherInstance(InLineBatcherInstance), LifeTime(InLifeTime)
+		{	
+		}
+
+		FORCEINLINE void DrawSolidBox(const FVector& Center, const FVector& Extent, FColor const& Color) const
+		{
+			const FBox Box = FBox::BuildAABB(Center, Extent);
+			LineBatcherInstance->DrawSolidBox(Box, FTransform::Identity, Color, /*DepthPriority=*/0, LifeTime);
+		}
+
+		FORCEINLINE void DrawWireBox(const FVector& Center, const FVector& Extent, FColor const& Color) const
+		{
+			LineBatcherInstance->DrawBox(Center, Extent, Color, LifeTime, /*DepthPriority=*/0, /*Thickness=*/0.f);
+		}
+
+		FORCEINLINE void DrawSphere(const FVector& Center, const float Radius, FLinearColor const& Color) const
+		{
+			LineBatcherInstance->DrawSphere(Center, Radius, /*Segments=*/8, Color, LifeTime, /*DepthPriority=*/0, /*Thickness=*/0.f);
+		}
+
+		FORCEINLINE void DrawArrow(const FTransform& Transform, const float Length, FColor const& Color) const
+		{
+			LineBatcherInstance->DrawDirectionalArrow(Transform.ToMatrixNoScale(), Color, Length, Length/5, /*DepthPriority=*/0);
+		}
+
+		TNotNull<ULineBatchComponent* const> LineBatcherInstance;
+		float LifeTime = -1.f;
+	};
+}
\ No newline at end of file
