[MassGameplay] Fixed MassAgentComponent trying to register with the agent subsystem in inactive worlds

diff --git a/Engine/Plugins/Runtime/MassGameplay/Source/MassActors/Private/MassAgentComponent.cpp b/Engine/Plugins/Runtime/MassGameplay/Source/MassActors/Private/MassAgentComponent.cpp
--- a/Engine/Plugins/Runtime/MassGameplay/Source/MassActors/Private/MassAgentComponent.cpp
+++ b/Engine/Plugins/Runtime/MassGameplay/Source/MassActors/Private/MassAgentComponent.cpp
@@ -88,7 +88,7 @@ void UMassAgentComponent::OnRegister()
 	}
 
 	UWorld* World = GetWorld();
-	if (World == nullptr
+	if (World == nullptr || World->WorldType == EWorldType::None || World->WorldType == EWorldType::Inactive
 #if WITH_EDITOR
 		|| World->IsPreviewWorld() || (bAutoRegisterInEditorMode == false && World->IsGameWorld() == false)
 #endif // WITH_EDITOR
