Minor MassMovement changes: Follow up from 40255404

diff --git a/Engine/Plugins/Runtime/MassGameplay/Source/MassMovement/Private/MassMovementProcessors.cpp b/Engine/Plugins/Runtime/MassGameplay/Source/MassMovement/Private/MassMovementProcessors.cpp
--- a/Engine/Plugins/Runtime/MassGameplay/Source/MassMovement/Private/MassMovementProcessors.cpp
+++ b/Engine/Plugins/Runtime/MassGameplay/Source/MassMovement/Private/MassMovementProcessors.cpp
@@ -36,8 +36,6 @@ void UMassApplyForceProcessor::Execute(FMassEntityManager& EntityManager, FMassE
 	// Clamp max delta time to avoid force explosion on large time steps (i.e. during initialization).
 	const float DeltaTime = FMath::Min(0.1f, Context.GetDeltaTimeSeconds());
 
-	QUICK_SCOPE_CYCLE_COUNTER(UMassApplyForceProcessor);
-
 	EntityQuery.ForEachEntityChunk(Context, [this, DeltaTime](FMassExecutionContext& Context)
 	{
 		const TArrayView<FMassForceFragment> ForceList = Context.GetMutableFragmentView<FMassForceFragment>();
@@ -84,8 +82,6 @@ void UMassApplyMovementProcessor::Execute(FMassEntityManager& EntityManager,
 	// Clamp max delta time to avoid force explosion on large time steps (i.e. during initialization).
 	const float DeltaTime = FMath::Min(0.1f, Context.GetDeltaTimeSeconds());
 
-	QUICK_SCOPE_CYCLE_COUNTER(UMassApplyMovementProcessor);
-
 	EntityQuery.ForEachEntityChunk(Context, [this, DeltaTime](FMassExecutionContext& Context)
 	{
 		const TArrayView<FTransformFragment> LocationList = Context.GetMutableFragmentView<FTransformFragment>();
diff --git a/Engine/Plugins/Runtime/MassGameplay/Source/MassMovement/Public/MassMovementFragments.h b/Engine/Plugins/Runtime/MassGameplay/Source/MassMovement/Public/MassMovementFragments.h
--- a/Engine/Plugins/Runtime/MassGameplay/Source/MassMovement/Public/MassMovementFragments.h
+++ b/Engine/Plugins/Runtime/MassGameplay/Source/MassMovement/Public/MassMovementFragments.h
@@ -11,8 +11,10 @@
 #include "RandomSequence.h"
 #include "MassMovementFragments.generated.h"
 
-// This represents the actual physical velocity of the mass entity in the world
-// For agents with an actor representation, this is the velocity of the movement component
+/**
+ * This represents the actual physical velocity of the mass entity in the world
+ * For agents with an actor representation, this is the velocity of the movement component
+ */
 USTRUCT()
 struct MASSMOVEMENT_API FMassVelocityFragment : public FMassFragment
 {
@@ -25,8 +27,10 @@ struct MASSMOVEMENT_API FMassVelocityFragment : public FMassFragment
 #endif	
 };
 
-// This is the output of all processors that intend to affect movement
-// It is the input to the movement system (e.g. mover, animation etc.)
+/**
+ * This is the output of all processors that intend to affect movement
+ * It is the input to the movement system (e.g. mover, animation etc.)
+ */
 USTRUCT()
 struct MASSMOVEMENT_API FMassDesiredMovementFragment : public FMassFragment
 {
@@ -37,7 +41,7 @@ struct MASSMOVEMENT_API FMassDesiredMovementFragment : public FMassFragment
 	// Todo: Could add a DesiredFacing here
 };
 
-// Accumulator for steering / avoidance forces to apply to the desired velocity 
+/** Accumulator for steering / avoidance forces to apply to the desired velocity */ 
 USTRUCT()
 struct MASSMOVEMENT_API FMassForceFragment : public FMassFragment
 {
@@ -46,20 +50,22 @@ struct MASSMOVEMENT_API FMassForceFragment : public FMassFragment
 	FVector Value = FVector::ZeroVector;
 };
 
-// The presence of this tag indicates that this mass agent's velocity should be
-// controlled by the FMassMovementCommand.
-//
-// For code driven displacement, we want the desired velocity to affect the velocity
-// directly, which is then applied to the character mover.
-// For e.g. Root Motion driven displacement, we just need to pipe the DesiredVelocity to the
-// animation system and let it do the test.
+/**
+ * The presence of this tag indicates that this mass agent's velocity should be
+ * controlled by the FMassDesiredMovementFragment.
+ *
+ * For code driven displacement, we want the desired velocity to affect the velocity
+ * directly, which is then applied to the character mover.
+ * For e.g. Root Motion driven displacement, we just need to pipe the DesiredVelocity to the
+ * animation system and let it do the test.
+ */
 USTRUCT()
 struct MASSMOVEMENT_API FMassCodeDrivenMovementTag : public FMassTag
 {
 	GENERATED_BODY()
 };
 
-// Parameters describing how this mass agent should move
+/** Parameters describing how this mass agent should move */
 USTRUCT()
 struct MASSMOVEMENT_API FMassMovementParameters : public FMassConstSharedFragment
 {
