[MassGameplay] Regenerating entity config GUID when copy-pasting MassAgentComponent-owning Actor into a scene in the Editor

diff --git a/Engine/Plugins/Runtime/MassGameplay/Source/MassActors/Private/MassAgentComponent.cpp b/Engine/Plugins/Runtime/MassGameplay/Source/MassActors/Private/MassAgentComponent.cpp
--- a/Engine/Plugins/Runtime/MassGameplay/Source/MassActors/Private/MassAgentComponent.cpp
+++ b/Engine/Plugins/Runtime/MassGameplay/Source/MassActors/Private/MassAgentComponent.cpp
@@ -68,6 +68,18 @@ void UMassAgentComponent::PostDuplicate(const bool bDuplicateForPIE)
 	Super::PostDuplicate(bDuplicateForPIE);
 	EntityConfig.PostDuplicate(bDuplicateForPIE);
 }
+
+void UMassAgentComponent::PostEditChangeProperty(FPropertyChangedEvent& PropertyChangedEvent)
+{
+	static const FName NAME_EntityConfig = GET_MEMBER_NAME_CHECKED(UMassAgentComponent, EntityConfig);
+
+	Super::PostEditChangeProperty(PropertyChangedEvent);
+
+	if (PropertyChangedEvent.Property && PropertyChangedEvent.Property->GetFName() == NAME_EntityConfig)
+	{
+		EntityConfig.ForceRegenerateGUID();
+	}
+}
 #endif // WITH_EDITOR
 
 void UMassAgentComponent::OnRegister()
diff --git a/Engine/Plugins/Runtime/MassGameplay/Source/MassActors/Public/MassAgentComponent.h b/Engine/Plugins/Runtime/MassGameplay/Source/MassActors/Public/MassAgentComponent.h
--- a/Engine/Plugins/Runtime/MassGameplay/Source/MassActors/Public/MassAgentComponent.h
+++ b/Engine/Plugins/Runtime/MassGameplay/Source/MassActors/Public/MassAgentComponent.h
@@ -54,6 +54,7 @@ protected:
 	virtual void PostInitProperties() override;
 	virtual void PostLoad() override;
 	virtual void PostDuplicate(const bool bDuplicateForPIE) override;
+	virtual void PostEditChangeProperty(struct FPropertyChangedEvent& PropertyChangedEvent) override;
 #endif // WITH_EDITOR
 
 public:
diff --git a/Engine/Plugins/Runtime/MassGameplay/Source/MassSpawner/Private/MassEntityConfigAsset.cpp b/Engine/Plugins/Runtime/MassGameplay/Source/MassSpawner/Private/MassEntityConfigAsset.cpp
--- a/Engine/Plugins/Runtime/MassGameplay/Source/MassSpawner/Private/MassEntityConfigAsset.cpp
+++ b/Engine/Plugins/Runtime/MassGameplay/Source/MassSpawner/Private/MassEntityConfigAsset.cpp
@@ -190,10 +190,15 @@ void FMassEntityConfig::PostDuplicate(const bool bDuplicateForPIE)
 {
 	if (bDuplicateForPIE == false)
 	{
-		ConfigGuid = FGuid::NewGuid();
+		ForceRegenerateGUID();
 	}
 }
 
+void FMassEntityConfig::ForceRegenerateGUID()
+{
+	ConfigGuid = FGuid::NewGuid();
+}
+
 UMassEntityTraitBase* FMassEntityConfig::FindMutableTrait(TSubclassOf<UMassEntityTraitBase> TraitClass, const bool bExactMatch)
 {
 	return FindTraitInternal(TraitClass, bExactMatch);
diff --git a/Engine/Plugins/Runtime/MassGameplay/Source/MassSpawner/Public/MassEntityConfigAsset.h b/Engine/Plugins/Runtime/MassGameplay/Source/MassSpawner/Public/MassEntityConfigAsset.h
--- a/Engine/Plugins/Runtime/MassGameplay/Source/MassSpawner/Public/MassEntityConfigAsset.h
+++ b/Engine/Plugins/Runtime/MassGameplay/Source/MassSpawner/Public/MassEntityConfigAsset.h
@@ -71,6 +71,7 @@ struct MASSSPAWNER_API FMassEntityConfig
 #if WITH_EDITOR
 	/** Needs to be called when the given config is being duplicated - ensured the ConfigGuid remains unique */
 	void PostDuplicate(bool bDuplicateForPIE);
+	void ForceRegenerateGUID();
 
 	UMassEntityTraitBase* FindMutableTrait(TSubclassOf<UMassEntityTraitBase> TraitClass, const bool bExactMatch = false);
 #endif // WITH_EDITOR
