[MassGameplay] fixed creation context in UMassAgentSubsystem::HandlePendingInitialization not actually being created.

#rnx

diff --git a/Engine/Plugins/Runtime/MassGameplay/Source/MassActors/Private/MassAgentSubsystem.cpp b/Engine/Plugins/Runtime/MassGameplay/Source/MassActors/Private/MassAgentSubsystem.cpp
--- a/Engine/Plugins/Runtime/MassGameplay/Source/MassActors/Private/MassAgentSubsystem.cpp
+++ b/Engine/Plugins/Runtime/MassGameplay/Source/MassActors/Private/MassAgentSubsystem.cpp
@@ -306,7 +306,7 @@ void UMassAgentSubsystem::HandlePendingInitialization()
 
 	{
 		// this block is scoped so that CreationContext gets released by the end
-		TSharedPtr<FMassEntityManager::FEntityCreationContext> CreationContext;
+		TSharedPtr<FMassEntityManager::FEntityCreationContext> CreationContext = EntityManager->GetOrMakeCreationContext();
 
 		for (TTuple<FMassEntityTemplateID, FMassAgentInitializationQueue>& Data : PendingAgentEntities)
 		{
