[MassGameplay] Refactored single-line functions in MassEntityTemplate.h to the end of the file.

diff --git a/Engine/Plugins/Runtime/MassGameplay/Source/MassSpawner/Public/MassEntityTemplate.h b/Engine/Plugins/Runtime/MassGameplay/Source/MassSpawner/Public/MassEntityTemplate.h
--- a/Engine/Plugins/Runtime/MassGameplay/Source/MassSpawner/Public/MassEntityTemplate.h
+++ b/Engine/Plugins/Runtime/MassGameplay/Source/MassSpawner/Public/MassEntityTemplate.h
@@ -74,7 +74,10 @@ public:
 		return GetTypeHash(TemplateID.TotalHash);
 	}
 
-	bool IsValid() const { return (TotalHash != InvalidHash); }
+	bool IsValid() const
+	{
+		return (TotalHash != InvalidHash);
+	}
 
 	FString ToString() const;
 
@@ -108,17 +111,17 @@ struct MASSSPAWNER_API FMassEntityTemplateData
 	FMassEntityTemplateData() = default;
 	explicit FMassEntityTemplateData(const FMassEntityTemplate& InFinalizedTemplate);
 
-	bool IsEmpty() const { return Composition.IsEmpty(); }
+	bool IsEmpty() const;
 
-	TConstArrayView<FObjectFragmentInitializerFunction> GetObjectFragmentInitializers() const { return ObjectInitializers; }
-	const FString& GetTemplateName() const { return TemplateName; }
-	const FMassArchetypeCompositionDescriptor& GetCompositionDescriptor() const { return Composition; }
-	const FMassArchetypeSharedFragmentValues& GetSharedFragmentValues() const { return SharedFragmentValues; }
-	TConstArrayView<FInstancedStruct> GetInitialFragmentValues() const { return InitialFragmentValues; }
+	TConstArrayView<FObjectFragmentInitializerFunction> GetObjectFragmentInitializers() const;
+	const FString& GetTemplateName() const;
+	const FMassArchetypeCompositionDescriptor& GetCompositionDescriptor() const;
+	const FMassArchetypeSharedFragmentValues& GetSharedFragmentValues() const;
+	TConstArrayView<FInstancedStruct> GetInitialFragmentValues() const;
 
-	TArray<FMassEntityTemplateData::FObjectFragmentInitializerFunction>& GetMutableObjectFragmentInitializers() { return ObjectInitializers; }
+	TArray<FMassEntityTemplateData::FObjectFragmentInitializerFunction>& GetMutableObjectFragmentInitializers();
 
-	void SetTemplateName(const FString& Name) { TemplateName = Name; }
+	void SetTemplateName(const FString& Name);
 	
 	template<typename T>
 	void AddFragment()
@@ -200,8 +203,8 @@ struct MASSSPAWNER_API FMassEntityTemplateData
 		Composition.Tags.Remove(TagType);
 	}
 
-	const FMassTagBitSet& GetTags() const { return Composition.Tags; }
-	FMassTagBitSet& GetMutableTags() { return Composition.Tags; }
+	const FMassTagBitSet& GetTags() const;
+	FMassTagBitSet& GetMutableTags();
 
 	template<typename T>
 	void AddChunkFragment()
@@ -306,7 +309,7 @@ struct MASSSPAWNER_API FMassEntityTemplateData
 	 *  @Note that the function can be slow, depending on how elaborate the template is. This function is meant for debugging purposes. */
 	bool SlowIsEquivalent(const FMassEntityTemplateData& Other) const;
 
-	FMassArchetypeCreationParams& GetArchetypeCreationParams() { return CreationParams; }
+	FMassArchetypeCreationParams& GetArchetypeCreationParams();
 	
 protected:
 	FMassArchetypeCompositionDescriptor Composition;
@@ -338,12 +341,12 @@ struct MASSSPAWNER_API FMassEntityTemplate final : public TSharedFromThis<FMassE
 
 	/** InArchetype is expected to be valid. The function will crash-check it. */
 	void SetArchetype(const FMassArchetypeHandle& InArchetype);
-	const FMassArchetypeHandle& GetArchetype() const { return Archetype; }
+	const FMassArchetypeHandle& GetArchetype() const;
 
-	bool IsValid() const { return Archetype.IsValid(); }
+	bool IsValid() const;
 
-	void SetTemplateID(FMassEntityTemplateID InTemplateID) { TemplateID = InTemplateID; }
-	FMassEntityTemplateID GetTemplateID() const { return TemplateID; }
+	void SetTemplateID(FMassEntityTemplateID InTemplateID);
+	FMassEntityTemplateID GetTemplateID() const;
 
 	FString DebugGetDescription(FMassEntityManager* EntityManager = nullptr) const;
 	FString DebugGetArchetypeDescription(FMassEntityManager& EntityManager) const;
@@ -353,13 +356,12 @@ struct MASSSPAWNER_API FMassEntityTemplate final : public TSharedFromThis<FMassE
 	//-----------------------------------------------------------------------------
 	// FMassEntityTemplateData getters
 	//-----------------------------------------------------------------------------
-	FORCEINLINE TConstArrayView<FMassEntityTemplateData::FObjectFragmentInitializerFunction> GetObjectFragmentInitializers() const { return TemplateData.GetObjectFragmentInitializers(); }
-	FORCEINLINE const FString& GetTemplateName() const { return TemplateData.GetTemplateName(); }
-	FORCEINLINE const FMassArchetypeCompositionDescriptor& GetCompositionDescriptor() const { return TemplateData.GetCompositionDescriptor(); }
-	FORCEINLINE const FMassArchetypeSharedFragmentValues& GetSharedFragmentValues() const { return TemplateData.GetSharedFragmentValues(); }
-	FORCEINLINE TConstArrayView<FInstancedStruct> GetInitialFragmentValues() const { return TemplateData.GetInitialFragmentValues(); }
-
-	const FMassEntityTemplateData& GetTemplateData() const { return TemplateData; }
+	TConstArrayView<FMassEntityTemplateData::FObjectFragmentInitializerFunction> GetObjectFragmentInitializers() const;
+	const FString& GetTemplateName() const;
+	const FMassArchetypeCompositionDescriptor& GetCompositionDescriptor() const;
+	const FMassArchetypeSharedFragmentValues& GetSharedFragmentValues() const;
+	TConstArrayView<FInstancedStruct> GetInitialFragmentValues() const;
+	const FMassEntityTemplateData& GetTemplateData() const;
 
 private:
 	FMassEntityTemplateData TemplateData;
@@ -373,3 +375,113 @@ struct MASSSPAWNER_API FMassEntityTemplateIDFactory
 	static FMassEntityTemplateID Make(const FGuid& ConfigGuid);
 	static FMassEntityTemplateID MakeFlavor(const FMassEntityTemplateID& SourceTemplateID, const int32 Flavor);
 };
+
+
+//-----------------------------------------------------------------------------
+// INLINES
+//-----------------------------------------------------------------------------
+
+inline bool FMassEntityTemplateData::IsEmpty() const
+{
+	return Composition.IsEmpty(); 
+}
+
+inline TConstArrayView<FMassEntityTemplateData::FObjectFragmentInitializerFunction> FMassEntityTemplateData::GetObjectFragmentInitializers() const
+{ 
+	return ObjectInitializers; 
+}
+
+inline const FString& FMassEntityTemplateData::GetTemplateName() const
+{ 
+	return TemplateName; 
+}
+
+inline const FMassArchetypeCompositionDescriptor& FMassEntityTemplateData::GetCompositionDescriptor() const
+{ 
+	return Composition; 
+}
+
+inline const FMassArchetypeSharedFragmentValues& FMassEntityTemplateData::GetSharedFragmentValues() const
+{ 
+	return SharedFragmentValues; 
+}
+
+inline TConstArrayView<FInstancedStruct> FMassEntityTemplateData::GetInitialFragmentValues() const
+{ 
+	return InitialFragmentValues; 
+}
+
+inline TArray<FMassEntityTemplateData::FObjectFragmentInitializerFunction>& FMassEntityTemplateData::GetMutableObjectFragmentInitializers()
+{ 
+	return ObjectInitializers; 
+}
+
+inline void FMassEntityTemplateData::SetTemplateName(const FString& Name)
+{
+	TemplateName = Name; 
+}
+
+inline const FMassTagBitSet& FMassEntityTemplateData::GetTags() const
+{
+	return Composition.Tags;
+}
+
+inline FMassTagBitSet& FMassEntityTemplateData::GetMutableTags()
+{
+	return Composition.Tags;
+}
+
+inline FMassArchetypeCreationParams& FMassEntityTemplateData::GetArchetypeCreationParams()
+{
+	return CreationParams;
+}
+
+inline const FMassArchetypeHandle& FMassEntityTemplate::GetArchetype() const
+{
+	return Archetype;
+}
+
+inline bool FMassEntityTemplate::IsValid() const
+{
+	return Archetype.IsValid();
+}
+
+inline void FMassEntityTemplate::SetTemplateID(FMassEntityTemplateID InTemplateID)
+{
+	TemplateID = InTemplateID;
+}
+
+inline FMassEntityTemplateID FMassEntityTemplate::GetTemplateID() const
+{
+	return TemplateID;
+}
+
+inline TConstArrayView<FMassEntityTemplateData::FObjectFragmentInitializerFunction> FMassEntityTemplate::GetObjectFragmentInitializers() const
+{
+	return TemplateData.GetObjectFragmentInitializers();
+}
+
+inline const FString& FMassEntityTemplate::GetTemplateName() const
+{
+	return TemplateData.GetTemplateName();
+}
+
+inline const FMassArchetypeCompositionDescriptor& FMassEntityTemplate::GetCompositionDescriptor() const
+{
+	return TemplateData.GetCompositionDescriptor();
+}
+
+inline const FMassArchetypeSharedFragmentValues& FMassEntityTemplate::GetSharedFragmentValues() const
+{
+	return TemplateData.GetSharedFragmentValues();
+}
+
+inline TConstArrayView<FInstancedStruct> FMassEntityTemplate::GetInitialFragmentValues() const
+{
+	return TemplateData.GetInitialFragmentValues();
+}
+
+inline const FMassEntityTemplateData& FMassEntityTemplate::GetTemplateData() const
+{
+	return TemplateData;
+}
