[InstancedActors] Supporting Actors that contain multiple UStaticMeshComponents with their own relative transforms. Also support CustomDataFloats being passed along to UInstancedStaticMeshComponent::PerInstanceSMCustomData. Deprecate the version of FAdditionalSetupStepsFunction that takes in FISMComponentDescriptor& as a parameter, as it's already built into the FInstancedActorsVisualizationDesc

diff --git a/Engine/Plugins/Runtime/InstancedActors/Source/InstancedActors/Private/InstancedActorsData.cpp b/Engine/Plugins/Runtime/InstancedActors/Source/InstancedActors/Private/InstancedActorsData.cpp
--- a/Engine/Plugins/Runtime/InstancedActors/Source/InstancedActors/Private/InstancedActorsData.cpp
+++ b/Engine/Plugins/Runtime/InstancedActors/Source/InstancedActors/Private/InstancedActorsData.cpp
@@ -1796,7 +1796,8 @@ int32 UInstancedActorsData::GetEntityIndexFromCollisionIndex(const UInstancedSta
 		const int32 ISMComponentIndex = IAVisualization.ISMComponents.Find(const_cast<UInstancedStaticMeshComponent*>(&ISMComponent));
 		if (ISMComponentIndex != INDEX_NONE)
 		{
-			const FMassISMCSharedData* RelevantISMCData = RepresentationSubsystem->GetISMCSharedDataForDescriptionIndex(IAVisualization.MassStaticMeshDescHandle.ToIndex());
+			// We can't use GetISMCSharedDataForDescriptionIndex() with the MassStaticMeshDescHandle, since it corresponds to a given FStaticMeshInstanceVisualizationDesc, not its owned ISMC data
+			const FMassISMCSharedData* RelevantISMCData = RepresentationSubsystem->GetISMCSharedDataForInstancedStaticMesh(&ISMComponent);
 
 			ensureMsgf(RelevantISMCData && RelevantISMCData->GetISMComponent() == &ISMComponent, TEXT("We never expect to hit this ensure. Failing the test indicates serious mis-setup."));
 
diff --git a/Engine/Plugins/Runtime/InstancedActors/Source/InstancedActors/Private/InstancedActorsTypes.cpp b/Engine/Plugins/Runtime/InstancedActors/Source/InstancedActors/Private/InstancedActorsTypes.cpp
--- a/Engine/Plugins/Runtime/InstancedActors/Source/InstancedActors/Private/InstancedActorsTypes.cpp
+++ b/Engine/Plugins/Runtime/InstancedActors/Source/InstancedActors/Private/InstancedActorsTypes.cpp
@@ -95,24 +95,40 @@ FInstancedActorsVisualizationDesc::FInstancedActorsVisualizationDesc(const FInst
 }
 
 FInstancedActorsVisualizationDesc FInstancedActorsVisualizationDesc::FromActor(const AActor& ExemplarActor
-	, const FInstancedActorsVisualizationDesc::FAdditionalSetupStepsFunction& AdditionalSetupSteps)
+	, const FAdditionalSetupStepsFunction& AdditionalSetupSteps)
 {
-	FInstancedActorsVisualizationDesc Visualization;
+	return FromActor(ExemplarActor, [&AdditionalSetupSteps](const AActor& ExemplarActor, FInstancedActorsVisualizationDesc& OutVisualization)
+	{
+		if (OutVisualization.ISMComponentDescriptors.Num() > 0)
+		{
+			AdditionalSetupSteps(ExemplarActor, OutVisualization.ISMComponentDescriptors[0], OutVisualization);
+		}
+	});
+}
 
-	if (const UStaticMeshComponent* SourceStaticMeshComponent = ExemplarActor.GetComponentByClass<UStaticMeshComponent>())
+FInstancedActorsVisualizationDesc FInstancedActorsVisualizationDesc::FromActor(const AActor& ExemplarActor,
+	const FVisualizationDescSetupFunction& AdditionalSetupSteps)
+{
+	FInstancedActorsVisualizationDesc Visualization;
+	
+	ExemplarActor.ForEachComponent<UStaticMeshComponent>(/*bIncludeFromChildActors=*/false, [&Visualization](UStaticMeshComponent* SourceStaticMeshComponent)
 	{
 		const UStaticMesh* StaticMesh = SourceStaticMeshComponent->GetStaticMesh();
 		if (!IsValid(StaticMesh))
 		{
 			// No mesh = no visualization
-			return MoveTemp(Visualization);
+			return;
 		}
 
 		FISMComponentDescriptor& ISMComponentDescriptor = Visualization.ISMComponentDescriptors.AddDefaulted_GetRef();
 		ISMComponentDescriptor.InitFrom(SourceStaticMeshComponent);
 
-		AdditionalSetupSteps(ExemplarActor, ISMComponentDescriptor, Visualization);
-	}
+		// LocalTransform means local to the Actor/Entity, so we need to compute based on the UStaticMeshComponent's relative transform accordingly
+		// (in case this UStaticMeshComponent was a child of another UStaticMeshComponent within the Actor hierarchy)
+		ISMComponentDescriptor.LocalTransform = SourceStaticMeshComponent->GetComponentToWorld();
+	});
+
+	AdditionalSetupSteps(ExemplarActor, Visualization);
 
 	return MoveTemp(Visualization);
 }
@@ -130,12 +146,15 @@ FStaticMeshInstanceVisualizationDesc FInstancedActorsVisualizationDesc::ToMassVi
 
 		FMassStaticMeshInstanceVisualizationMeshDesc& MeshDesc = OutMassVisualizationDesc.Meshes.AddDefaulted_GetRef();
 		MeshDesc.Mesh = ISMComponentDescriptor.StaticMesh;
+		MeshDesc.LocalTransform = ISMComponentDescriptor.LocalTransform;
 		MeshDesc.bCastShadows = ISMComponentDescriptor.bCastShadow;
 		MeshDesc.Mobility = EComponentMobility::Stationary;
 		MeshDesc.MaterialOverrides = ISMComponentDescriptor.OverrideMaterials;
 		MeshDesc.ISMComponentClass = UInstancedStaticMeshComponent::StaticClass();
 	}
 
+	OutMassVisualizationDesc.CustomDataFloats = CustomDataFloats;
+
 	return MoveTemp(OutMassVisualizationDesc);
 }
 
diff --git a/Engine/Plugins/Runtime/InstancedActors/Source/InstancedActors/Public/InstancedActorsTypes.h b/Engine/Plugins/Runtime/InstancedActors/Source/InstancedActors/Public/InstancedActorsTypes.h
--- a/Engine/Plugins/Runtime/InstancedActors/Source/InstancedActors/Public/InstancedActorsTypes.h
+++ b/Engine/Plugins/Runtime/InstancedActors/Source/InstancedActors/Public/InstancedActorsTypes.h
@@ -139,16 +139,24 @@ struct INSTANCEDACTORS_API FInstancedActorsVisualizationDesc
 	UPROPERTY(VisibleAnywhere, Category = InstancedActors)
 	TArray<FISMComponentDescriptor> ISMComponentDescriptors;
 
-	using FAdditionalSetupStepsFunction = TFunctionRef<void(const AActor& /*ExemplarActor*/, FISMComponentDescriptor& /*NewISMComponentDescriptor*/, FInstancedActorsVisualizationDesc& /*OutVisualization*/)>;
+	UPROPERTY(VisibleAnywhere, Category = InstancedActors)
+	TArray<float> CustomDataFloats;
+
+	using FAdditionalSetupStepsFunction = TFunctionRef<void(const AActor& /*ExemplarActor*/, FISMComponentDescriptor& /*NewISMComponentDescriptor*/, FInstancedActorsVisualizationDesc& /*OutVisualization*/)>;	
+	using FVisualizationDescSetupFunction = TFunctionRef<void(const AActor& /*ExemplarActor*/, FInstancedActorsVisualizationDesc& /*OutVisualization*/)>;
+	
+	UE_DEPRECATED(5.6, "Using FAdditionalSetupStepsFunction which takes a FISMComponentDescriptor& parameter is deprecated. Use FVisualizationDescSetupFunction instead, since FInstancedActorsVisualizationDesc already contains the built TArray of FISMComponentDescriptor")
+	static FInstancedActorsVisualizationDesc FromActor(const AActor& ExemplarActor, const FAdditionalSetupStepsFunction& AdditionalSetupSteps);
+	
 	/**
 	 * Helper function to deduce appropriate instanced static mesh representation for an ActorClass exemplar actor.
 	 * @param ExemplarActor A fully constructed exemplar actor, including BP construction scripts up to (but not including) BeginPlay.
-	 * @param AdditionalSetupSteps function called once a new ISMComponentDescriptor is created, allowing custom code to 
+	 * @param AdditionalSetupFunction function called once a new ISMComponentDescriptor is created, allowing custom code to 
 	 *	perform additional set up steps. Note that this function will get called only if ExemplarActor has a 
 	 *	StaticMeshComponent with a valid StaticMesh configured.
 	 * @see UInstancedActorsSubsystem::GetOrCreateExemplarActor
-	 */
-	static FInstancedActorsVisualizationDesc FromActor(const AActor& ExemplarActor, const FAdditionalSetupStepsFunction& AdditionalSetupSteps = [](const AActor& /*ExemplarActor*/, FISMComponentDescriptor& /*NewISMComponentDescriptor*/, FInstancedActorsVisualizationDesc& /*OutVisualization*/){});
+	 */	
+	static FInstancedActorsVisualizationDesc FromActor(const AActor& ExemplarActor, const FVisualizationDescSetupFunction& AdditionalSetupFunction = [](const AActor& /*ExemplarActor*/, FInstancedActorsVisualizationDesc& /*OutVisualization*/){});
 
 	FStaticMeshInstanceVisualizationDesc ToMassVisualizationDesc() const;
 
diff --git a/Engine/Plugins/Runtime/MassGameplay/Source/MassRepresentation/Private/MassStationaryISMSwitcherProcessor.cpp b/Engine/Plugins/Runtime/MassGameplay/Source/MassRepresentation/Private/MassStationaryISMSwitcherProcessor.cpp
--- a/Engine/Plugins/Runtime/MassGameplay/Source/MassRepresentation/Private/MassStationaryISMSwitcherProcessor.cpp
+++ b/Engine/Plugins/Runtime/MassGameplay/Source/MassRepresentation/Private/MassStationaryISMSwitcherProcessor.cpp
@@ -105,6 +105,11 @@ void UMassStationaryISMSwitcherProcessor::ProcessContext(FMassExecutionContext&
 				{
 					NewRange->AddInstance(EntityHandle, Transform);
 				}
+
+				if (ISMInfo.GetDesc().CustomDataFloats.Num() > 0)
+				{
+					NewRange->AddBatchedCustomDataFloats(ISMInfo.GetDesc().CustomDataFloats, {});					
+				}
 			}
 
 			// consume "prev" data
diff --git a/Engine/Plugins/Runtime/MassGameplay/Source/MassRepresentation/Private/MassVisualizationComponent.cpp b/Engine/Plugins/Runtime/MassGameplay/Source/MassRepresentation/Private/MassVisualizationComponent.cpp
--- a/Engine/Plugins/Runtime/MassGameplay/Source/MassRepresentation/Private/MassVisualizationComponent.cpp
+++ b/Engine/Plugins/Runtime/MassGameplay/Source/MassRepresentation/Private/MassVisualizationComponent.cpp
@@ -134,7 +134,7 @@ FStaticMeshInstanceVisualizationDescHandle UMassVisualizationComponent::AddVisua
 
 		// 2. Stores one FMassISMCSharedData for each ISMC in ISMCSharedData
 		// NOTE: FMassISMCSharedData stores a reference to an ISMC and instance transform updates performed in the current frame.
-		FMassISMCSharedData& NewData = ISMCSharedData.FindOrAdd(ISMComponents[EntryIndex], FMassISMCSharedData(ISMComponents[EntryIndex], /*bInRequiresExternalInstanceIDTracking=*/true));
+		FMassISMCSharedData& NewData = ISMCSharedData.FindOrAdd(ISMComponents[EntryIndex], FMassISMCSharedData(ISMComponents[EntryIndex], /*bInRequiresExternalInstanceIDTracking=*/true, /*InTransformOffset=*/MeshDesc.LocalTransform));
 		// 3. Stores the newly created FMassISMCSharedData in a VisualHandle -> FMassISMCSharedData array
 		// NOTE: This means that one VisualHandle may reference multiple FMassISMCSharedData
 		InstancedStaticMeshInfos[VisualHandle.ToIndex()].AddISMComponent(NewData);
@@ -261,7 +261,7 @@ void UMassVisualizationComponent::ConstructStaticMeshComponents()
 
 				if (SharedData == nullptr)
 				{
-					SharedData = &ISMCSharedData.Add(ISMC, FMassISMCSharedData(ISMC, MeshDesc.bRequiresExternalInstanceIDTracking));
+					SharedData = &ISMCSharedData.Add(ISMC, FMassISMCSharedData(ISMC, MeshDesc.bRequiresExternalInstanceIDTracking, MeshDesc.LocalTransform));
 				}
 				else
 				{
@@ -448,8 +448,12 @@ void UMassVisualizationComponent::HandleChangesWithExternalIDTracking(UInstanced
 			for (int32 i = 0; i < EntityHandles.Num(); ++i)
 			{
 				SharedIdMap.GetByElementId(ElementIds[i]).Value = NewIds[i];
+
+				for (int32 j = 0; j < CustomFloatData.Num(); ++j)
+				{
+					ISMComponent.SetCustomDataValue(i, j, CustomFloatData[j]);
+				}
 			}
-			ensureMsgf(CustomFloatData.Num() == 0, TEXT("Custom floats not supported with this set up just yet."));
 		}
 	}
 
@@ -655,7 +659,10 @@ void FMassLODSignificanceRange::AddBatchedTransform(const FMassEntityHandle Enti
 		if (FMassISMCSharedData* SharedData = ISMCSharedDataPtr->GetAndMarkDirty(StaticMeshRefs[StaticMeshIndex]))
 		{
 			SharedData->UpdateInstanceIds.Add(EntityHandle);
-			SharedData->StaticMeshInstanceTransforms.Add(Transform);
+
+			FTransform AdjustedTransform{SharedData->LocalTransform * Transform};
+			
+			SharedData->StaticMeshInstanceTransforms.Add(AdjustedTransform);
 			SharedData->StaticMeshInstancePrevTransforms.Add(PrevTransform);
 		}
 	}
@@ -686,8 +693,11 @@ void FMassLODSignificanceRange::AddInstance(const FMassEntityHandle EntityHandle
 		if (FMassISMCSharedData* SharedData = ISMCSharedDataPtr->GetAndMarkDirty(StaticMeshRefs[StaticMeshIndex]))
 		{
 			SharedData->UpdateInstanceIds.Add(EntityHandle);
-			SharedData->StaticMeshInstanceTransforms.Add(Transform);
-			SharedData->StaticMeshInstancePrevTransforms.Add(Transform);
+
+			FTransform AdjustedTransform{SharedData->LocalTransform * Transform};
+			
+			SharedData->StaticMeshInstanceTransforms.Add(AdjustedTransform);
+			SharedData->StaticMeshInstancePrevTransforms.Add(AdjustedTransform);
 		}
 	}
 }
diff --git a/Engine/Plugins/Runtime/MassGameplay/Source/MassRepresentation/Public/MassRepresentationTypes.h b/Engine/Plugins/Runtime/MassGameplay/Source/MassRepresentation/Public/MassRepresentationTypes.h
--- a/Engine/Plugins/Runtime/MassGameplay/Source/MassRepresentation/Public/MassRepresentationTypes.h
+++ b/Engine/Plugins/Runtime/MassGameplay/Source/MassRepresentation/Public/MassRepresentationTypes.h
@@ -52,6 +52,10 @@ struct MASSREPRESENTATION_API FMassStaticMeshInstanceVisualizationMeshDesc
 	GENERATED_BODY()
 
 	FMassStaticMeshInstanceVisualizationMeshDesc();
+
+	/** A relative Transform from the parent, passed along to FMassISMCSharedData::LocalTransform */
+	UPROPERTY(EditAnywhere, Category = "Mass|Visual")
+	FTransform LocalTransform = FTransform::Identity;
 	
 	/** The static mesh visual representation */
 	UPROPERTY(EditAnywhere, Category = "Mass|Visual")
@@ -92,13 +96,14 @@ struct MASSREPRESENTATION_API FMassStaticMeshInstanceVisualizationMeshDesc
 	TSubclassOf<UInstancedStaticMeshComponent> ISMComponentClass;
 
 	bool operator==(const FMassStaticMeshInstanceVisualizationMeshDesc& Other) const
-	{
+	{		
 		return Mesh == Other.Mesh && 
 			MaterialOverrides == Other.MaterialOverrides && 
 			FMath::IsNearlyEqual(MinLODSignificance, Other.MinLODSignificance, KINDA_SMALL_NUMBER) &&
 			FMath::IsNearlyEqual(MaxLODSignificance, Other.MaxLODSignificance, KINDA_SMALL_NUMBER) &&
 			bCastShadows == Other.bCastShadows && 
-			Mobility == Other.Mobility;
+			Mobility == Other.Mobility &&
+			LocalTransform.Equals(Other.LocalTransform);
 	}
 	
 	friend FORCEINLINE uint32 GetTypeHash(const FMassStaticMeshInstanceVisualizationMeshDesc& MeshDesc)
@@ -107,6 +112,7 @@ struct MASSREPRESENTATION_API FMassStaticMeshInstanceVisualizationMeshDesc
 		Hash = PointerHash(MeshDesc.Mesh, Hash);
 		Hash = HashCombine(GetTypeHash(MeshDesc.bCastShadows), Hash);
 		Hash = HashCombine(GetTypeHash(MeshDesc.Mobility), Hash);
+		Hash = HashCombine(GetTypeHash(MeshDesc.LocalTransform), Hash);
 		for (UMaterialInterface* MaterialOverride : MeshDesc.MaterialOverrides)
 		{
 			if (MaterialOverride)
@@ -144,7 +150,11 @@ struct FStaticMeshInstanceVisualizationDesc : public FTableRowBase
 
 	/** Transform to offset the static meshes if not align the mass agent transform */
 	UPROPERTY(EditAnywhere, Category = "Mass|Visual", meta=(EditCondition="bUseTransformOffset"))
-	FTransform TransformOffset;
+	FTransform TransformOffset = FTransform::Identity;
+
+	/** Custom data that can eventually be propagated to UInstancedStaticMeshComponent::PerInstanceSMCustomData */
+	UPROPERTY(EditAnywhere, Category = "Mass|Visual")
+	TArray<float> CustomDataFloats;
 
 	bool operator==(const FStaticMeshInstanceVisualizationDesc& Other) const
 	{
@@ -232,8 +242,8 @@ PRAGMA_DISABLE_DEPRECATION_WARNINGS
 	{		
 	}
 
-	explicit FMassISMCSharedData(UInstancedStaticMeshComponent* InISMC, bool bInRequiresExternalInstanceIDTracking = false)
-		: ISMC(InISMC), bRequiresExternalInstanceIDTracking(bInRequiresExternalInstanceIDTracking)
+	explicit FMassISMCSharedData(UInstancedStaticMeshComponent* InISMC, bool bInRequiresExternalInstanceIDTracking = false, const FTransform& InTransformOffset = FTransform::Identity)
+		: LocalTransform(InTransformOffset), ISMC(InISMC), bRequiresExternalInstanceIDTracking(bInRequiresExternalInstanceIDTracking)
 	{
 	}
 
@@ -302,6 +312,9 @@ PRAGMA_ENABLE_DEPRECATION_WARNINGS
 protected:
 	friend FMassLODSignificanceRange;
 	friend UMassVisualizationComponent;
+
+	FTransform LocalTransform = FTransform::Identity;
+
 	/** Buffer holding current frame transforms for the static mesh instances, used to batch update the transforms */
 	TArray<FMassEntityHandle> UpdateInstanceIds;
 	TArray<FTransform> StaticMeshInstanceTransforms;
diff --git a/Engine/Source/Runtime/Engine/Private/ISMPartition/ISMComponentDescriptor.cpp b/Engine/Source/Runtime/Engine/Private/ISMPartition/ISMComponentDescriptor.cpp
--- a/Engine/Source/Runtime/Engine/Private/ISMPartition/ISMComponentDescriptor.cpp
+++ b/Engine/Source/Runtime/Engine/Private/ISMPartition/ISMComponentDescriptor.cpp
@@ -183,6 +183,7 @@ void FISMComponentDescriptorBase::InitFrom(const UStaticMeshComponent* Template,
 void FISMComponentDescriptor::InitFrom(const UStaticMeshComponent* Template, bool bInitBodyInstance)
 {
 	StaticMesh = Template->GetStaticMesh();
+	StaticMeshComponent = Template;
 	OverrideMaterials = Template->OverrideMaterials;
 	OverlayMaterial = Template->OverlayMaterial;
 	RuntimeVirtualTextures = Template->RuntimeVirtualTextures;
@@ -297,12 +298,13 @@ bool FISMComponentDescriptorBase::operator==(const FISMComponentDescriptorBase&
 }
 
 bool FISMComponentDescriptor::operator==(const FISMComponentDescriptor& Other) const
-{
+{	
 	return (Hash == 0 || Other.Hash == 0 || Hash == Other.Hash) && // Check hash first, other checks are in case of Hash collision
 		StaticMesh == Other.StaticMesh &&
 		OverrideMaterials == Other.OverrideMaterials &&
 		OverlayMaterial == Other.OverlayMaterial &&
 		RuntimeVirtualTextures == Other.RuntimeVirtualTextures &&
+		LocalTransform.Equals(Other.LocalTransform) &&
 		Super::operator==(Other);
 }
 
@@ -337,6 +339,7 @@ uint32 FISMComponentDescriptor::ComputeHash() const
 	CrcArchive << MutableSelf.OverrideMaterials;
 	CrcArchive << MutableSelf.OverlayMaterial;
 	CrcArchive << MutableSelf.RuntimeVirtualTextures;
+	CrcArchive << MutableSelf.LocalTransform;
 	Hash = CrcArchive.GetCrc();
 
 	return Hash;
diff --git a/Engine/Source/Runtime/Engine/Public/ISMPartition/ISMComponentDescriptor.h b/Engine/Source/Runtime/Engine/Public/ISMPartition/ISMComponentDescriptor.h
--- a/Engine/Source/Runtime/Engine/Public/ISMPartition/ISMComponentDescriptor.h
+++ b/Engine/Source/Runtime/Engine/Public/ISMPartition/ISMComponentDescriptor.h
@@ -261,6 +261,9 @@ public:
 	UPROPERTY(EditAnywhere, Category = "Component Settings", meta = (DisplayAfter = "ComponentClass"))
 	TObjectPtr<UStaticMesh> StaticMesh = nullptr;
 
+	UPROPERTY(EditAnywhere, Category = "Component Settings", meta = (DisplayAfter = "ComponentClass"))
+	TObjectPtr<const UStaticMeshComponent> StaticMeshComponent = nullptr;
+
 	UPROPERTY(EditAnywhere, Category = "Component Settings", meta = (DisplayAfter = "ComponentClass"))
 	TArray<TObjectPtr<UMaterialInterface>> OverrideMaterials;
 
@@ -269,6 +272,9 @@ public:
 
 	UPROPERTY(EditAnywhere, Category = "Component Settings", meta = (DisplayAfter = "ComponentClass"))
 	TArray<TObjectPtr<URuntimeVirtualTexture>> RuntimeVirtualTextures;
+
+	UPROPERTY(EditAnywhere, Category = "Component Settings", meta = (DisplayAfter = "ComponentClass"))
+	FTransform LocalTransform = FTransform::Identity;
 };
 
 USTRUCT()
