[Mass] Minor clean up, renamed some `FMassEntityManager System`function parameters

#rnx

diff --git a/Engine/Plugins/Runtime/InstancedActors/Source/InstancedActors/Public/InstancedActorsCommands.h b/Engine/Plugins/Runtime/InstancedActors/Source/InstancedActors/Public/InstancedActorsCommands.h
--- a/Engine/Plugins/Runtime/InstancedActors/Source/InstancedActors/Public/InstancedActorsCommands.h
+++ b/Engine/Plugins/Runtime/InstancedActors/Source/InstancedActors/Public/InstancedActorsCommands.h
@@ -59,12 +59,12 @@ protected:
 		return Super::GetAllocatedSize() + ManagersToResave.GetAllocatedSize();
 	}
 
-	virtual void Execute(FMassEntityManager& System) const override
+	virtual void Execute(FMassEntityManager& EntityManager) const override
 	{
 		TRACE_CPUPROFILER_EVENT_SCOPE(MassCommandAddFragmentInstancesAndResaveIAPersistence_Execute);
 
 		// Add / update fragments
-		Super::Execute(System);
+		Super::Execute(EntityManager);
 
 		// Resave Instanced Actor persistence for now-updated fragments
 		for (const TObjectPtr<AInstancedActorsManager>& ManagerToResave : ManagersToResave)
diff --git a/Engine/Plugins/Runtime/MassGameplay/Source/MassActors/Private/MassAgentSubsystem.cpp b/Engine/Plugins/Runtime/MassGameplay/Source/MassActors/Private/MassAgentSubsystem.cpp
--- a/Engine/Plugins/Runtime/MassGameplay/Source/MassActors/Private/MassAgentSubsystem.cpp
+++ b/Engine/Plugins/Runtime/MassGameplay/Source/MassActors/Private/MassAgentSubsystem.cpp
@@ -187,13 +187,13 @@ void UMassAgentSubsystem::UnregisterAgentComponent(UMassAgentComponent& AgentCom
 			    {
 				    // need to request via command buffer since we can't move entities while processing is happening
 					FMassArchetypeCompositionDescriptor Composition = AgentComp.GetPuppetSpecificAddition();
-					EntityManager->Defer().PushCommand<FMassDeferredRemoveCommand>([Entity, Composition](FMassEntityManager& System)
+					EntityManager->Defer().PushCommand<FMassDeferredRemoveCommand>([Entity, Composition](FMassEntityManager& InOutEntityManager)
 						{
-							if (System.IsEntityValid(Entity) == false)
+							if (InOutEntityManager.IsEntityValid(Entity) == false)
 							{
 								return;
 							}
-							System.RemoveCompositionFromEntity(Entity, Composition);
+							InOutEntityManager.RemoveCompositionFromEntity(Entity, Composition);
 						});
 			    }
 			    else
diff --git a/Engine/Plugins/Runtime/MassGameplay/Source/MassSignals/Private/MassSignalSubsystem.cpp b/Engine/Plugins/Runtime/MassGameplay/Source/MassSignals/Private/MassSignalSubsystem.cpp
--- a/Engine/Plugins/Runtime/MassGameplay/Source/MassSignals/Private/MassSignalSubsystem.cpp
+++ b/Engine/Plugins/Runtime/MassGameplay/Source/MassSignals/Private/MassSignalSubsystem.cpp
@@ -108,9 +108,9 @@ void UMassSignalSubsystem::SignalEntityDeferred(FMassExecutionContext& Context,
 void UMassSignalSubsystem::SignalEntitiesDeferred(FMassExecutionContext& Context, FName SignalName, TConstArrayView<FMassEntityHandle> Entities)
 {
 	checkf(Entities.Num() > 0, TEXT("Expecting entities to signal"));
-	Context.Defer().PushCommand<FMassDeferredSetCommand>([SignalName, InEntities = TArray<FMassEntityHandle>(Entities)](const FMassEntityManager& System)
+	Context.Defer().PushCommand<FMassDeferredSetCommand>([SignalName, InEntities = TArray<FMassEntityHandle>(Entities)](const FMassEntityManager& InOutEntityManager)
 	{
-		UMassSignalSubsystem* SignalSubsystem = UWorld::GetSubsystem<UMassSignalSubsystem>(System.GetWorld());
+		UMassSignalSubsystem* SignalSubsystem = UWorld::GetSubsystem<UMassSignalSubsystem>(InOutEntityManager.GetWorld());
 		SignalSubsystem->SignalEntities(SignalName, InEntities);
 	});
 
@@ -127,9 +127,9 @@ void UMassSignalSubsystem::DelaySignalEntityDeferred(FMassExecutionContext& Cont
 void UMassSignalSubsystem::DelaySignalEntitiesDeferred(FMassExecutionContext& Context, FName SignalName, TConstArrayView<FMassEntityHandle> Entities, const float DelayInSeconds)
 {
 	checkf(Entities.Num() > 0, TEXT("Expecting entities to signal"));
-	Context.Defer().PushCommand<FMassDeferredSetCommand>([SignalName, InEntities = TArray<FMassEntityHandle>(Entities), DelayInSeconds](const FMassEntityManager& System)
+	Context.Defer().PushCommand<FMassDeferredSetCommand>([SignalName, InEntities = TArray<FMassEntityHandle>(Entities), DelayInSeconds](const FMassEntityManager& InOutEntityManager)
 	{
-		UMassSignalSubsystem* SignalSubsystem = UWorld::GetSubsystem<UMassSignalSubsystem>(System.GetWorld());
+		UMassSignalSubsystem* SignalSubsystem = UWorld::GetSubsystem<UMassSignalSubsystem>(InOutEntityManager.GetWorld());
 		SignalSubsystem->DelaySignalEntities(SignalName, InEntities, DelayInSeconds);
 	});
 
diff --git a/Engine/Source/Developer/MassEntityTestSuite/Private/MassCommandsTest.cpp b/Engine/Source/Developer/MassEntityTestSuite/Private/MassCommandsTest.cpp
--- a/Engine/Source/Developer/MassEntityTestSuite/Private/MassCommandsTest.cpp
+++ b/Engine/Source/Developer/MassEntityTestSuite/Private/MassCommandsTest.cpp
@@ -239,7 +239,7 @@ struct FCommands_DeferredFunction : FEntityTestBase
 			FMassEntityView View(IntsArchetype, Entity);
 			View.GetFragmentData<FTestFragment_Int>().Value = Offset + i++;
 
-			EntityManager->Defer().PushCommand<FMassDeferredSetCommand>([Entity, Archetype = IntsArchetype, Offset](FMassEntityManager& System)
+			EntityManager->Defer().PushCommand<FMassDeferredSetCommand>([Entity, Archetype = IntsArchetype, Offset](FMassEntityManager&)
 				{
 					FMassEntityView View(Archetype, Entity);
 					View.GetFragmentData<FTestFragment_Int>().Value -= Offset;
