[MassGameplay] Switched UMassDebugVisualizationComponent from using HISM components over to ISM components

This change fixes issues with DebugVisualization trait

diff --git a/Engine/Plugins/Runtime/MassGameplay/Source/MassGameplayDebug/Private/DebugVisLocationProcessor.cpp b/Engine/Plugins/Runtime/MassGameplay/Source/MassGameplayDebug/Private/DebugVisLocationProcessor.cpp
--- a/Engine/Plugins/Runtime/MassGameplay/Source/MassGameplayDebug/Private/DebugVisLocationProcessor.cpp
+++ b/Engine/Plugins/Runtime/MassGameplay/Source/MassGameplayDebug/Private/DebugVisLocationProcessor.cpp
@@ -4,7 +4,7 @@
 #include "MassDebuggerSubsystem.h"
 #include "MassDebugVisualizationComponent.h"
 #include "MassCommonFragments.h"
-#include "Components/HierarchicalInstancedStaticMeshComponent.h"
+#include "Components/InstancedStaticMeshComponent.h"
 #include "MassEntityManager.h"
 #include "MassExecutionContext.h"
 #include "MassMovementFragments.h"
@@ -59,7 +59,7 @@ void UDebugVisLocationProcessor::Execute(FMassEntityManager& EntityManager, FMas
 				UMassDebuggerSubsystem& Debugger = Context.GetMutableSubsystemChecked<UMassDebuggerSubsystem>();
 				UMassDebugVisualizationComponent* Visualizer = Debugger.GetVisualizationComponent();
 				check(Visualizer);
-				TArrayView<UHierarchicalInstancedStaticMeshComponent* const> VisualDataISMCs = Visualizer->GetVisualDataISMCs();
+				TArrayView<UInstancedStaticMeshComponent* const> VisualDataISMCs = Visualizer->GetVisualDataISMCs();
 				if (VisualDataISMCs.Num() > 0)
 				{
 					const int32 NumEntities = Context.GetNumEntities();
diff --git a/Engine/Plugins/Runtime/MassGameplay/Source/MassGameplayDebug/Private/MassDebugVisualizationComponent.cpp b/Engine/Plugins/Runtime/MassGameplay/Source/MassGameplayDebug/Private/MassDebugVisualizationComponent.cpp
--- a/Engine/Plugins/Runtime/MassGameplay/Source/MassGameplayDebug/Private/MassDebugVisualizationComponent.cpp
+++ b/Engine/Plugins/Runtime/MassGameplay/Source/MassGameplayDebug/Private/MassDebugVisualizationComponent.cpp
@@ -1,7 +1,7 @@
 // Copyright Epic Games, Inc. All Rights Reserved.
 
 #include "MassDebugVisualizationComponent.h"
-#include "Components/HierarchicalInstancedStaticMeshComponent.h"
+#include "Components/InstancedStaticMeshComponent.h"
 
 
 #if WITH_EDITORONLY_DATA
@@ -18,7 +18,7 @@ void UMassDebugVisualizationComponent::PostInitProperties()
 
 void UMassDebugVisualizationComponent::DirtyVisuals()
 {
-	for (UHierarchicalInstancedStaticMeshComponent* ISM : VisualDataISMCs)
+	for (UInstancedStaticMeshComponent* ISM : VisualDataISMCs)
 	{
 		if (ensure(ISM))
 		{
@@ -45,25 +45,25 @@ void UMassDebugVisualizationComponent::ConstructVisualComponent()
 	AActor* ActorOwner = GetOwner();
 	check(ActorOwner);
 	
-	// add HISMCs only for types not added yet
+	// add ISMComponents only for types not added yet
 	for (int NewTypeIndex = VisualDataISMCs.Num(); NewTypeIndex < VisualDataTable.Num(); ++NewTypeIndex)
 	{
 		const FAgentDebugVisualization& VisualData = VisualDataTable[NewTypeIndex];
-		UHierarchicalInstancedStaticMeshComponent* HISMC = NewObject<UHierarchicalInstancedStaticMeshComponent>(ActorOwner);
-		HISMC->SetStaticMesh(VisualData.Mesh);
+		UInstancedStaticMeshComponent* ISMComponent = NewObject<UInstancedStaticMeshComponent>(ActorOwner);
+		ISMComponent->SetStaticMesh(VisualData.Mesh);
 		if (VisualData.MaterialOverride != nullptr)
 		{
-			HISMC->SetMaterial(0, VisualData.MaterialOverride);
+			ISMComponent->SetMaterial(0, VisualData.MaterialOverride);
 		}
 
-		HISMC->SetCullDistances(VisualData.VisualNearCullDistance, VisualData.VisualFarCullDistance);
-		HISMC->SetupAttachment(ActorOwner->GetRootComponent());
-		HISMC->SetCanEverAffectNavigation(false);
-		HISMC->bDisableCollision = true;
-		HISMC->SetCastShadow(false);
-		HISMC->RegisterComponent();
+		ISMComponent->SetCullDistances(VisualData.VisualNearCullDistance, VisualData.VisualFarCullDistance);
+		ISMComponent->SetupAttachment(ActorOwner->GetRootComponent());
+		ISMComponent->SetCanEverAffectNavigation(false);
+		ISMComponent->bDisableCollision = true;
+		ISMComponent->SetCastShadow(false);
+		ISMComponent->RegisterComponent();
 
-		VisualDataISMCs.Add(HISMC);
+		VisualDataISMCs.Add(ISMComponent);
 	}
 }
 
@@ -76,7 +76,7 @@ uint16 UMassDebugVisualizationComponent::AddDebugVisType(const FAgentDebugVisual
 
 void UMassDebugVisualizationComponent::Clear()
 {
-	for (UHierarchicalInstancedStaticMeshComponent* ISM : VisualDataISMCs)
+	for (UInstancedStaticMeshComponent* ISM : VisualDataISMCs)
 	{
 		if (ensure(ISM))
 		{
diff --git a/Engine/Plugins/Runtime/MassGameplay/Source/MassGameplayDebug/Public/MassDebugVisualizationComponent.h b/Engine/Plugins/Runtime/MassGameplay/Source/MassGameplayDebug/Public/MassDebugVisualizationComponent.h
--- a/Engine/Plugins/Runtime/MassGameplay/Source/MassGameplayDebug/Public/MassDebugVisualizationComponent.h
+++ b/Engine/Plugins/Runtime/MassGameplay/Source/MassGameplayDebug/Public/MassDebugVisualizationComponent.h
@@ -7,7 +7,7 @@
 #include "MassDebugVisualizationComponent.generated.h"
 
 
-class UHierarchicalInstancedStaticMeshComponent;
+class UInstancedStaticMeshComponent;
 
 /** meant to be created procedurally and owned by an AMassDebugVisualizer instance. Will ensure if placed on a different type of actor */
 UCLASS()
@@ -22,7 +22,7 @@ public:
 	int32 AddDebugVisInstance(const uint16 VisualType);
 	/** returns index to the newly created VisualDataTable entry */
 	uint16 AddDebugVisType(const FAgentDebugVisualization& Data);
-	TArrayView<UHierarchicalInstancedStaticMeshComponent* const> GetVisualDataISMCs() { return MakeArrayView(VisualDataISMCs); }
+	TArrayView<UInstancedStaticMeshComponent* const> GetVisualDataISMCs() { return MakeArrayView(VisualDataISMCs); }
 
 	void Clear();
 protected:
@@ -35,6 +35,6 @@ protected:
 	TArray<FAgentDebugVisualization> VisualDataTable;
 
 	UPROPERTY(Transient)
-	TArray<TObjectPtr<UHierarchicalInstancedStaticMeshComponent>> VisualDataISMCs;
+	TArray<TObjectPtr<UInstancedStaticMeshComponent>> VisualDataISMCs;
 #endif // WITH_EDITORONLY_DATA
 };
