[InstancedActors] In UMassVisualizationComponent::HandleChangesWithExternalIDTracking(), properly set UInstancedStaticMeshComponent::NumCustomDataFloats (which is meant to be the number of CustomDataFloats per instance), as is done in EndVisualChanges().Added a new UInstancedStaticMeshComponent::SetCustomData() overload that can batch add custom data floats for a range of instance indices, and properly update the newly updated instances

diff --git a/Engine/Plugins/Runtime/MassGameplay/Source/MassRepresentation/Private/MassVisualizationComponent.cpp b/Engine/Plugins/Runtime/MassGameplay/Source/MassRepresentation/Private/MassVisualizationComponent.cpp
--- a/Engine/Plugins/Runtime/MassGameplay/Source/MassRepresentation/Private/MassVisualizationComponent.cpp
+++ b/Engine/Plugins/Runtime/MassGameplay/Source/MassRepresentation/Private/MassVisualizationComponent.cpp
@@ -398,7 +398,7 @@ void UMassVisualizationComponent::HandleChangesWithExternalIDTracking(UInstanced
 
 	// NOTE: This code path is designed to only perform Adds, never updates so updates are filtered out along with duplicates.
 	TArray<FMassEntityHandle>& EntityHandles = SharedData.UpdateInstanceIds;
-	if (EntityHandles.Num())
+	if (!EntityHandles.IsEmpty())
 	{
 		INC_DWORD_STAT_BY(STAT_Mass_VisualizationComponent_InstancesAddedNum, EntityHandles.Num());
 
@@ -427,33 +427,30 @@ void UMassVisualizationComponent::HandleChangesWithExternalIDTracking(UInstanced
 		{
 			check(ElementIds.Num() == EntityHandles.Num());
 
-			const TConstArrayView<FMassEntityHandle> InstanceIds = SharedData.UpdateInstanceIds;
-
 			const TArray<FTransform>& InstanceTransforms = SharedData.GetStaticMeshInstanceTransformsArray();
-			const int32 InNumCustomDataFloats = SharedData.GetStaticMeshInstanceCustomFloats().Num();
 			TConstArrayView<float> CustomFloatData = SharedData.GetStaticMeshInstanceCustomFloats();
+			const int32 NumCustomDataFloatsPerInstance = CustomFloatData.Num() / EntityHandles.Num();
 
 			// if these are the first entities we're adding we need to set NumCustomDataFloats so that the PerInstanceSMCustomData
-			// gets populated properly by the AddInstancesInternal call below
+			// gets populated properly by the AddInstancesById call below
 			const int32 StartingCount = ISMComponent.GetNumInstances();
-			const bool bInitiallyEmpty = (StartingCount == 0); 
-			if (StartingCount == 0 && CustomFloatData.Num() && ISMComponent.Mobility != EComponentMobility::Static)
+			const bool bInitiallyEmpty = (StartingCount == 0);
+			ensure(bInitiallyEmpty || ISMComponent.NumCustomDataFloats == NumCustomDataFloatsPerInstance);
+			if (bInitiallyEmpty && NumCustomDataFloatsPerInstance > 0 && ISMComponent.Mobility != EComponentMobility::Static)
 			{
-				ISMComponent.SetNumCustomDataFloats(InNumCustomDataFloats);
+				ISMComponent.SetNumCustomDataFloats(NumCustomDataFloatsPerInstance);
 			}
 
 			check(EntityHandles.Num() == InstanceTransforms.Num());
 			TArray<FPrimitiveInstanceId> NewIds = ISMComponent.AddInstancesById(InstanceTransforms, /*bWorldSpace=*/true, /*bUpdateNavigation =*/bInitiallyEmpty);
 			check(EntityHandles.Num() == NewIds.Num());
-			for (int32 i = 0; i < EntityHandles.Num(); ++i)
+			int32 EntityHandleToUpdateIndex = StartingCount;
+            for (int32 EntityHandleIndex = 0; EntityHandleIndex < EntityHandles.Num(); ++EntityHandleIndex, ++EntityHandleToUpdateIndex)
 			{
-				SharedIdMap.GetByElementId(ElementIds[i]).Value = NewIds[i];
-
-				for (int32 j = 0; j < CustomFloatData.Num(); ++j)
-				{
-					ISMComponent.SetCustomDataValue(i, j, CustomFloatData[j]);
-				}
+				SharedIdMap.GetByElementId(ElementIds[EntityHandleIndex]).Value = NewIds[EntityHandleIndex];
 			}
+
+			ISMComponent.SetCustomData(StartingCount, StartingCount + EntityHandles.Num() - 1, CustomFloatData);
 		}
 	}
 
diff --git a/Engine/Source/Runtime/Engine/Classes/Components/InstancedStaticMeshComponent.h b/Engine/Source/Runtime/Engine/Classes/Components/InstancedStaticMeshComponent.h
--- a/Engine/Source/Runtime/Engine/Classes/Components/InstancedStaticMeshComponent.h
+++ b/Engine/Source/Runtime/Engine/Classes/Components/InstancedStaticMeshComponent.h
@@ -317,7 +317,8 @@ public:
 	ENGINE_API virtual void SetNumCustomDataFloats(int32 InNumCustomDataFloats);
 
 	/** Per Instance Custom Data */
-	ENGINE_API virtual bool SetCustomData(int32 InstanceIndex, TArrayView<const float> CustomDataFloats, bool bMarkRenderStateDirty = false); 
+	ENGINE_API virtual bool SetCustomData(int32 InstanceIndex, TArrayView<const float> CustomDataFloats, bool bMarkRenderStateDirty = false);
+	ENGINE_API virtual bool SetCustomData(int32 InstanceIndexStart, int32 InstanceIndexEnd, TConstArrayView<float> CustomDataFloats, bool bMarkRenderStateDirty = false);
 
 	/** Preallocated memory to include the new added instances count, to prevent reallloc during the add operation. */
 	ENGINE_API virtual void PreAllocateInstancesMemory(int32 AddedInstanceCount);
diff --git a/Engine/Source/Runtime/Engine/Private/InstancedStaticMesh.cpp b/Engine/Source/Runtime/Engine/Private/InstancedStaticMesh.cpp
--- a/Engine/Source/Runtime/Engine/Private/InstancedStaticMesh.cpp
+++ b/Engine/Source/Runtime/Engine/Private/InstancedStaticMesh.cpp
@@ -3728,6 +3728,46 @@ bool UInstancedStaticMeshComponent::SetCustomData(int32 InstanceIndex, TArrayVie
 	return true;
 }
 
+bool UInstancedStaticMeshComponent::SetCustomData(int32 InstanceIndexStart, int32 InstanceIndexEnd, TConstArrayView<float> CustomDataFloats, bool bMarkRenderStateDirty)
+{
+	const int32 CustomDataIndexStart = InstanceIndexStart * NumCustomDataFloats;
+	checkf(PerInstanceSMCustomData.Num() >= (CustomDataIndexStart + CustomDataFloats.Num()) && CustomDataFloats.GetTypeSize() == PerInstanceSMCustomData.GetTypeSize(), TEXT("CustomDataFloats and offset must fit within PerInstanceSMCustomData"));
+	
+	if (CustomDataFloats.IsEmpty())
+	{
+		return false;
+	}
+	
+	if (!ensureMsgf(InstanceIndexStart <= InstanceIndexEnd && InstanceIndexStart >= 0 && InstanceIndexEnd < PerInstanceSMData.Num(), TEXT("InstanceIndexStart and InstanceIndexEnd need to be valid indices for PerInstanceSMData")))
+	{
+		return false;
+	}
+
+	if (!ensureMsgf((InstanceIndexEnd - InstanceIndexStart + 1) * NumCustomDataFloats == CustomDataFloats.Num(), TEXT("CustomDataFloats ArrayNum must match the number of Instances * NumCustomDataFloats")))
+	{
+		return false;
+	}
+
+	if (bMarkRenderStateDirty)
+	{
+		Modify();
+	}
+	
+	FMemory::Memcpy(&PerInstanceSMCustomData[CustomDataIndexStart], CustomDataFloats.GetData(), CustomDataFloats.Num() * CustomDataFloats.GetTypeSize());
+
+	for (int32 InstanceIndex = InstanceIndexStart; InstanceIndex <= InstanceIndexEnd; ++InstanceIndex)
+	{
+		PrimitiveInstanceDataManager.CustomDataChanged(InstanceIndex);
+	}
+	
+	if (bMarkRenderStateDirty)
+	{
+		MarkRenderStateDirty();
+	}
+	
+	return true;
+}
+
 void UInstancedStaticMeshComponent::SetNumCustomDataFloats(int32 InNumCustomDataFloats)
 {
 	if (FMath::Max(InNumCustomDataFloats, 0) != NumCustomDataFloats)
