Fix non-deterministic dependencies for FMassEntityConfig : The dependencies for blueprints contain the default object of all the classes it uses. FMassEntityConfig has a guid initialized to a new value in the constructor. It means the class default object has a different value everytime the executable runs and it creates a different dependency.
To prevent this, the guid is left to 0 for cdo and archetype. It's initialized to a new value in every other cases.

#rnx
#incrementalcook

diff --git a/Engine/Plugins/Runtime/MassGameplay/Source/MassSpawner/Private/MassEntityConfigAsset.cpp b/Engine/Plugins/Runtime/MassGameplay/Source/MassSpawner/Private/MassEntityConfigAsset.cpp
--- a/Engine/Plugins/Runtime/MassGameplay/Source/MassSpawner/Private/MassEntityConfigAsset.cpp
+++ b/Engine/Plugins/Runtime/MassGameplay/Source/MassSpawner/Private/MassEntityConfigAsset.cpp
@@ -28,7 +28,12 @@ FMassEntityConfig::FMassEntityConfig()
 FMassEntityConfig::FMassEntityConfig(UObject& InOwner)
 	: ConfigOwner(&InOwner)
 {
-	ConfigGuid = FGuid::NewGuid();
+	// Only initialize the guid if this is not the class default object or archetype. The class default object can be serialized
+	// in cook dependencies for blueprints and needs to be deterministic. So leave it to 0 if it's the CDO.
+	if (!InOwner.HasAnyFlags(RF_ArchetypeObject | RF_ClassDefaultObject))
+	{
+		ConfigGuid = FGuid::NewGuid();
+	}
 }
 
 UMassEntityTraitBase* FMassEntityConfig::FindTraitInternal(TSubclassOf<UMassEntityTraitBase> TraitClass, const bool bExactMatch) const
