[MassSpawner] assigning deterministic guids, including config associated to archetype owners

diff --git a/Engine/Plugins/Runtime/MassGameplay/Source/MassSpawner/Private/MassEntityConfigAsset.cpp b/Engine/Plugins/Runtime/MassGameplay/Source/MassSpawner/Private/MassEntityConfigAsset.cpp
--- a/Engine/Plugins/Runtime/MassGameplay/Source/MassSpawner/Private/MassEntityConfigAsset.cpp
+++ b/Engine/Plugins/Runtime/MassGameplay/Source/MassSpawner/Private/MassEntityConfigAsset.cpp
@@ -28,12 +28,7 @@ FMassEntityConfig::FMassEntityConfig()
 FMassEntityConfig::FMassEntityConfig(UObject& InOwner)
 	: ConfigOwner(&InOwner)
 {
-	// Only initialize the guid if this is not the class default object or archetype. The class default object can be serialized
-	// in cook dependencies for blueprints and needs to be deterministic. So leave it to 0 if it's the CDO.
-	if (!InOwner.HasAnyFlags(RF_ArchetypeObject | RF_ClassDefaultObject))
-	{
-		ConfigGuid = FGuid::NewGuid();
-	}
+	ConfigGuid = FGuid::NewDeterministicGuid(InOwner.GetPathName());
 }
 
 UMassEntityTraitBase* FMassEntityConfig::FindTraitInternal(TSubclassOf<UMassEntityTraitBase> TraitClass, const bool bExactMatch) const
