Fix MassMoverDivergence debug.
DebugPreviousValue was not stored anymore because UMassToNavMoverTranslator is not running anymore.

diff --git a/Engine/Plugins/Runtime/MassGameplay/Source/MassMovement/Private/MassMovementProcessors.cpp b/Engine/Plugins/Runtime/MassGameplay/Source/MassMovement/Private/MassMovementProcessors.cpp
--- a/Engine/Plugins/Runtime/MassGameplay/Source/MassMovement/Private/MassMovementProcessors.cpp
+++ b/Engine/Plugins/Runtime/MassGameplay/Source/MassMovement/Private/MassMovementProcessors.cpp
@@ -102,6 +102,9 @@ void UMassApplyMovementProcessor::Execute(FMassEntityManager& EntityManager,
 			{
 				Velocity.Value = FVector::ZeroVector;
 			}
+
+			// Keep as "expected value" for next frame.
+			Velocity.DebugPreviousValue = Velocity.Value;
 #endif // WITH_MASSGAMEPLAY_DEBUG
 
 			FVector CurrentLocation = CurrentTransform.GetLocation();
