[MassGameplay] improved detection and reporting of LODCollector class issues with LODCollectorTraits

diff --git a/Engine/Plugins/Runtime/MassGameplay/Source/MassLOD/Private/MassLODTrait.cpp b/Engine/Plugins/Runtime/MassGameplay/Source/MassLOD/Private/MassLODTrait.cpp
--- a/Engine/Plugins/Runtime/MassGameplay/Source/MassLOD/Private/MassLODTrait.cpp
+++ b/Engine/Plugins/Runtime/MassGameplay/Source/MassLOD/Private/MassLODTrait.cpp
@@ -15,6 +15,8 @@ namespace UE::Mass::Private
 	{
 		check(ProcessorClass);
 
+		bool bValidCollectorActive = false;
+
 		// @todo make this processor class configurable. Could be something like:
 		// * every processor declared "identifying tag", a unique combination that would identify given processor class
 		// * traits can add "required processors"
@@ -22,20 +24,48 @@ namespace UE::Mass::Private
 		// * we can also report required processors that are marked as not auto-added.
 		//	-- downside of that last, similar to the one bellow, is that disabling a given processor might be deliberate and it can still be enabled at runtime
 		const UMassProcessor* LODCollectorProcessor = ProcessorClass->GetDefaultObject<UMassProcessor>();
-		if (LODCollectorProcessor)
+		if (ensureMsgf(LODCollectorProcessor, TEXT("Failed to fetch CDO from class %s"), *ProcessorClass->GetName()))
 		{
-			if (LODCollectorProcessor->ShouldAutoAddToGlobalList() == false)
+			if (LODCollectorProcessor->ShouldAutoAddToGlobalList())
 			{
-				CA_ASSUME(LODCollectorProcessor->GetClass());
-				UE_LOG(LogMassLOD, Error, TEXT("Using %s trait while required processor %s is not enabled by default")
-					, *Trait->GetName(), *LODCollectorProcessor->GetClass()->GetName());
-				return false;
+				bValidCollectorActive = true;
+			}
+			else
+			{
+				// look for subclasses
+				TArray<UClass*> Results;
+				GetDerivedClasses(ProcessorClass, Results, /*bRecursive=*/true);
+				for (UClass* Class : Results)
+				{
+					check(Class);
+					const UMassProcessor* LODCollectorSubclassProcessor = Class->GetDefaultObject<UMassProcessor>();
+					check(LODCollectorSubclassProcessor);
+					if (LODCollectorSubclassProcessor->ShouldAutoAddToGlobalList())
+					{
+						bValidCollectorActive = true;
+						break;
+					}
+				}
+
+				if (bValidCollectorActive == false)
+				{
+					TStringBuilder<256> SubclassesStringBuilder;
+					if (Results.Num() > 0)
+					{
+						SubclassesStringBuilder << ", nor any of its subclasses: ";
+						for (UClass* Class : Results)
+						{
+							SubclassesStringBuilder << GetNameSafe(Class) << ", ";
+						}
+					}
+
+					UE_LOG(LogMassLOD, Error, TEXT("Using %s trait while the required processor %s%s is not enabled by default")
+						, *Trait->GetName(), *GetNameSafe(LODCollectorProcessor->GetClass()), SubclassesStringBuilder.ToString());
+				}
 			}
-
-			return true;
 		}
-		checkf(LODCollectorProcessor, TEXT("Failed to fetch CDO from class %s"), *ProcessorClass->GetName());
-		return false;
+		
+		return bValidCollectorActive;
 	}
 }
 
