[MassGameplay] Marked UMassCapsuleTransformToMassTranslator as "game-thread only" to avoid data races when the observed component is modified in game thread as the data is being accessed by the processor

diff --git a/Engine/Plugins/Runtime/MassGameplay/Source/MassActors/Private/Translators/MassCapsuleComponentTranslators.cpp b/Engine/Plugins/Runtime/MassGameplay/Source/MassActors/Private/Translators/MassCapsuleComponentTranslators.cpp
--- a/Engine/Plugins/Runtime/MassGameplay/Source/MassActors/Private/Translators/MassCapsuleComponentTranslators.cpp
+++ b/Engine/Plugins/Runtime/MassGameplay/Source/MassActors/Private/Translators/MassCapsuleComponentTranslators.cpp
@@ -16,6 +16,9 @@ UMassCapsuleTransformToMassTranslator::UMassCapsuleTransformToMassTranslator()
 	ExecutionFlags = (int32)EProcessorExecutionFlags::AllNetModes;
 	RequiredTags.Add<FMassCapsuleTransformCopyToMassTag>();
 	ExecutionOrder.ExecuteInGroup = UE::Mass::ProcessorGroupNames::SyncWorldToMass;
+	// @todo won't be required anymore once component ticking is run through Mass, or Mass has a way to synchronize with that.
+	// Same comment to all the component translators
+	bRequiresGameThreadExecution = true;
 }
 
 void UMassCapsuleTransformToMassTranslator::ConfigureQueries(const TSharedRef<FMassEntityManager>& EntityManager)
