PR #13066: Make Meshes array elements easier for Devs to identify

#rnx

diff --git a/Engine/Plugins/Runtime/MassGameplay/Source/MassRepresentation/Public/MassRepresentationTypes.h b/Engine/Plugins/Runtime/MassGameplay/Source/MassRepresentation/Public/MassRepresentationTypes.h
--- a/Engine/Plugins/Runtime/MassGameplay/Source/MassRepresentation/Public/MassRepresentationTypes.h
+++ b/Engine/Plugins/Runtime/MassGameplay/Source/MassRepresentation/Public/MassRepresentationTypes.h
@@ -141,7 +141,7 @@ struct FStaticMeshInstanceVisualizationDesc : public FTableRowBase
 	 * Mesh descriptions. These will be instanced together using the same transform for each, to 
 	 * visualize the given agent.
 	 */
-	UPROPERTY(EditAnywhere, Category = "Mass|Visual")
+	UPROPERTY(EditAnywhere, Category = "Mass|Visual", meta=(TitleProperty="{Mesh}"))
 	TArray<FMassStaticMeshInstanceVisualizationMeshDesc> Meshes;
 
 	/** Boolean to enable code to transform the static meshes if not align the mass agent transform */
