[MassGameplayDebugger] fixed selected entity details being cleared before getting read by the gameplay debugger

diff --git a/Engine/Plugins/Runtime/MassGameplay/Source/MassGameplayDebug/Private/MassDebuggerSubsystem.cpp b/Engine/Plugins/Runtime/MassGameplay/Source/MassGameplayDebug/Private/MassDebuggerSubsystem.cpp
--- a/Engine/Plugins/Runtime/MassGameplay/Source/MassGameplayDebug/Private/MassDebuggerSubsystem.cpp
+++ b/Engine/Plugins/Runtime/MassGameplay/Source/MassGameplayDebug/Private/MassDebuggerSubsystem.cpp
@@ -57,7 +57,6 @@ void UMassDebuggerSubsystem::PreTickProcessors()
 {
 	Entities.Reset();
 	Locations.Reset();
-	SelectedEntityDetails.Empty();
 }
 
 void UMassDebuggerSubsystem::OnProcessingPhaseStarted(const float DeltaSeconds)
@@ -73,6 +72,12 @@ void UMassDebuggerSubsystem::SetSelectedEntity(const FMassEntityHandle InSelecte
 
 void UMassDebuggerSubsystem::AppendSelectedEntityInfo(const FString& Info)
 {
+	if (UpdateFrameNumber != GFrameNumber)
+	{
+		SelectedEntityDetails.Empty();
+		UpdateFrameNumber = GFrameNumber;
+	}
+
 	SelectedEntityDetails += Info;
 }
 
diff --git a/Engine/Plugins/Runtime/MassGameplay/Source/MassGameplayDebug/Public/MassDebuggerSubsystem.h b/Engine/Plugins/Runtime/MassGameplay/Source/MassGameplayDebug/Public/MassDebuggerSubsystem.h
--- a/Engine/Plugins/Runtime/MassGameplay/Source/MassGameplayDebug/Public/MassDebuggerSubsystem.h
+++ b/Engine/Plugins/Runtime/MassGameplay/Source/MassGameplayDebug/Public/MassDebuggerSubsystem.h
@@ -65,6 +65,7 @@ protected:
 	TArray<FVector> Locations;
 	FMassEntityHandle SelectedEntity;
 	FString SelectedEntityDetails;
+	uint64 UpdateFrameNumber = INDEX_NONE;
 
 	UPROPERTY(Transient)
 	TObjectPtr<UMassDebugVisualizationComponent> VisualizationComponent;
