[MassGameplay] Added FMassRepresentationParameters::Identical to implement logically-correct FMassRepresentationParameters instances comparison.
The operation by design ignores CachedDefaultRepresentationType and CachedRepresentationActorManagement.
Also moved some member properties around - the operation reduced the memory footprint of the type by 8 bytes.

diff --git a/Engine/Plugins/Runtime/MassGameplay/Source/MassRepresentation/Public/MassRepresentationFragments.h b/Engine/Plugins/Runtime/MassGameplay/Source/MassRepresentation/Public/MassRepresentationFragments.h
--- a/Engine/Plugins/Runtime/MassGameplay/Source/MassRepresentation/Public/MassRepresentationFragments.h
+++ b/Engine/Plugins/Runtime/MassGameplay/Source/MassRepresentation/Public/MassRepresentationFragments.h
@@ -112,6 +112,17 @@ struct FMassRepresentationParameters : public FMassConstSharedFragment
 	UPROPERTY(EditAnywhere, Category = "Mass|Representation", config)
 	EMassRepresentationType LODRepresentation[EMassLOD::Max] = { EMassRepresentationType::HighResSpawnedActor, EMassRepresentationType::LowResSpawnedActor, EMassRepresentationType::StaticMeshInstance, EMassRepresentationType::None };
 
+	/** World Partition grid name to test collision against, default None will be the main grid */
+	UPROPERTY(EditAnywhere, Category = "Mass|Representation", config)
+	FName WorldPartitionGridNameContainingCollision;
+
+	/** At what rate should the not visible entity be updated in seconds */
+	UPROPERTY(EditAnywhere, Category = "Mass|Visualization", config)
+	float NotVisibleUpdateRate = 0.5f;
+
+	// -- all propeties comparable with Memcpy go above this line. 
+	// -- all other properties need to be manually handled in Identical implementation (as long as they are not "transient")
+
 	/** 
 	 * If true, forces UMassRepresentationProcessor to override the WantedRepresentationType to actor representation whenever an external (non Mass owned)
 	 * actor is set on an entitie's FMassActorFragment fragment. If / when the actor fragment is reset, WantedRepresentationType resumes selecting the 
@@ -140,14 +151,6 @@ struct FMassRepresentationParameters : public FMassConstSharedFragment
 	uint8 bCanModifyRepresentationActorManagementClass : 1 = true;
 #endif // WITH_EDITORONLY_DATA
 
-	/** World Partition grid name to test collision against, default None will be the main grid */
-	UPROPERTY(EditAnywhere, Category = "Mass|Representation", config)
-	FName WorldPartitionGridNameContainingCollision;
-
-	/** At what rate should the not visible entity be updated in seconds */
-	UPROPERTY(EditAnywhere, Category = "Mass|Visualization", config)
-	float NotVisibleUpdateRate = 0.5f;
-
 	inline void ComputeCachedValues() const;
 
 	/** Default representation when unable to spawn an actor, gets calculated at initialization */
@@ -156,6 +159,36 @@ struct FMassRepresentationParameters : public FMassConstSharedFragment
 
 	UPROPERTY(Transient)
 	mutable TObjectPtr<UMassRepresentationActorManagement> CachedRepresentationActorManagement = nullptr;
+
+	/** 
+	 * The operation by design ignores CachedDefaultRepresentationType and CachedRepresentationActorManagement,
+	 * since those two properties are fully rutime and are derived from the others.
+	 */
+	bool Identical(const FMassRepresentationParameters* InOtherValue, uint32 /*PortFlags*/) const
+	{
+		constexpr size_t NumBytesToCompare = offsetof(FMassRepresentationParameters, NotVisibleUpdateRate) + sizeof(NotVisibleUpdateRate);
+		const int32 Result = FMemory::Memcmp(this, InOtherValue, NumBytesToCompare);
+
+		return Result == 0
+			// comparing bitfields separately since the padding might contain garbage
+			&& bForceActorRepresentationForExternalActors == InOtherValue->bForceActorRepresentationForExternalActors
+			&& bKeepLowResActors == InOtherValue->bKeepLowResActors
+			&& bKeepActorExtraFrame == InOtherValue->bKeepActorExtraFrame
+			&& bSpreadFirstVisualizationUpdate == InOtherValue->bSpreadFirstVisualizationUpdate
+#if WITH_EDITORONLY_DATA
+			&& bCanModifyRepresentationActorManagementClass == InOtherValue->bCanModifyRepresentationActorManagementClass
+#endif // WITH_EDITORONLY_DATA
+		;
+	}
+};
+
+template<>
+struct TStructOpsTypeTraits<FMassRepresentationParameters> : public TStructOpsTypeTraitsBase2<FMassRepresentationParameters>
+{
+	enum 
+	{
+		WithIdentical = true,
+	};
 };
 
 template<>
