Integration: Addition of MassSceneComponentVelocityTranslator

diff --git a/Engine/Plugins/Runtime/MassGameplay/Source/MassActors/Private/Translators/MassSceneComponentVelocityTranslator.cpp b/Engine/Plugins/Runtime/MassGameplay/Source/MassActors/Private/Translators/MassSceneComponentVelocityTranslator.cpp
new file mode 100644
--- /dev/null
+++ b/Engine/Plugins/Runtime/MassGameplay/Source/MassActors/Private/Translators/MassSceneComponentVelocityTranslator.cpp
@@ -0,0 +1,43 @@
+// Copyright Epic Games, Inc. All Rights Reserved.
+
+#include "Translators/MassSceneComponentVelocityTranslator.h"
+#include "MassCommonTypes.h"
+#include "Components/SceneComponent.h"
+#include "MassExecutionContext.h"
+#include "MassMovementFragments.h"
+#include "Translators/MassSceneComponentLocationTranslator.h"
+
+//----------------------------------------------------------------------//
+//  UMassSceneComponentVelocityToMassTranslator
+//----------------------------------------------------------------------//
+UMassSceneComponentVelocityToMassTranslator::UMassSceneComponentVelocityToMassTranslator()
+	: EntityQuery(*this)
+{
+	ExecutionFlags = (int32)EProcessorExecutionFlags::AllNetModes;
+	ExecutionOrder.ExecuteInGroup = UE::Mass::ProcessorGroupNames::SyncWorldToMass;
+	RequiredTags.Add<FMassSceneComponentVelocityCopyToMassTag>();
+}
+
+void UMassSceneComponentVelocityToMassTranslator::ConfigureQueries(const TSharedRef<FMassEntityManager>& EntityManager)
+{
+	AddRequiredTagsToQuery(EntityQuery);
+	EntityQuery.AddRequirement<FMassSceneComponentWrapperFragment>(EMassFragmentAccess::ReadOnly);
+	EntityQuery.AddRequirement<FMassVelocityFragment>(EMassFragmentAccess::ReadWrite);
+}
+
+void UMassSceneComponentVelocityToMassTranslator::Execute(FMassEntityManager& EntityManager, FMassExecutionContext& Context)
+{
+	EntityQuery.ForEachEntityChunk(Context, [this](FMassExecutionContext& Context)
+	{
+		const TConstArrayView<FMassSceneComponentWrapperFragment> ComponentList = Context.GetFragmentView<FMassSceneComponentWrapperFragment>();
+		const TArrayView<FMassVelocityFragment> VelocityList = Context.GetMutableFragmentView<FMassVelocityFragment>();
+
+		for (FMassExecutionContext::FEntityIterator EntityIt = Context.CreateEntityIterator(); EntityIt; ++EntityIt)
+		{
+			if (const USceneComponent* AsComponent = ComponentList[EntityIt].Component.Get())
+			{
+				VelocityList[EntityIt].Value = AsComponent->GetComponentVelocity();
+			}
+		}
+	});
+}
diff --git a/Engine/Plugins/Runtime/MassGameplay/Source/MassActors/Public/Translators/MassSceneComponentVelocityTranslator.h b/Engine/Plugins/Runtime/MassGameplay/Source/MassActors/Public/Translators/MassSceneComponentVelocityTranslator.h
new file mode 100644
--- /dev/null
+++ b/Engine/Plugins/Runtime/MassGameplay/Source/MassActors/Public/Translators/MassSceneComponentVelocityTranslator.h
@@ -0,0 +1,31 @@
+// Copyright Epic Games, Inc. All Rights Reserved.
+
+#pragma once
+
+#include "MassTranslator.h"
+#include "MassSceneComponentVelocityTranslator.generated.h"
+
+#define UE_API MASSACTORS_API
+
+USTRUCT()
+struct FMassSceneComponentVelocityCopyToMassTag : public FMassTag
+{
+	GENERATED_BODY()
+};
+
+UCLASS(MinimalAPI)
+class UMassSceneComponentVelocityToMassTranslator : public UMassTranslator
+{
+	GENERATED_BODY()
+
+public:
+	UE_API UMassSceneComponentVelocityToMassTranslator();
+
+protected:
+	UE_API virtual void ConfigureQueries(const TSharedRef<FMassEntityManager>& EntityManager) override;
+	UE_API virtual void Execute(FMassEntityManager& EntityManager, FMassExecutionContext& Context) override;
+
+	FMassEntityQuery EntityQuery;
+};
+
+#undef UE_API
