[MassGameplay] Modified UMassRepresentationSubsystem::Initialize to bind to processing phase events only if VisualizationComponent gets set - the phase even handlers required that to be set, so no point in binding if it's not true.

[RN] rnx

diff --git a/Engine/Plugins/Runtime/MassGameplay/Source/MassRepresentation/Private/MassRepresentationSubsystem.cpp b/Engine/Plugins/Runtime/MassGameplay/Source/MassRepresentation/Private/MassRepresentationSubsystem.cpp
--- a/Engine/Plugins/Runtime/MassGameplay/Source/MassRepresentation/Private/MassRepresentationSubsystem.cpp
+++ b/Engine/Plugins/Runtime/MassGameplay/Source/MassRepresentation/Private/MassRepresentationSubsystem.cpp
@@ -398,12 +398,14 @@ void UMassRepresentationSubsystem::Initialize(FSubsystemCollectionBase& Collecti
 #if WITH_EDITOR
 			Visualizer->SetActorLabel(FString::Printf(TEXT("%sVisualizer"), *GetClass()->GetName()), /*bMarkDirty*/false);
 #endif
-		}
 
-		UMassSimulationSubsystem* SimSystem = World->GetSubsystem<UMassSimulationSubsystem>();
-		check(SimSystem);
-		SimSystem->GetOnProcessingPhaseStarted(EMassProcessingPhase::PrePhysics).AddUObject(this, &UMassRepresentationSubsystem::OnProcessingPhaseStarted, EMassProcessingPhase::PrePhysics);
-		SimSystem->GetOnProcessingPhaseFinished(EMassProcessingPhase::PostPhysics).AddUObject(this, &UMassRepresentationSubsystem::OnProcessingPhaseStarted, EMassProcessingPhase::PostPhysics);
+			// bind to processing phase events only if VisualizationComponent gets set - the phase event
+			// handlers required that to be set, so no point in binding if it's not true
+			UMassSimulationSubsystem* SimSystem = World->GetSubsystem<UMassSimulationSubsystem>();
+			check(SimSystem);
+			SimSystem->GetOnProcessingPhaseStarted(EMassProcessingPhase::PrePhysics).AddUObject(this, &UMassRepresentationSubsystem::OnProcessingPhaseStarted, EMassProcessingPhase::PrePhysics);
+			SimSystem->GetOnProcessingPhaseFinished(EMassProcessingPhase::PostPhysics).AddUObject(this, &UMassRepresentationSubsystem::OnProcessingPhaseStarted, EMassProcessingPhase::PostPhysics);
+		}
 
 		UMassAgentSubsystem* AgentSystem = World->GetSubsystem<UMassAgentSubsystem>();
 		check(AgentSystem);
