Add missing headers that cause compile errors in non-unity builds with Mass debug tools disabled.

diff --git a/Engine/Plugins/AI/MassAI/Source/MassNavigation/Public/MassNavigationProcessors.h b/Engine/Plugins/AI/MassAI/Source/MassNavigation/Public/MassNavigationProcessors.h
--- a/Engine/Plugins/AI/MassAI/Source/MassNavigation/Public/MassNavigationProcessors.h
+++ b/Engine/Plugins/AI/MassAI/Source/MassNavigation/Public/MassNavigationProcessors.h
@@ -2,6 +2,7 @@
 
 #pragma once
 
+#include "MassEntityQuery.h"
 #include "MassProcessor.h"
 #include "MassObserverProcessor.h"
 #include "MassNavigationProcessors.generated.h"
diff --git a/Engine/Plugins/AI/MassAI/Source/MassNavigation/Public/SmoothOrientation/MassSmoothOrientationProcessors.h b/Engine/Plugins/AI/MassAI/Source/MassNavigation/Public/SmoothOrientation/MassSmoothOrientationProcessors.h
--- a/Engine/Plugins/AI/MassAI/Source/MassNavigation/Public/SmoothOrientation/MassSmoothOrientationProcessors.h
+++ b/Engine/Plugins/AI/MassAI/Source/MassNavigation/Public/SmoothOrientation/MassSmoothOrientationProcessors.h
@@ -2,6 +2,7 @@
 
 #pragma once
 
+#include "MassEntityQuery.h"
 #include "MassProcessor.h"
 #include "MassSmoothOrientationProcessors.generated.h"
 
diff --git a/Engine/Plugins/AI/MassAI/Source/MassNavigation/Public/Steering/MassSteeringProcessors.h b/Engine/Plugins/AI/MassAI/Source/MassNavigation/Public/Steering/MassSteeringProcessors.h
--- a/Engine/Plugins/AI/MassAI/Source/MassNavigation/Public/Steering/MassSteeringProcessors.h
+++ b/Engine/Plugins/AI/MassAI/Source/MassNavigation/Public/Steering/MassSteeringProcessors.h
@@ -2,6 +2,7 @@
 
 #pragma once
 
+#include "MassEntityQuery.h"
 #include "MassProcessor.h"
 #if UE_ENABLE_INCLUDE_ORDER_DEPRECATED_IN_5_6
 #include "MassObserverProcessor.h"
diff --git a/Engine/Plugins/Runtime/MassGameplay/Source/MassActors/Public/Translators/MassCharacterMovementTranslators.h b/Engine/Plugins/Runtime/MassGameplay/Source/MassActors/Public/Translators/MassCharacterMovementTranslators.h
--- a/Engine/Plugins/Runtime/MassGameplay/Source/MassActors/Public/Translators/MassCharacterMovementTranslators.h
+++ b/Engine/Plugins/Runtime/MassGameplay/Source/MassActors/Public/Translators/MassCharacterMovementTranslators.h
@@ -2,6 +2,7 @@
 
 #pragma once
 
+#include "MassEntityQuery.h"
 #include "MassCommonFragments.h"
 #include "MassTranslator.h"
 #include "MassObserverProcessor.h"
diff --git a/Engine/Plugins/Runtime/MassGameplay/Source/MassActors/Public/Translators/MassSceneComponentLocationTranslator.h b/Engine/Plugins/Runtime/MassGameplay/Source/MassActors/Public/Translators/MassSceneComponentLocationTranslator.h
--- a/Engine/Plugins/Runtime/MassGameplay/Source/MassActors/Public/Translators/MassSceneComponentLocationTranslator.h
+++ b/Engine/Plugins/Runtime/MassGameplay/Source/MassActors/Public/Translators/MassSceneComponentLocationTranslator.h
@@ -2,6 +2,7 @@
 
 #pragma once
 
+#include "MassEntityQuery.h"
 #include "MassCommonFragments.h"
 #include "MassTranslator.h"
 #include "MassSceneComponentLocationTranslator.generated.h"
diff --git a/Engine/Plugins/Runtime/MassGameplay/Source/MassActors/Public/Translators/MassSceneComponentVelocityTranslator.h b/Engine/Plugins/Runtime/MassGameplay/Source/MassActors/Public/Translators/MassSceneComponentVelocityTranslator.h
--- a/Engine/Plugins/Runtime/MassGameplay/Source/MassActors/Public/Translators/MassSceneComponentVelocityTranslator.h
+++ b/Engine/Plugins/Runtime/MassGameplay/Source/MassActors/Public/Translators/MassSceneComponentVelocityTranslator.h
@@ -2,6 +2,7 @@
 
 #pragma once
 
+#include "MassEntityQuery.h"
 #include "MassTranslator.h"
 #include "MassSceneComponentVelocityTranslator.generated.h"
 
diff --git a/Engine/Plugins/Runtime/MassGameplay/Source/MassActors/Public/Translators/MassTranslators_BehaviorTree.h b/Engine/Plugins/Runtime/MassGameplay/Source/MassActors/Public/Translators/MassTranslators_BehaviorTree.h
--- a/Engine/Plugins/Runtime/MassGameplay/Source/MassActors/Public/Translators/MassTranslators_BehaviorTree.h
+++ b/Engine/Plugins/Runtime/MassGameplay/Source/MassActors/Public/Translators/MassTranslators_BehaviorTree.h
@@ -2,6 +2,7 @@
 
 #pragma once
 
+#include "MassEntityQuery.h"
 #include "MassTranslator.h"
 #include "MassCommonFragments.h"
 #include "MassTranslators_BehaviorTree.generated.h"
