[MassGameplay] MassAgentComponent changes to update EntityHandle-to-actor mapping in MassActorSubsystem when dealing with a player-owned component

The change fixes an issue with GameplayDebugger's Mass category not being able to display info on the player-owned entity. A tiny change to the category removes the friction when doing that (no longer requiring the transform fragment to display data).

diff --git a/Engine/Plugins/AI/MassAI/Source/MassAIDebug/Private/GameplayDebuggerCategory_Mass.cpp b/Engine/Plugins/AI/MassAI/Source/MassAIDebug/Private/GameplayDebuggerCategory_Mass.cpp
--- a/Engine/Plugins/AI/MassAI/Source/MassAIDebug/Private/GameplayDebuggerCategory_Mass.cpp
+++ b/Engine/Plugins/AI/MassAI/Source/MassAIDebug/Private/GameplayDebuggerCategory_Mass.cpp
@@ -488,9 +488,11 @@ void FGameplayDebuggerCategory_Mass::CollectData(APlayerController* OwnerPC, AAc
 					AddTextLine(FString::Printf(TEXT("{Green}Shared Fragments:{White}%s"), *DescriptionBuilder(ItemNames)));
 				}
 
-				check(TransformFragment);
-				constexpr float CapsuleRadius = 50.f;
-				AddShape(FGameplayDebuggerShape::MakeCapsule(TransformFragment->GetTransform().GetLocation() + 2.f * CapsuleRadius * FVector::UpVector, CapsuleRadius, CapsuleRadius * 2.f, FColor::Orange));
+				if (TransformFragment)
+				{
+					constexpr float CapsuleRadius = 50.f;
+					AddShape(FGameplayDebuggerShape::MakeCapsule(TransformFragment->GetTransform().GetLocation() + 2.f * CapsuleRadius * FVector::UpVector, CapsuleRadius, CapsuleRadius * 2.f, FColor::Orange));
+				}
 			}
 			else
 			{
diff --git a/Engine/Plugins/Runtime/MassGameplay/Source/MassActors/Private/MassAgentComponent.cpp b/Engine/Plugins/Runtime/MassGameplay/Source/MassActors/Private/MassAgentComponent.cpp
--- a/Engine/Plugins/Runtime/MassGameplay/Source/MassActors/Private/MassAgentComponent.cpp
+++ b/Engine/Plugins/Runtime/MassGameplay/Source/MassActors/Private/MassAgentComponent.cpp
@@ -238,6 +238,17 @@ void UMassAgentComponent::SetEntityHandleInternal(const FMassEntityHandle NewHan
 				}
 			}
 		}
+		else if (IsPuppet() == false)
+		{
+			// Since this component is owned by an actor not owned by Mass (i.e. it's not a "puppet")
+			// we need to manually notify UMassActorSubsystem about this entity-handle-to-actor mapping
+			if (FMassActorFragment* ActorInfo = EntitySubsystem->GetEntityManager().GetFragmentDataPtr<FMassActorFragment>(AgentHandle))
+			{
+				checkf(!ActorInfo->IsValid(), TEXT("Expecting ActorInfo fragment to be null"));
+				ActorInfo->SetAndUpdateHandleMap(AgentHandle, GetOwner(), !IsNetSimulating()/*bIsOwnedByMass*/);
+			}
+			// @todo consider adding the fragment with an appropriate initialization, if the fragment's missing
+		}
 	}
 }
 
@@ -298,17 +309,15 @@ void UMassAgentComponent::ClearEntityHandleInternal()
 			AgentSubsystem->NotifyMassAgentComponentEntityDetaching(*this);
 		}
 
-		if (IsNetSimulating())
+		if (const UMassEntitySubsystem* EntitySubsystem = World->GetSubsystem<UMassEntitySubsystem>())
 		{
 			// Sync up with mass
-			if (const UMassEntitySubsystem* EntitySubsystem = World->GetSubsystem<UMassEntitySubsystem>())
+			FMassActorFragment* ActorInfo = EntitySubsystem->GetEntityManager().GetFragmentDataPtr<FMassActorFragment>(AgentHandle);
+			if (ActorInfo && (IsNetSimulating() == true || IsPuppet() == false))
 			{
-				const FMassEntityView EntityView(EntitySubsystem->GetEntityManager(), AgentHandle);
-				if (FMassActorFragment* ActorInfo = EntityView.GetFragmentDataPtr<FMassActorFragment>())
-				{
-					checkf(!ActorInfo->IsValid() || ActorInfo->Get() == GetOwner(), TEXT("Expecting actor pointer to be the Component\'s owner"));
-					ActorInfo->ResetAndUpdateHandleMap(World->GetSubsystem<UMassActorSubsystem>());
-				}
+				checkf(!ActorInfo->IsValid() || ActorInfo->Get() == GetOwner(), TEXT("[%s] Expecting actor pointer to be the Component\'s owner")
+					, IsNetSimulating() ? TEXT("NetSimulating") : TEXT("Non-Puppet"));
+				ActorInfo->ResetAndUpdateHandleMap(World->GetSubsystem<UMassActorSubsystem>());
 			}
 		}
 	}
