[MassSmartObject] Added basic MRU slots implementation
- Added new options in UMassSmartObjectSettings
-- MRUSlotsMaxCount: Possible values between 0 and 4 (Default: 0, meaning that the feature is not used; up to 4 slots)
-- bUseCooldownForMRUSlots: to allow the most recently used slots to be considered again a given period of time
-- MRUSlotsCooldown: period of time during which a given slot will not be considered after being used
- Enabled UE types in namespaces for MassSmartObject (i.e., AllowUETypesInNamespaces)

diff --git a/Engine/Plugins/AI/MassAI/Source/MassAIBehavior/Private/Tasks/MassClaimSmartObjectTask.cpp b/Engine/Plugins/AI/MassAI/Source/MassAIBehavior/Private/Tasks/MassClaimSmartObjectTask.cpp
--- a/Engine/Plugins/AI/MassAI/Source/MassAIBehavior/Private/Tasks/MassClaimSmartObjectTask.cpp
+++ b/Engine/Plugins/AI/MassAI/Source/MassAIBehavior/Private/Tasks/MassClaimSmartObjectTask.cpp
@@ -8,10 +8,11 @@
 #include "MassSignalSubsystem.h"
 #include "MassSmartObjectFragments.h"
 #include "MassSmartObjectHandler.h"
+#include "MassSmartObjectSettings.h"
+#include "MassStateTreeDependency.h"
 #include "MassStateTreeExecutionContext.h"
 #include "SmartObjectSubsystem.h"
 #include "StateTreeLinker.h"
-#include "MassStateTreeDependency.h"
 
 //----------------------------------------------------------------------//
 // FMassClaimSmartObjectTask
@@ -28,6 +29,7 @@ FMassClaimSmartObjectTask::FMassClaimSmartObjectTask()
 bool FMassClaimSmartObjectTask::Link(FStateTreeLinker& Linker)
 {
 	Linker.LinkExternalData(SmartObjectUserHandle);
+	Linker.LinkExternalData(SmartObjectMRUSlotsHandle);
 	Linker.LinkExternalData(SmartObjectSubsystemHandle);
 	Linker.LinkExternalData(MassSignalSubsystemHandle);
 
@@ -37,6 +39,7 @@ bool FMassClaimSmartObjectTask::Link(FStateTreeLinker& Linker)
 void FMassClaimSmartObjectTask::GetDependencies(UE::MassBehavior::FStateTreeDependencyBuilder& Builder) const
 {
 	Builder.AddReadWrite(SmartObjectUserHandle);
+	Builder.AddReadWrite(SmartObjectMRUSlotsHandle);
 	// @todo SmartObjectSubsystemHandle is being used in a RW fashion, but we need this task to be able to run in parallel with
 	// everything else, so we need to ensure TMassExternalSubsystemTraits<USmartObjectSubsystem> is marked up for parallel access
 	// and that this information is properly utilized during Mass processing graph creation
@@ -93,11 +96,30 @@ void FMassClaimSmartObjectTask::ExitState(FStateTreeExecutionContext& Context, c
 	if (InstanceData.ClaimedSlot.IsValid())
 	{
 		const FMassStateTreeExecutionContext& MassStateTreeContext = static_cast<FMassStateTreeExecutionContext&>(Context);
+		FMassExecutionContext& MassContext = MassStateTreeContext.GetMassEntityExecutionContext();
+		const FMassEntityHandle Entity = MassStateTreeContext.GetEntity();
+
+		if (UE::Mass::SmartObject::FMRUSlotsFragment* ExistingMRUSlotsFragment = Context.GetExternalDataPtr(SmartObjectMRUSlotsHandle))
+		{
+			ExistingMRUSlotsFragment->Slots.Push(InstanceData.ClaimedSlot.SlotHandle);
+		}
+		else
+		{
+			const UMassSmartObjectSettings* MassSmartObjectSettings = GetDefault<UMassSmartObjectSettings>();
+			if (MassSmartObjectSettings != nullptr
+				&& MassSmartObjectSettings->MRUSlotsMaxCount > 0)
+			{
+				UE::Mass::SmartObject::FMRUSlotsFragment MRUSlotsFragment;
+				MRUSlotsFragment.Slots.Push(InstanceData.ClaimedSlot.SlotHandle);
+				MassContext.Defer().PushCommand<FMassCommandAddFragmentInstances>(Entity, MRUSlotsFragment);
+			}
+		}
+
 		USmartObjectSubsystem& SmartObjectSubsystem = Context.GetExternalData(SmartObjectSubsystemHandle);
 		UMassSignalSubsystem& SignalSubsystem = Context.GetExternalData(MassSignalSubsystemHandle);
-		const FMassSmartObjectHandler MassSmartObjectHandler(MassStateTreeContext.GetMassEntityExecutionContext(), SmartObjectSubsystem, SignalSubsystem);
+		const FMassSmartObjectHandler MassSmartObjectHandler(MassContext, SmartObjectSubsystem, SignalSubsystem);
 
-		MassSmartObjectHandler.ReleaseSmartObject(MassStateTreeContext.GetEntity(), SOUser, InstanceData.ClaimedSlot);
+		MassSmartObjectHandler.ReleaseSmartObject(Entity, SOUser, InstanceData.ClaimedSlot);
 	}
 	else
 	{
@@ -117,6 +139,7 @@ EStateTreeRunStatus FMassClaimSmartObjectTask::Tick(FStateTreeExecutionContext&
 
 	// Check that the claimed slot is still valid, and if not, fail the task.
 	// The slot can become invalid if the whole SO or slot becomes invalidated.
+	// In this case, we don't consider the slot as used and don't update the MRU slots
 	if (InstanceData.ClaimedSlot.IsValid())
 	{
 		const USmartObjectSubsystem& SmartObjectSubsystem = Context.GetExternalData(SmartObjectSubsystemHandle);
diff --git a/Engine/Plugins/AI/MassAI/Source/MassAIBehavior/Private/Tasks/MassFindSmartObjectTask.cpp b/Engine/Plugins/AI/MassAI/Source/MassAIBehavior/Private/Tasks/MassFindSmartObjectTask.cpp
--- a/Engine/Plugins/AI/MassAI/Source/MassAIBehavior/Private/Tasks/MassFindSmartObjectTask.cpp
+++ b/Engine/Plugins/AI/MassAI/Source/MassAIBehavior/Private/Tasks/MassFindSmartObjectTask.cpp
@@ -30,6 +30,7 @@ bool FMassFindSmartObjectTask::Link(FStateTreeLinker& Linker)
 	Linker.LinkExternalData(MassSignalSubsystemHandle);
 	Linker.LinkExternalData(EntityTransformHandle);
 	Linker.LinkExternalData(SmartObjectUserHandle);
+	Linker.LinkExternalData(SmartObjectMRUSlotsHandle);
 	Linker.LinkExternalData(LocationHandle);
 
 	return true;
@@ -44,6 +45,7 @@ void FMassFindSmartObjectTask::GetDependencies(UE::MassBehavior::FStateTreeDepen
 	Builder.AddReadWrite(MassSignalSubsystemHandle);
 	Builder.AddReadOnly(EntityTransformHandle);
 	Builder.AddReadWrite(SmartObjectUserHandle);
+	Builder.AddReadOnly(SmartObjectMRUSlotsHandle);
 	Builder.AddReadOnly(LocationHandle);
 }
 
@@ -102,7 +104,9 @@ EStateTreeRunStatus FMassFindSmartObjectTask::Tick(FStateTreeExecutionContext& C
 		SmartObjectSubsystem,
 		SignalSubsystem);
 
+	using namespace UE::Mass::SmartObject;
 	FMassSmartObjectUserFragment& SOUser = Context.GetExternalData(SmartObjectUserHandle);
+	FMRUSlotsFragment* MRUSlots = Context.GetExternalDataPtr(SmartObjectMRUSlotsHandle);
 
 	FInstanceDataType& InstanceData = Context.GetInstanceData(*this);
 	
@@ -142,13 +146,29 @@ EStateTreeRunStatus FMassFindSmartObjectTask::Tick(FStateTreeExecutionContext& C
 						*LexToString(LaneLocation->LaneLength));
 
 					InstanceData.LastLane = LaneLocation->LaneHandle;
-					InstanceData.SearchRequestID = MassSmartObjectHandler.FindCandidatesAsync(RequestingEntity, SOUser.UserTags, ActivityRequirements, { LaneLocation->LaneHandle, LaneLocation->DistanceAlongLane });
+
+					FFindCandidatesParameters Parameters
+					{
+						.UserTags = SOUser.UserTags,
+						.ActivityRequirements = ActivityRequirements,
+						.LaneLocation = { LaneLocation->LaneHandle, LaneLocation->DistanceAlongLane },
+						.MRUSlots = MRUSlots != nullptr ? MRUSlots->Slots : FMRUSlots{}
+					};
+
+					InstanceData.SearchRequestID = MassSmartObjectHandler.FindCandidatesAsync(RequestingEntity, MoveTemp(Parameters));
 				}
 			}
 			else
 			{
-				const FTransformFragment& TransformFragment = Context.GetExternalData(EntityTransformHandle);
-				InstanceData.SearchRequestID = MassSmartObjectHandler.FindCandidatesAsync(RequestingEntity, SOUser.UserTags, ActivityRequirements, TransformFragment.GetTransform().GetLocation());
+				FFindCandidatesParameters Parameters
+				{
+					.UserTags = SOUser.UserTags,
+					.ActivityRequirements = ActivityRequirements,
+					.Location = Context.GetExternalData(EntityTransformHandle).GetTransform().GetLocation(),
+					.MRUSlots = MRUSlots != nullptr ? MRUSlots->Slots : FMRUSlots{}
+				};
+
+				InstanceData.SearchRequestID = MassSmartObjectHandler.FindCandidatesAsync(RequestingEntity, MoveTemp(Parameters));
 			}
 		}
 	}
diff --git a/Engine/Plugins/AI/MassAI/Source/MassAIBehavior/Public/Tasks/MassClaimSmartObjectTask.h b/Engine/Plugins/AI/MassAI/Source/MassAIBehavior/Public/Tasks/MassClaimSmartObjectTask.h
--- a/Engine/Plugins/AI/MassAI/Source/MassAIBehavior/Public/Tasks/MassClaimSmartObjectTask.h
+++ b/Engine/Plugins/AI/MassAI/Source/MassAIBehavior/Public/Tasks/MassClaimSmartObjectTask.h
@@ -19,6 +19,10 @@ namespace UE::MassBehavior
 {
 	struct FStateTreeDependencyBuilder;
 };
+namespace UE::Mass::SmartObject
+{
+	struct FMRUSlotsFragment;
+}
 
 /**
  * Tasks to claim a smart object from search results and release it when done.
@@ -54,6 +58,7 @@ protected:
 	MASSAIBEHAVIOR_API virtual void GetDependencies(UE::MassBehavior::FStateTreeDependencyBuilder& Builder) const override;
 
 	TStateTreeExternalDataHandle<FMassSmartObjectUserFragment> SmartObjectUserHandle;
+	TStateTreeExternalDataHandle<UE::Mass::SmartObject::FMRUSlotsFragment, EStateTreeExternalDataRequirement::Optional> SmartObjectMRUSlotsHandle;
 	TStateTreeExternalDataHandle<USmartObjectSubsystem> SmartObjectSubsystemHandle;
 	TStateTreeExternalDataHandle<UMassSignalSubsystem> MassSignalSubsystemHandle;
 
diff --git a/Engine/Plugins/AI/MassAI/Source/MassAIBehavior/Public/Tasks/MassFindSmartObjectTask.h b/Engine/Plugins/AI/MassAI/Source/MassAIBehavior/Public/Tasks/MassFindSmartObjectTask.h
--- a/Engine/Plugins/AI/MassAI/Source/MassAIBehavior/Public/Tasks/MassFindSmartObjectTask.h
+++ b/Engine/Plugins/AI/MassAI/Source/MassAIBehavior/Public/Tasks/MassFindSmartObjectTask.h
@@ -18,6 +18,10 @@ namespace UE::MassBehavior
 {
 	struct FStateTreeDependencyBuilder;
 };
+namespace UE::Mass::SmartObject
+{
+	struct FMRUSlotsFragment;
+};
 
 USTRUCT()
 struct FMassFindSmartObjectTaskInstanceData
@@ -65,6 +69,7 @@ protected:
 	TStateTreeExternalDataHandle<UMassSignalSubsystem> MassSignalSubsystemHandle;
 	TStateTreeExternalDataHandle<FTransformFragment> EntityTransformHandle;
 	TStateTreeExternalDataHandle<FMassSmartObjectUserFragment> SmartObjectUserHandle;
+	TStateTreeExternalDataHandle<UE::Mass::SmartObject::FMRUSlotsFragment, EStateTreeExternalDataRequirement::Optional> SmartObjectMRUSlotsHandle;
 	TStateTreeExternalDataHandle<FMassZoneGraphLaneLocationFragment, EStateTreeExternalDataRequirement::Optional> LocationHandle;
 
 	/** Gameplay tag query for finding matching smart objects. */
diff --git a/Engine/Plugins/Runtime/MassGameplay/Source/MassSmartObjects/MassSmartObjects.Build.cs b/Engine/Plugins/Runtime/MassGameplay/Source/MassSmartObjects/MassSmartObjects.Build.cs
--- a/Engine/Plugins/Runtime/MassGameplay/Source/MassSmartObjects/MassSmartObjects.Build.cs
+++ b/Engine/Plugins/Runtime/MassGameplay/Source/MassSmartObjects/MassSmartObjects.Build.cs
@@ -6,6 +6,8 @@ namespace UnrealBuildTool.Rules
 	{
 		public MassSmartObjects(ReadOnlyTargetRules Target) : base(Target)
 		{
+			bAllowUETypesInNamespaces = true;
+
 			PCHUsage = PCHUsageMode.UseExplicitOrSharedPCHs;
 
 			CppCompileWarningSettings.UnsafeTypeCastWarningLevel = WarningLevel.Warning;
diff --git a/Engine/Plugins/Runtime/MassGameplay/Source/MassSmartObjects/Private/MassSmartObjectHandler.cpp b/Engine/Plugins/Runtime/MassGameplay/Source/MassSmartObjects/Private/MassSmartObjectHandler.cpp
--- a/Engine/Plugins/Runtime/MassGameplay/Source/MassSmartObjects/Private/MassSmartObjectHandler.cpp
+++ b/Engine/Plugins/Runtime/MassGameplay/Source/MassSmartObjects/Private/MassSmartObjectHandler.cpp
@@ -67,34 +67,69 @@ namespace UE::Mass::SmartObject
 //----------------------------------------------------------------------//
 FMassSmartObjectRequestID FMassSmartObjectHandler::FindCandidatesAsync(const FMassEntityHandle RequestingEntity, const FGameplayTagContainer& UserTags, const FGameplayTagQuery& ActivityRequirements, const FVector& Location) const
 {
-	const FMassEntityHandle ReservedEntity = ExecutionContext.GetEntityManagerChecked().ReserveEntity();
-
-	FMassSmartObjectWorldLocationRequestFragment RequestFragment;
-	RequestFragment.SearchOrigin = Location;
-	RequestFragment.RequestingEntity = RequestingEntity;
-	RequestFragment.UserTags = UserTags;
-	RequestFragment.ActivityRequirements = ActivityRequirements;
+	UE::Mass::SmartObject::FFindCandidatesParameters Parameters
+	{
+		.UserTags = UserTags,
+		.ActivityRequirements = ActivityRequirements,
+		.Location = Location
+	};
 
-	FMassSmartObjectRequestResultFragment ResultFragment;
+	return FindCandidatesAsync(RequestingEntity, MoveTemp(Parameters));
+}
 
-	ExecutionContext.Defer().PushCommand<FMassCommandAddFragmentInstances>(ReservedEntity, RequestFragment, ResultFragment);
+FMassSmartObjectRequestID FMassSmartObjectHandler::FindCandidatesAsync(const FMassEntityHandle RequestingEntity, const FGameplayTagContainer& UserTags, const FGameplayTagQuery& ActivityRequirements, const FZoneGraphCompactLaneLocation& LaneLocation) const
+{
+	UE::Mass::SmartObject::FFindCandidatesParameters Parameters
+	{
+		.UserTags = UserTags,
+		.ActivityRequirements = ActivityRequirements,
+		.LaneLocation = LaneLocation
+	};
 
-	return FMassSmartObjectRequestID(ReservedEntity);
+	return FindCandidatesAsync(RequestingEntity, MoveTemp(Parameters));
 }
 
-FMassSmartObjectRequestID FMassSmartObjectHandler::FindCandidatesAsync(const FMassEntityHandle RequestingEntity, const FGameplayTagContainer& UserTags, const FGameplayTagQuery& ActivityRequirements, const FZoneGraphCompactLaneLocation& LaneLocation) const
+FMassSmartObjectRequestID FMassSmartObjectHandler::FindCandidatesAsync(const FMassEntityHandle RequestingEntity, UE::Mass::SmartObject::FFindCandidatesParameters&& Parameters) const
 {
 	const FMassEntityHandle ReservedEntity = ExecutionContext.GetEntityManagerChecked().ReserveEntity();
+	FMassSmartObjectRequestResultFragment ResultFragment;
 
-	FMassSmartObjectLaneLocationRequestFragment RequestFragment;
-	RequestFragment.CompactLaneLocation = LaneLocation;
-	RequestFragment.RequestingEntity = RequestingEntity;
-	RequestFragment.UserTags = UserTags;
-	RequestFragment.ActivityRequirements = ActivityRequirements;
+	if (Parameters.LaneLocation.LaneHandle.IsValid())
+	{
+		FMassSmartObjectLaneLocationRequestFragment RequestFragment;
+		RequestFragment.CompactLaneLocation = Parameters.LaneLocation;
+		RequestFragment.RequestingEntity = RequestingEntity;
+		RequestFragment.UserTags = Parameters.UserTags;
+		RequestFragment.ActivityRequirements = Parameters.ActivityRequirements;
 
-	FMassSmartObjectRequestResultFragment ResultFragment;
+		if (Parameters.MRUSlots.NumSlots > 0)
+		{
+			UE::Mass::SmartObject::FMRUSlotsFragment MRUSlotsFragment{.Slots = Parameters.MRUSlots};
+			ExecutionContext.Defer().PushCommand<FMassCommandAddFragmentInstances>(ReservedEntity, RequestFragment, MRUSlotsFragment, ResultFragment);
+		}
+		else
+		{
+			ExecutionContext.Defer().PushCommand<FMassCommandAddFragmentInstances>(ReservedEntity, RequestFragment, ResultFragment);
+		}
+	}
+	else
+	{
+		FMassSmartObjectWorldLocationRequestFragment RequestFragment;
+		RequestFragment.SearchOrigin = Parameters.Location;
+		RequestFragment.RequestingEntity = RequestingEntity;
+		RequestFragment.UserTags = Parameters.UserTags;
+		RequestFragment.ActivityRequirements = Parameters.ActivityRequirements;
 
-	ExecutionContext.Defer().PushCommand<FMassCommandAddFragmentInstances>(ReservedEntity, RequestFragment, ResultFragment);
+		if (Parameters.MRUSlots.NumSlots > 0)
+		{
+			UE::Mass::SmartObject::FMRUSlotsFragment MRUSlotsFragment{ .Slots = Parameters.MRUSlots };
+			ExecutionContext.Defer().PushCommand<FMassCommandAddFragmentInstances>(ReservedEntity, RequestFragment, MRUSlotsFragment, ResultFragment);
+		}
+		else
+		{
+			ExecutionContext.Defer().PushCommand<FMassCommandAddFragmentInstances>(ReservedEntity, RequestFragment, ResultFragment);
+		}
+	}
 
 	return FMassSmartObjectRequestID(ReservedEntity);
 }
diff --git a/Engine/Plugins/Runtime/MassGameplay/Source/MassSmartObjects/Private/MassSmartObjectProcessor.cpp b/Engine/Plugins/Runtime/MassGameplay/Source/MassSmartObjects/Private/MassSmartObjectProcessor.cpp
--- a/Engine/Plugins/Runtime/MassGameplay/Source/MassSmartObjects/Private/MassSmartObjectProcessor.cpp
+++ b/Engine/Plugins/Runtime/MassGameplay/Source/MassSmartObjects/Private/MassSmartObjectProcessor.cpp
@@ -29,12 +29,14 @@ void UMassSmartObjectCandidatesFinderProcessor::ConfigureQueries(const TSharedRe
 {
 	WorldRequestQuery.AddRequirement<FMassSmartObjectWorldLocationRequestFragment>(EMassFragmentAccess::ReadOnly);
 	WorldRequestQuery.AddRequirement<FMassSmartObjectRequestResultFragment>(EMassFragmentAccess::ReadWrite);
+	WorldRequestQuery.AddRequirement<UE::Mass::SmartObject::FMRUSlotsFragment>(EMassFragmentAccess::ReadOnly, EMassFragmentPresence::Optional);
 	WorldRequestQuery.AddTagRequirement<FMassSmartObjectCompletedRequestTag>(EMassFragmentPresence::None);
 	WorldRequestQuery.AddSubsystemRequirement<USmartObjectSubsystem>(EMassFragmentAccess::ReadOnly);
 	WorldRequestQuery.RegisterWithProcessor(*this);
 
 	LaneRequestQuery.AddRequirement<FMassSmartObjectLaneLocationRequestFragment>(EMassFragmentAccess::ReadOnly);
 	LaneRequestQuery.AddRequirement<FMassSmartObjectRequestResultFragment>(EMassFragmentAccess::ReadWrite);
+	LaneRequestQuery.AddRequirement<UE::Mass::SmartObject::FMRUSlotsFragment>(EMassFragmentAccess::ReadOnly, EMassFragmentPresence::Optional);
 	LaneRequestQuery.AddTagRequirement<FMassSmartObjectCompletedRequestTag>(EMassFragmentPresence::None);
 	LaneRequestQuery.AddSubsystemRequirement<UZoneGraphSubsystem>(EMassFragmentAccess::ReadOnly);
 	LaneRequestQuery.AddSubsystemRequirement<USmartObjectSubsystem>(EMassFragmentAccess::ReadOnly);
@@ -96,6 +98,7 @@ void UMassSmartObjectCandidatesFinderProcessor::Execute(FMassEntityManager& Enti
 		EntitiesToSignal.Reserve(EntitiesToSignal.Num() + NumEntities);
 
 		const TConstArrayView<FMassSmartObjectWorldLocationRequestFragment> RequestList = Context.GetFragmentView<FMassSmartObjectWorldLocationRequestFragment>();
+		const TConstArrayView<UE::Mass::SmartObject::FMRUSlotsFragment> MRUSlotsList = Context.GetFragmentView<UE::Mass::SmartObject::FMRUSlotsFragment>();
 		const TArrayView<FMassSmartObjectRequestResultFragment> ResultList = Context.GetMutableFragmentView<FMassSmartObjectRequestResultFragment>();
 
 		TArray<FSmartObjectRequestResult> QueryResults;
@@ -129,8 +132,29 @@ void UMassSmartObjectCandidatesFinderProcessor::Execute(FMassEntityManager& Enti
 			SmartObjectSubsystem.FindSmartObjects(FSmartObjectRequest(SearchBounds, Filter), QueryResults);
 
 			SortedCandidateSlots.Reset(QueryResults.Num());
+
+			// Gather MRU slots to ignore if the settings are set to use them
+			TArray<FSmartObjectSlotHandle> SlotsToIgnore;
+			if (MRUSlotsList.Num() > 0)
+			{
+				MRUSlotsList[EntityIt].Slots.AppendTo(SlotsToIgnore);
+
+				const UE::Mass::SmartObject::FMRUSlots& Slots = MRUSlotsList[EntityIt].Slots;
+				SlotsToIgnore.Reserve(Slots.NumSlots);
+
+				for (int32 SlotIndex = Slots.NumSlots - 1; SlotIndex >= 0; SlotIndex--)
+				{
+					SlotsToIgnore.Push(Slots.Slots[SlotIndex].Slot);
+				}
+			}
+
 			for (const FSmartObjectRequestResult& QueryResult : QueryResults)
 			{
+				if (SlotsToIgnore.Contains(QueryResult.SlotHandle))
+				{
+					continue;
+				}
+
 				const FVector SlotLocation = SmartObjectSubsystem.GetSlotLocation(QueryResult.SlotHandle).GetValue();
 				SortedCandidateSlots.Emplace(QueryResult, UE_REAL_TO_FLOAT_CLAMPED_MAX(FVector::DistSquared(SearchOrigin, SlotLocation)));
 
@@ -170,6 +194,7 @@ void UMassSmartObjectCandidatesFinderProcessor::Execute(FMassEntityManager& Enti
 
 			TConstArrayView<FMassSmartObjectLaneLocationRequestFragment> RequestList = Context.GetFragmentView<FMassSmartObjectLaneLocationRequestFragment>();
 			TArrayView<FMassSmartObjectRequestResultFragment> ResultList = Context.GetMutableFragmentView<FMassSmartObjectRequestResultFragment>();
+			const TConstArrayView<UE::Mass::SmartObject::FMRUSlotsFragment> MRUSlotsList = Context.GetFragmentView<UE::Mass::SmartObject::FMRUSlotsFragment>();
 
 			// Cache latest used data since request are most of the time on the same zone graph
 			FZoneGraphDataHandle LastUsedDataHandle;
@@ -263,8 +288,20 @@ void UMassSmartObjectCandidatesFinderProcessor::Execute(FMassEntityManager& Enti
 						continue;
 					}
 
+					// Gather MRU slots to ignore if the settings are set to use them
+					TArray<FSmartObjectSlotHandle> SlotsToIgnore;
+					if (MRUSlotsList.Num() > 0)
+					{
+						MRUSlotsList[EntityIt].Slots.AppendTo(SlotsToIgnore);
+					}
+
 					for (FSmartObjectSlotHandle SlotHandle : SlotHandles)
 					{
+						if (SlotsToIgnore.Contains(SlotHandle))
+						{
+							continue;
+						}
+
 						Result.Candidates.Slots[Result.Candidates.NumSlots++] = FSmartObjectCandidateSlot(FSmartObjectRequestResult(Handle, SlotHandle), Cost);
 
 #if WITH_MASSGAMEPLAY_DEBUG
@@ -381,6 +418,7 @@ void UMassSmartObjectTimedBehaviorProcessor::Execute(FMassEntityManager& EntityM
 	}
 }
 
+
 //----------------------------------------------------------------------//
 //  UMassSmartObjectUserFragmentDeinitializer
 //----------------------------------------------------------------------//
@@ -411,3 +449,58 @@ void UMassSmartObjectUserFragmentDeinitializer::Execute(FMassEntityManager& Enti
 			}
 		});
 }
+
+namespace UE::Mass::SmartObject
+{
+
+//----------------------------------------------------------------------//
+// UMRUSlotsProcessor
+//----------------------------------------------------------------------//
+void UMRUSlotsProcessor::ConfigureQueries(const TSharedRef<FMassEntityManager>& EntityManager)
+{
+	EntityQuery.AddRequirement<FMassSmartObjectUserFragment>(EMassFragmentAccess::ReadOnly);
+	EntityQuery.AddRequirement<FMRUSlotsFragment>(EMassFragmentAccess::ReadWrite);
+}
+
+UMRUSlotsProcessor::UMRUSlotsProcessor()
+	: EntityQuery(*this)
+{
+	ExecutionOrder.ExecuteInGroup = ProcessorGroupNames::SyncWorldToMass;
+}
+
+void UMRUSlotsProcessor::Execute(FMassEntityManager& EntityManager, FMassExecutionContext& Context)
+{
+	QUICK_SCOPE_CYCLE_COUNTER(UMassProcessor_SmartObjecMRUSlots);
+
+	const UMassSmartObjectSettings* Settings = GetDefault<UMassSmartObjectSettings>();
+	if (Settings && Settings->bUseCooldownForMRUSlots == false)
+	{
+		return;
+	}
+
+	EntityQuery.ForEachEntityChunk(Context, [this](FMassExecutionContext& Context)
+	{
+		const float DT = Context.GetDeltaTimeSeconds();
+		const TArrayView<FMRUSlotsFragment> MRUsList = Context.GetMutableFragmentView<FMRUSlotsFragment>();
+
+		for (FMassExecutionContext::FEntityIterator EntityIt = Context.CreateEntityIterator(); EntityIt; ++EntityIt)
+		{
+			FMRUSlotsFragment& MRUs = MRUsList[EntityIt];
+			FMRUSlots UpdatedSlots;
+			for (int32 Index = MRUs.Slots.NumSlots - 1; Index >= 0; Index--)
+			{
+				float& Time = MRUs.Slots.Slots[Index].RemainingTime;
+				Time = FMath::Max(Time - DT, 0.0f);
+				if (Time > 0.f)
+				{
+					UpdatedSlots.Slots[UpdatedSlots.NumSlots] = MRUs.Slots.Slots[Index];
+					UpdatedSlots.NumSlots++;
+				}
+			}
+
+			MRUs.Slots = MoveTemp(UpdatedSlots);
+		}
+	});
+}
+
+} // UE::Mass::SmartObject
\ No newline at end of file
diff --git a/Engine/Plugins/Runtime/MassGameplay/Source/MassSmartObjects/Private/MassSmartObjectTypes.cpp b/Engine/Plugins/Runtime/MassGameplay/Source/MassSmartObjects/Private/MassSmartObjectTypes.cpp
--- a/Engine/Plugins/Runtime/MassGameplay/Source/MassSmartObjects/Private/MassSmartObjectTypes.cpp
+++ b/Engine/Plugins/Runtime/MassGameplay/Source/MassSmartObjects/Private/MassSmartObjectTypes.cpp
@@ -1,5 +1,42 @@
 // Copyright Epic Games, Inc. All Rights Reserved.
 
 #include "MassSmartObjectTypes.h"
+#include "MassSmartObjectSettings.h"
 
 #include UE_INLINE_GENERATED_CPP_BY_NAME(MassSmartObjectTypes)
+
+namespace UE::Mass::SmartObject
+{
+void FMRUSlots::Push(const FSmartObjectSlotHandle& ClaimedSlot)
+{
+	int32 InsertIndex = NumSlots;
+	if (NumSlots == MaxNumSlots)
+	{
+		InsertIndex = NumSlots - 1;
+		for (int32 Index = 0; Index < InsertIndex; ++Index)
+		{
+			Slots[Index] = Slots[Index + 1];
+		}
+	}
+	else
+	{
+		NumSlots++;
+	}
+
+	const UMassSmartObjectSettings* Settings = GetDefault<UMassSmartObjectSettings>();
+	Slots[InsertIndex].Slot = ClaimedSlot;
+	Slots[InsertIndex].RemainingTime = (Settings && Settings->bUseCooldownForMRUSlots)
+		? Settings->MRUSlotsCooldown
+		: FMRUSlot::NoCoolDown;
+}
+
+void FMRUSlots::AppendTo(TArray<FSmartObjectSlotHandle>& OutSlots) const
+{
+	OutSlots.Reserve(OutSlots.Num() + NumSlots);
+
+	for (int32 SlotIndex = NumSlots - 1; SlotIndex >= 0; SlotIndex--)
+	{
+		OutSlots.Push(Slots[SlotIndex].Slot);
+	}
+}
+} // UE::Mass::SmartObject
diff --git a/Engine/Plugins/Runtime/MassGameplay/Source/MassSmartObjects/Public/MassSmartObjectFragments.h b/Engine/Plugins/Runtime/MassGameplay/Source/MassSmartObjects/Public/MassSmartObjectFragments.h
--- a/Engine/Plugins/Runtime/MassGameplay/Source/MassSmartObjects/Public/MassSmartObjectFragments.h
+++ b/Engine/Plugins/Runtime/MassGameplay/Source/MassSmartObjects/Public/MassSmartObjectFragments.h
@@ -45,3 +45,18 @@ struct FMassSmartObjectTimedBehaviorFragment : public FMassFragment
 	UPROPERTY(Transient)
 	float UseTime = 0.f;
 };
+
+namespace UE::Mass::SmartObject
+{
+
+/** Fragment used to track most recently used slots for a given smart object user */
+USTRUCT()
+struct FMRUSlotsFragment : public FMassFragment
+{
+	GENERATED_BODY()
+
+	UPROPERTY(Transient)
+	FMRUSlots Slots;
+};
+
+} //  UE::Mass::SmartObject
diff --git a/Engine/Plugins/Runtime/MassGameplay/Source/MassSmartObjects/Public/MassSmartObjectHandler.h b/Engine/Plugins/Runtime/MassGameplay/Source/MassSmartObjects/Public/MassSmartObjectHandler.h
--- a/Engine/Plugins/Runtime/MassGameplay/Source/MassSmartObjects/Public/MassSmartObjectHandler.h
+++ b/Engine/Plugins/Runtime/MassGameplay/Source/MassSmartObjects/Public/MassSmartObjectHandler.h
@@ -52,6 +52,7 @@ struct FMassSmartObjectHandler
 	 * @param Location The center of the query
 	 * @return Request identifier that can be used to try claiming a result once available
 	 */
+	UE_DEPRECATED(5.7, "Use the overload taking FFindCandidatesParameters instead")
 	[[nodiscard]] UE_API FMassSmartObjectRequestID FindCandidatesAsync(const FMassEntityHandle RequestingEntity, const FGameplayTagContainer& UserTags, const FGameplayTagQuery& ActivityRequirements, const FVector& Location) const;
 
 	/**
@@ -62,8 +63,19 @@ struct FMassSmartObjectHandler
 	 * @param LaneLocation The lane location as reference for the query
 	 * @return Request identifier that can be used to try claiming a result once available
 	 */
+	UE_DEPRECATED(5.7, "Use the overload taking FFindCandidatesParameters instead")
 	[[nodiscard]] UE_API FMassSmartObjectRequestID FindCandidatesAsync(const FMassEntityHandle RequestingEntity, const FGameplayTagContainer& UserTags, const FGameplayTagQuery& ActivityRequirements, const FZoneGraphCompactLaneLocation& LaneLocation) const;
 
+	/**
+	 * Creates an async request to build a list of compatible smart objects
+	 * using the provided set of parameters. The caller must poll using the request id
+	 * to know when the reservation can be done.
+	 * @param RequestingEntity Entity requesting the candidates list
+	 * @param Parameters All parameters of the query
+	 * @return Request identifier that can be used to try claiming a result once available
+	 */
+	[[nodiscard]] UE_API FMassSmartObjectRequestID FindCandidatesAsync(const FMassEntityHandle RequestingEntity, UE::Mass::SmartObject::FFindCandidatesParameters&& Parameters) const;
+
 	/**
 	 * Provides the result of a previously created request from FindCandidatesAsync to indicate if it has been processed
 	 * and the results can be used by ClaimCandidate.
diff --git a/Engine/Plugins/Runtime/MassGameplay/Source/MassSmartObjects/Public/MassSmartObjectProcessor.h b/Engine/Plugins/Runtime/MassGameplay/Source/MassSmartObjects/Public/MassSmartObjectProcessor.h
--- a/Engine/Plugins/Runtime/MassGameplay/Source/MassSmartObjects/Public/MassSmartObjectProcessor.h
+++ b/Engine/Plugins/Runtime/MassGameplay/Source/MassSmartObjects/Public/MassSmartObjectProcessor.h
@@ -65,5 +65,24 @@ protected:
 
 	FMassEntityQuery EntityQuery;
 };
+namespace UE::Mass::SmartObject
+{
+
+/** Processor to decay smart object MRU slots */
+UCLASS(MinimalAPI)
+class UMRUSlotsProcessor : public UMassProcessor
+{
+	GENERATED_BODY()
+
+public:
+	UE_API UMRUSlotsProcessor();
+
+protected:
+	UE_API virtual void Execute(FMassEntityManager& EntityManager, FMassExecutionContext& Context) override;
+	UE_API virtual void ConfigureQueries(const TSharedRef<FMassEntityManager>& EntityManager) override;
+
+	FMassEntityQuery EntityQuery;
+};
 
+} // UE::Mass::SmartObject
 #undef UE_API
diff --git a/Engine/Plugins/Runtime/MassGameplay/Source/MassSmartObjects/Public/MassSmartObjectRequest.h b/Engine/Plugins/Runtime/MassGameplay/Source/MassSmartObjects/Public/MassSmartObjectRequest.h
--- a/Engine/Plugins/Runtime/MassGameplay/Source/MassSmartObjects/Public/MassSmartObjectRequest.h
+++ b/Engine/Plugins/Runtime/MassGameplay/Source/MassSmartObjects/Public/MassSmartObjectRequest.h
@@ -8,6 +8,7 @@
 #include "Containers/StaticArray.h"
 #include "MassEntityHandle.h"
 #include "MassEntityTypes.h"
+#include "MassSmartObjectTypes.h"
 #include "SmartObjectRequestTypes.h"
 #include "ZoneGraphTypes.h"
 #include "MassSmartObjectRequest.generated.h"
@@ -164,3 +165,32 @@ struct FMassSmartObjectCompletedRequestTag : public FMassTag
 {
 	GENERATED_BODY()
 };
+
+namespace UE::Mass::SmartObject
+{
+
+/**
+ * Struct used to store parameters for FindCandidatesAsync requests
+ */
+USTRUCT()
+struct FFindCandidatesParameters
+{
+	GENERATED_BODY()
+
+	UPROPERTY(Transient)
+	FGameplayTagContainer UserTags;
+
+	UPROPERTY(Transient)
+	FGameplayTagQuery ActivityRequirements;
+
+	UPROPERTY(Transient)
+	FZoneGraphCompactLaneLocation LaneLocation;
+
+	UPROPERTY(Transient)
+	FVector Location = FVector::ZeroVector;
+
+	UPROPERTY(Transient)
+	FMRUSlots MRUSlots;
+};
+
+} // UE::Mass::SmartObject
\ No newline at end of file
diff --git a/Engine/Plugins/Runtime/MassGameplay/Source/MassSmartObjects/Public/MassSmartObjectSettings.h b/Engine/Plugins/Runtime/MassGameplay/Source/MassSmartObjects/Public/MassSmartObjectSettings.h
--- a/Engine/Plugins/Runtime/MassGameplay/Source/MassSmartObjects/Public/MassSmartObjectSettings.h
+++ b/Engine/Plugins/Runtime/MassGameplay/Source/MassSmartObjects/Public/MassSmartObjectSettings.h
@@ -34,6 +34,37 @@ public:
 	UPROPERTY(EditDefaultsOnly, Category = ZoneGraph, config)
 	float SearchExtents = 500.f;
 
+	/**
+	 * Can be used to prevent smart object users from reusing the same smart object slots more than once.
+	 * The list of most recently used slots can contain up to 4 slots, then a newly used slot will 
+	 * replace the least recently used one.
+	 * It is also possible to use a cooldown to allow slots to be considered again after a given
+	 * period of time.
+	 * @see bUseCooldownForMRUSlots
+	 * @see MRUSlotsCooldown
+	 */
+	//~ Max values must stay in sync with FMRUSlots::MaxSlots
+	UPROPERTY(EditDefaultsOnly, Category = SmartObject, config, meta = (UIMin = "0", ClampMin = "0", UIMax = "4", ClampMax = "4", DisplayName = ""))
+	int32 MRUSlotsMaxCount = 0;
+
+	/**
+	 * Set this flag to allow the most recently used slots to be considered again after the period of time specified by MRUSlotsCooldown
+	 * This option requires MRUSlotsMaxCount to be a non-zero value
+	 * @see MRUSlotsCooldown
+	 * @see MRUSlotsMaxCount
+	 */
+	UPROPERTY(config, EditAnywhere, Category = SmartObject, meta = (EditCondition = "MRUSlotsMaxCount > 0", EditConditionHides))
+	bool bUseCooldownForMRUSlots = false;
+
+	/**
+	 * Period of time during which a given slot will not be considered after being used
+	 * This option requires 'bUseCooldownForMRUSlots' to be true and 'MRUSlotsMaxCount' to be a non-zero value
+	 * @see bUseCooldownForMRUSlots
+	 * @see MRUSlotsMaxCount
+	 */
+	UPROPERTY(EditDefaultsOnly, Category = SmartObject, config, meta = (UIMin = "0", ClampMin = "0", EditCondition = "MRUSlotsMaxCount > 0 && bUseCooldownForMRUSlots", EditConditionHides))
+	float MRUSlotsCooldown = 0.f;
+
 protected:
 
 #if WITH_EDITOR
diff --git a/Engine/Plugins/Runtime/MassGameplay/Source/MassSmartObjects/Public/MassSmartObjectTypes.h b/Engine/Plugins/Runtime/MassGameplay/Source/MassSmartObjects/Public/MassSmartObjectTypes.h
--- a/Engine/Plugins/Runtime/MassGameplay/Source/MassSmartObjects/Public/MassSmartObjectTypes.h
+++ b/Engine/Plugins/Runtime/MassGameplay/Source/MassSmartObjects/Public/MassSmartObjectTypes.h
@@ -6,8 +6,11 @@
 #include "MassEntityTypes.h"
 #endif // UE_ENABLE_INCLUDE_ORDER_DEPRECATED_IN_5_6
 #include "MassEntityHandle.h"
+#include "SmartObjectTypes.h"
 #include "MassSmartObjectTypes.generated.h"
 
+#define UE_API MASSSMARTOBJECTS_API
+
 namespace UE::Mass::Signals
 {
 	const FName SmartObjectRequestCandidates = FName(TEXT("SmartObjectRequestCandidates"));
@@ -17,13 +20,13 @@ namespace UE::Mass::Signals
 }
 
 UENUM()
-enum class EMassSmartObjectInteractionStatus: uint8
+enum class EMassSmartObjectInteractionStatus : uint8
 {
 	Unset,
-    InProgress,			// Claimed and Behavior activated
-    BehaviorCompleted,	// Behavior is completed but task still running (not updated yet)
+	InProgress,			// Claimed and Behavior activated
+	BehaviorCompleted,	// Behavior is completed but task still running (not updated yet)
 	TaskCompleted,		// Task has been notified that behavior is completed and completes
-    Aborted				// Task and Behavior were aborted
+	Aborted				// Task and Behavior were aborted
 };
 
 /**
@@ -67,4 +70,54 @@ struct FSmartObjectMassEntityUserData
 
 	UPROPERTY()
 	FMassEntityHandle UserEntity;
-};
\ No newline at end of file
+};
+
+namespace UE::Mass::SmartObject
+{
+
+/**
+ * Struct holding data related to a single recently used smart object slot.
+ */
+USTRUCT()
+struct FMRUSlot
+{
+	GENERATED_BODY()
+
+	static constexpr float NoCoolDown = -1.f;
+
+	UPROPERTY(VisibleAnywhere, Category = SmartObject, transient)
+	FSmartObjectSlotHandle Slot;
+
+	UPROPERTY(VisibleAnywhere, Category = SmartObject, transient)
+	float RemainingTime = NoCoolDown;
+};
+
+/**
+ * Struct used as a container for the most recently used slots.
+ */
+USTRUCT()
+struct FMRUSlots
+{
+	GENERATED_BODY()
+
+	void Reset()
+	{
+		NumSlots = 0;
+	}
+
+	UE_API void Push(const FSmartObjectSlotHandle& ClaimedSlot);
+	UE_API void AppendTo(TArray<FSmartObjectSlotHandle>& OutSlots) const;
+
+	/** Using a fixed maximum number of slots to use fixed inline allocation in the fragments. */
+	static constexpr uint32 MaxNumSlots = 4;
+
+	UPROPERTY(Transient)
+	FMRUSlot Slots[MaxNumSlots];
+
+	UPROPERTY(Transient)
+	uint8 NumSlots = 0;
+};
+
+} // UE::Mass::SmartObject
+
+#undef UE_API
