[Mass] removed code deprecated in 5.4

#rnx

diff --git a/Engine/Plugins/Runtime/MassGameplay/Source/MassActors/Private/MassActorSpawnerSubsystem.cpp b/Engine/Plugins/Runtime/MassGameplay/Source/MassActors/Private/MassActorSpawnerSubsystem.cpp
--- a/Engine/Plugins/Runtime/MassGameplay/Source/MassActors/Private/MassActorSpawnerSubsystem.cpp
+++ b/Engine/Plugins/Runtime/MassGameplay/Source/MassActors/Private/MassActorSpawnerSubsystem.cpp
@@ -581,12 +581,3 @@ void UMassActorSpawnerSubsystem::ReleaseAllResources()
 	CSV_CUSTOM_STAT(MassActors, NumSpawned, NumActorSpawned, ECsvCustomStatOp::Accumulate);
 	CSV_CUSTOM_STAT(MassActors, NumPooled, NumActorPooled, ECsvCustomStatOp::Accumulate);
 }
-
-//-----------------------------------------------------------------------------
-// DEPRECATED
-//-----------------------------------------------------------------------------
-FMassActorSpawnRequestHandle UMassActorSpawnerSubsystem::GetNextRequestToSpawn() const
-{
-	int32 DummyIndex = 0;
-	return GetNextRequestToSpawn(DummyIndex);
-}
diff --git a/Engine/Plugins/Runtime/MassGameplay/Source/MassActors/Public/MassActorSpawnerSubsystem.h b/Engine/Plugins/Runtime/MassGameplay/Source/MassActors/Public/MassActorSpawnerSubsystem.h
--- a/Engine/Plugins/Runtime/MassGameplay/Source/MassActors/Public/MassActorSpawnerSubsystem.h
+++ b/Engine/Plugins/Runtime/MassGameplay/Source/MassActors/Public/MassActorSpawnerSubsystem.h
@@ -306,8 +306,4 @@ protected:
 	mutable int32 NumActorPooled = 0;
 
 	int32 StartingHandleIndex = INDEX_NONE;
-
-public:
-	UE_DEPRECATED(5.4, "This flavor of GetNextRequestToSpawn is deprecated. Use the alternative taking an int32& parameter")
-	MASSACTORS_API virtual FMassActorSpawnRequestHandle GetNextRequestToSpawn() const final;
 };
diff --git a/Engine/Plugins/Runtime/MassGameplay/Source/MassLOD/Private/MassLODSubsystem.cpp b/Engine/Plugins/Runtime/MassGameplay/Source/MassLOD/Private/MassLODSubsystem.cpp
--- a/Engine/Plugins/Runtime/MassGameplay/Source/MassLOD/Private/MassLODSubsystem.cpp
+++ b/Engine/Plugins/Runtime/MassGameplay/Source/MassLOD/Private/MassLODSubsystem.cpp
@@ -621,13 +621,6 @@ void UMassLODSubsystem::UnregisterActorViewer(AActor& ActorViewer)
 	}
 }
 
-FMassViewerHandle UMassLODSubsystem::GetViewerHandleFromPlayerController(const APlayerController* PlayerController) const
-{
-	PRAGMA_DISABLE_DEPRECATION_WARNINGS
-	return PlayerController ? GetViewerHandleFromActor(*PlayerController) : FMassViewerHandle();
-	PRAGMA_ENABLE_DEPRECATION_WARNINGS
-}
-
 #if WITH_MASSGAMEPLAY_DEBUG
 void UMassLODSubsystem::DebugUnregisterActorViewer()
 {
diff --git a/Engine/Plugins/Runtime/MassGameplay/Source/MassLOD/Public/MassLODSubsystem.h b/Engine/Plugins/Runtime/MassGameplay/Source/MassLOD/Public/MassLODSubsystem.h
--- a/Engine/Plugins/Runtime/MassGameplay/Source/MassLOD/Public/MassLODSubsystem.h
+++ b/Engine/Plugins/Runtime/MassGameplay/Source/MassLOD/Public/MassLODSubsystem.h
@@ -203,8 +203,4 @@ private:
 	/** Delegates to notify anyone who needs to know about viewer changes */
 	FOnViewerAdded OnViewerAddedDelegate;
 	FOnViewerRemoved OnViewerRemovedDelegate;
-
-public:
-	UE_DEPRECATED(5.4, "GetViewerHandleFromPlayerController is deprecated. Use more teneric GetViewerHandleFromActor")
-	MASSLOD_API FMassViewerHandle GetViewerHandleFromPlayerController(const APlayerController* PlayerController) const;
 };
diff --git a/Engine/Plugins/Runtime/MassGameplay/Source/MassRepresentation/Private/MassRepresentationActorManagement.cpp b/Engine/Plugins/Runtime/MassGameplay/Source/MassRepresentation/Private/MassRepresentationActorManagement.cpp
--- a/Engine/Plugins/Runtime/MassGameplay/Source/MassRepresentation/Private/MassRepresentationActorManagement.cpp
+++ b/Engine/Plugins/Runtime/MassGameplay/Source/MassRepresentation/Private/MassRepresentationActorManagement.cpp
@@ -145,34 +145,3 @@ void UMassRepresentationActorManagement::ReleaseAnyActorOrCancelAnySpawning(UMas
 	}
 	check(!Representation.ActorSpawnRequestHandle.IsValid());
 }
-
-
-//-----------------------------------------------------------------------------
-// DEPRECATED
-//-----------------------------------------------------------------------------
-AActor* UMassRepresentationActorManagement::GetOrSpawnActor(UMassRepresentationSubsystem& RepresentationSubsystem, FMassEntityManager& EntityManager
-	, const FMassEntityHandle MassAgent, FMassActorFragment&/* OutActorInfo*/, const FTransform& Transform, const int16 TemplateActorIndex
-	, FMassActorSpawnRequestHandle& InOutSpawnRequestHandle, const float Priority) const
-{
-	return GetOrSpawnActor(RepresentationSubsystem, EntityManager
-		, MassAgent, Transform, TemplateActorIndex
-		, InOutSpawnRequestHandle, Priority);
-}
-
-void UMassRepresentationActorManagement::OnPreActorSpawn(const FMassActorSpawnRequestHandle& SpawnRequestHandle, FConstStructView SpawnRequest, FMassEntityManager* EntityManager) const
-{
-	if (EntityManager)
-	{
-		OnPreActorSpawn(SpawnRequestHandle, SpawnRequest, EntityManager->AsShared());
-	}
-}
-
-EMassActorSpawnRequestAction UMassRepresentationActorManagement::OnPostActorSpawn(const FMassActorSpawnRequestHandle& SpawnRequestHandle, FConstStructView SpawnRequest, FMassEntityManager* EntityManager) const
-{
-	if (EntityManager)
-	{
-		return OnPostActorSpawn(SpawnRequestHandle, SpawnRequest, EntityManager->AsShared());
-	}
-
-	return EMassActorSpawnRequestAction::Remove;
-}
diff --git a/Engine/Plugins/Runtime/MassGameplay/Source/MassRepresentation/Private/MassRepresentationSubsystem.cpp b/Engine/Plugins/Runtime/MassGameplay/Source/MassRepresentation/Private/MassRepresentationSubsystem.cpp
--- a/Engine/Plugins/Runtime/MassGameplay/Source/MassRepresentation/Private/MassRepresentationSubsystem.cpp
+++ b/Engine/Plugins/Runtime/MassGameplay/Source/MassRepresentation/Private/MassRepresentationSubsystem.cpp
@@ -52,14 +52,6 @@ const FMassISMCSharedData* UMassRepresentationSubsystem::GetISMCSharedDataForIns
 	return VisualizationComponent->GetISMCSharedDataForInstancedStaticMesh(ISMC);
 }
 
-void UMassRepresentationSubsystem::RemoveISMComponent(UInstancedStaticMeshComponent& ISMComponent)
-{
-PRAGMA_DISABLE_DEPRECATION_WARNINGS
-	check(VisualizationComponent);
-	return VisualizationComponent->RemoveISMComponent(ISMComponent);
-PRAGMA_ENABLE_DEPRECATION_WARNINGS
-}
-
 void UMassRepresentationSubsystem::RemoveVisualDesc(const FStaticMeshInstanceVisualizationDescHandle VisualizationHandle)
 {
 	check(VisualizationComponent);
@@ -440,7 +432,7 @@ void UMassRepresentationSubsystem::Deinitialize()
 	Super::Deinitialize();
 }
 
-void UMassRepresentationSubsystem::OnProcessingPhaseStarted(const float DeltaSeconds, const EMassProcessingPhase Phase) const
+void UMassRepresentationSubsystem::OnProcessingPhaseStarted(const float /*DeltaSeconds*/, const EMassProcessingPhase Phase) const
 {
 	check(VisualizationComponent);
 	switch (Phase)
diff --git a/Engine/Plugins/Runtime/MassGameplay/Source/MassRepresentation/Private/MassVisualizationComponent.cpp b/Engine/Plugins/Runtime/MassGameplay/Source/MassRepresentation/Private/MassVisualizationComponent.cpp
--- a/Engine/Plugins/Runtime/MassGameplay/Source/MassRepresentation/Private/MassVisualizationComponent.cpp
+++ b/Engine/Plugins/Runtime/MassGameplay/Source/MassRepresentation/Private/MassVisualizationComponent.cpp
@@ -737,22 +737,3 @@ void FMassLODSignificanceRange::WriteCustomDataFloatsAtStartIndex(int32 StaticMe
 		}
 	}
 }
-
-//-----------------------------------------------------------------------------
-// DEPRECATED
-//-----------------------------------------------------------------------------
-void UMassVisualizationComponent::RemoveISMComponent(UInstancedStaticMeshComponent& ISMComponent)
-{
-	UE_MT_SCOPED_WRITE_ACCESS(InstancedStaticMeshInfosDetector);
-
-	const FStaticMeshInstanceVisualizationDescHandle* VisualDescHandlePtr = ISMComponentMap.Find(&ISMComponent);
-	if (VisualDescHandlePtr)
-	{
-		RemoveVisualDesc(*VisualDescHandlePtr);
-	}
-}
-
-void UMassVisualizationComponent::RemoveVisualDescByIndex(const int32 VisualizationIndex)
-{
-	RemoveVisualDesc(FStaticMeshInstanceVisualizationDescHandle(VisualizationIndex));
-}
diff --git a/Engine/Plugins/Runtime/MassGameplay/Source/MassRepresentation/Public/MassRepresentationActorManagement.h b/Engine/Plugins/Runtime/MassGameplay/Source/MassRepresentation/Public/MassRepresentationActorManagement.h
--- a/Engine/Plugins/Runtime/MassGameplay/Source/MassRepresentation/Public/MassRepresentationActorManagement.h
+++ b/Engine/Plugins/Runtime/MassGameplay/Source/MassRepresentation/Public/MassRepresentationActorManagement.h
@@ -107,17 +107,6 @@ public:
 	 */
 	static UE_API void ReleaseAnyActorOrCancelAnySpawning(UMassRepresentationSubsystem& RepresentationSubsystem, const FMassEntityHandle MassAgent
 		, FMassActorFragment& ActorInfo, FMassRepresentationFragment& Representation, UMassActorSubsystem* ActorSubsystem = nullptr);
-
-	UE_DEPRECATED(5.4, "This flavor of GetOrSpawnActor has been deprecated due to a defunct parameter, OutActorInfo, that was never being used.")
-	UE_API virtual AActor* GetOrSpawnActor(UMassRepresentationSubsystem& RepresentationSubsystem, FMassEntityManager& EntityManager
-		, const FMassEntityHandle MassAgent, FMassActorFragment&/* OutActorInfo*/, const FTransform& Transform, const int16 TemplateActorIndex
-		, FMassActorSpawnRequestHandle& InOutSpawnRequestHandle, const float Priority) const final;
-
-	UE_DEPRECATED(5.4, "This flavor of OnPreActorSpawn has been deprecated in favor of the version using TSharedPtr<FMassEntityManager> as a parameter. Use that one instead.")
-	UE_API virtual void OnPreActorSpawn(const FMassActorSpawnRequestHandle& SpawnRequestHandle, FConstStructView SpawnRequest, FMassEntityManager* EntityManager) const final;
-
-	UE_DEPRECATED(5.4, "This flavor of OnPostActorSpawn has been deprecated in favor of the version using TSharedPtr<FMassEntityManager> as a parameter. Use that one instead.")
-	UE_API virtual EMassActorSpawnRequestAction OnPostActorSpawn(const FMassActorSpawnRequestHandle& SpawnRequestHandle, FConstStructView SpawnRequest, FMassEntityManager* EntityManager) const final;
 };
 
 #undef UE_API
diff --git a/Engine/Plugins/Runtime/MassGameplay/Source/MassRepresentation/Public/MassRepresentationSubsystem.h b/Engine/Plugins/Runtime/MassGameplay/Source/MassRepresentation/Public/MassRepresentationSubsystem.h
--- a/Engine/Plugins/Runtime/MassGameplay/Source/MassRepresentation/Public/MassRepresentationSubsystem.h
+++ b/Engine/Plugins/Runtime/MassGameplay/Source/MassRepresentation/Public/MassRepresentationSubsystem.h
@@ -63,15 +63,6 @@ public:
 	 */
 	MASSREPRESENTATION_API const FMassISMCSharedData* GetISMCSharedDataForInstancedStaticMesh(const UInstancedStaticMeshComponent* ISMC) const;
 
-	/**
-	 * Removes the visualization data associated with the given ISM component. Note that this is safe to do only when
-	 * there are no entities relying on this data. No entity data patching will take place.
-	 * Note that the function will assert if there's more ISM components associated with given visualization. Also, in 
-	 * that case RemoveVisualDescByIndex will be called under the hood. 
-	 */
-	UE_DEPRECATED(5.4, "RemoveISMComponent has been deprecated in favor of RemoveVisualDescByIndex. Please use that instead.")
-	MASSREPRESENTATION_API void RemoveISMComponent(UInstancedStaticMeshComponent& ISMComponent);
-
 	/** 
 	 * Removes all data associated with a given VisualizationIndex. Note that this is safe to do only if there are no
 	 * entities relying on this index. No entity data patching will take place.
diff --git a/Engine/Plugins/Runtime/MassGameplay/Source/MassRepresentation/Public/MassRepresentationTypes.h b/Engine/Plugins/Runtime/MassGameplay/Source/MassRepresentation/Public/MassRepresentationTypes.h
--- a/Engine/Plugins/Runtime/MassGameplay/Source/MassRepresentation/Public/MassRepresentationTypes.h
+++ b/Engine/Plugins/Runtime/MassGameplay/Source/MassRepresentation/Public/MassRepresentationTypes.h
@@ -210,16 +210,6 @@ struct alignas(2) FStaticMeshInstanceVisualizationDescHandle
 
 	bool operator==(const FStaticMeshInstanceVisualizationDescHandle& Other) const = default;
 
-	UE_DEPRECATED(5.4, "Referring to registered FStaticMeshInstanceVisualizationDesc's by raw int16 index has been deprecated. Please use strictly typed FStaticMeshInstanceVisualizationDescHandle instead.")
-	operator int16() const
-	{
-		if (!IsValid())
-		{
-			return INDEX_NONE;
-		}
-		return ensure(Index < TNumericLimits<int16>::Max()) ? static_cast<int16>(Index) : INDEX_NONE; 
-	}
-
 private:
 
 	UPROPERTY()
@@ -345,15 +335,6 @@ private:
 
 protected:
 	FEntityToPrimitiveIdMap EntityHandleToPrimitiveIdMap;
-
-	UE_DEPRECATED(5.4, "RefCount is deprecated, use ISMComponentReferencesCount instead")
-	int32 RefCount = 0;
-
-public:
-	UE_DEPRECATED(5.4, "StoreReference is deprecated, use OnISMComponentReferenceStored instead")
-	int32 StoreReference() { return OnISMComponentReferenceStored(); }
-	UE_DEPRECATED(5.4, "ReleaseReference is deprecated, use OnISMComponentReferenceReleased instead")
-	int32 ReleaseReference() { return OnISMComponentReferenceReleased(); }
 };
 
 
@@ -605,17 +586,6 @@ public:
 	TArray<FISMCSharedDataKey> StaticMeshRefs;
 
 	FMassISMCSharedDataMap* ISMCSharedDataPtr = nullptr;
-
-	//-----------------------------------------------------------------------------
-	// DEPRECATED
-	//-----------------------------------------------------------------------------
-	UE_DEPRECATED(5.4, "Deprecated in favor of new version taking FMassEntityHandle parameter instead of int32 to identify the entity. This deprecated function is now defunct.")
-	void AddBatchedTransform(const int32 InstanceId, const FTransform& Transform, const FTransform& PrevTransform, const TArray<uint32>& ExcludeStaticMeshRefs) {}
-	UE_DEPRECATED(5.4, "Deprecated in favor of new version taking FMassEntityHandle parameter instead of int32 to identify the entity. This deprecated function is now defunct.")
-	void AddInstance(const int32 InstanceId, const FTransform& Transform) {}
-	UE_DEPRECATED(5.4, "Deprecated in favor of new version taking FMassEntityHandle parameter instead of int32 to identify the entity. This deprecated function is now defunct.")
-	void RemoveInstance(const int32 InstanceId) {}
-
 };
 
 USTRUCT()
@@ -773,11 +743,6 @@ protected:
 	//-----------------------------------------------------------------------------
 	UE_DEPRECATED(5.5, "Deprecated in flavor of the function taking the ISMComponent parameter. This version is not defunct")
 	void ClearVisualInstance(FMassISMCSharedDataMap& ISMCSharedData) {}
-public:
-	UE_DEPRECATED(5.4, "Deprecated in flavor of new version taking FMassEntityHandle parameter instead of int32 to identify the entity. This deprecated function is now defunct.")
-	void AddBatchedTransform(const int32 InstanceId, const FTransform& Transform, const FTransform& PrevTransform, const float LODSignificance, const float PrevLODSignificance = -1.0f) {}
-	UE_DEPRECATED(5.4, "Deprecated in flavor of new version taking FMassEntityHandle parameter instead of int32 to identify the entity. This deprecated function is now defunct.")
-	void RemoveInstance(const int32 InstanceId, const float LODSignificance) {}
 };
 
 #if ENABLE_MT_DETECTOR
diff --git a/Engine/Plugins/Runtime/MassGameplay/Source/MassRepresentation/Public/MassUpdateISMProcessor.h b/Engine/Plugins/Runtime/MassGameplay/Source/MassRepresentation/Public/MassUpdateISMProcessor.h
--- a/Engine/Plugins/Runtime/MassGameplay/Source/MassRepresentation/Public/MassUpdateISMProcessor.h
+++ b/Engine/Plugins/Runtime/MassGameplay/Source/MassRepresentation/Public/MassUpdateISMProcessor.h
@@ -21,10 +21,6 @@ public:
 
 	static UE_API void UpdateISMTransform(FMassEntityHandle EntityHandle, FMassInstancedStaticMeshInfo& ISMInfo, const FTransform& Transform, const FTransform& PrevTransform, const float LODSignificance, const float PrevLODSignificance = -1.0f);
 
-	UE_DEPRECATED(5.4, "Deprecated in favor of new version taking FMassEntityHandle parameter instead of int32 to identify the entity. This deprecated function is now defunct.")
-	static void UpdateISMTransform(int32 EntityId, FMassInstancedStaticMeshInfo& ISMInfo, const FTransform& Transform, const FTransform& PrevTransform, const float LODSignificance, const float PrevLODSignificance = -1.0f) {}
-
-
 protected:
 
 	/** Configure the owned FMassEntityQuery instances to express processor's requirements */
diff --git a/Engine/Plugins/Runtime/MassGameplay/Source/MassRepresentation/Public/MassVisualizationComponent.h b/Engine/Plugins/Runtime/MassGameplay/Source/MassRepresentation/Public/MassVisualizationComponent.h
--- a/Engine/Plugins/Runtime/MassGameplay/Source/MassRepresentation/Public/MassVisualizationComponent.h
+++ b/Engine/Plugins/Runtime/MassGameplay/Source/MassRepresentation/Public/MassVisualizationComponent.h
@@ -53,22 +53,6 @@ public:
 	 */
 	MASSREPRESENTATION_API const FMassISMCSharedData* GetISMCSharedDataForInstancedStaticMesh(const UInstancedStaticMeshComponent* ISMC) const;
 
-	/** 
-	 * Removes all the visualization data associated with the given ISM component. The function resolves the VisualizationIndex
-	 * associated with the given ISMComponent and calls RemoveVisualDescByIndex which will remove data on all ISMComponents
-	 * associated with the index. Note that this is safe to do only when there are no entities relying on this data. 
-	 * No entity data patching will take place. 
-	 */
-	UE_DEPRECATED(5.4, "RemoveISMComponent has been deprecated in favor of RemoveVisualDesc. Please use that instead.")
-	MASSREPRESENTATION_API void RemoveISMComponent(UInstancedStaticMeshComponent& ISMComponent);
-
-	/**
-	 * Removes all data associated with a given VisualizationIndex. Note that this is safe to do only if there are no
-	 * entities relying on this index. No entity data patching will take place.
-	 */
-	UE_DEPRECATED(5.4, "RemoveVisualDescByIndex has been deprecated in favor of RemoveVisualDesc. Please use that instead.")
-	MASSREPRESENTATION_API void RemoveVisualDescByIndex(const int32 VisualizationIndex);
-
 	/**
 	 * Removes all data associated with a given VisualizationIndex. Note that this is safe to do only if there are no
 	 * entities relying on this index. No entity data patching will take place.
@@ -143,12 +127,9 @@ protected:
 	 */
 	TMap<uint32, FISMCSharedDataKey> MeshDescToISMCMap;
 
-	/** Indicies to InstancedStaticMeshInfos that need their SMComponent constructed */
+	/** Indices to InstancedStaticMeshInfos that need their SMComponent constructed */
 	TArray<FStaticMeshInstanceVisualizationDescHandle> InstancedSMComponentsRequiringConstructing;
 
-	UE_DEPRECATED(5.4, "This flavor of BuildLODSignificanceForInfo is no longer supported and is defunct.")
-	void BuildLODSignificanceForInfo(FMassInstancedStaticMeshInfo& Info, const uint32 ForcedStaticMeshRefKey){}
-
 	UE_DEPRECATED(5.5, "This flavor of BuildLODSignificanceForInfo is no longer supported and is defunct.")
 	void BuildLODSignificanceForInfo(FMassInstancedStaticMeshInfo& Info) {}
 };
diff --git a/Engine/Source/Runtime/MassEntity/Public/MassExecutionContext.h b/Engine/Source/Runtime/MassEntity/Public/MassExecutionContext.h
--- a/Engine/Source/Runtime/MassEntity/Public/MassExecutionContext.h
+++ b/Engine/Source/Runtime/MassEntity/Public/MassExecutionContext.h
@@ -591,13 +591,6 @@ protected:
 			View.FragmentView.Reset();
 		}
 	}
-
-public:
-	//-----------------------------------------------------------------------------
-	// DEPRECATED functions
-	//-----------------------------------------------------------------------------
-	UE_DEPRECATED(5.4, "Deprecated in favor of 'SetCurrentArchetypeCompositionDescriptor' as this provides information on the entire archetype.")
-	void SetCurrentArchetypesTagBitSet(const FMassTagBitSet&) {}
 };
 
 //-----------------------------------------------------------------------------
