Change InstancedStaticMeshInfosDetector to use non-recursive version so that invalid write access is caught on the same thread.

diff --git a/Engine/Plugins/Runtime/MassGameplay/Source/MassRepresentation/Public/MassRepresentationTypes.h b/Engine/Plugins/Runtime/MassGameplay/Source/MassRepresentation/Public/MassRepresentationTypes.h
--- a/Engine/Plugins/Runtime/MassGameplay/Source/MassRepresentation/Public/MassRepresentationTypes.h
+++ b/Engine/Plugins/Runtime/MassGameplay/Source/MassRepresentation/Public/MassRepresentationTypes.h
@@ -750,7 +750,7 @@ protected:
 #define MAKE_MASS_INSTANCED_STATIC_MESH_INFO_ARRAY_VIEW(ArrayView, AccessDetector) FMassInstancedStaticMeshInfoArrayView(ArrayView, AccessDetector)
 struct FMassInstancedStaticMeshInfoArrayViewAccessDetector
 {
-	FMassInstancedStaticMeshInfoArrayViewAccessDetector(TArrayView<FMassInstancedStaticMeshInfo> InInstancedStaticMeshInfos, const FRWRecursiveAccessDetector& InAccessDetector)
+	FMassInstancedStaticMeshInfoArrayViewAccessDetector(TArrayView<FMassInstancedStaticMeshInfo> InInstancedStaticMeshInfos, const FRWAccessDetector& InAccessDetector)
 		: InstancedStaticMeshInfos(InInstancedStaticMeshInfos)
 		, AccessDetector(&InAccessDetector)
 	{
@@ -791,7 +791,7 @@ struct FMassInstancedStaticMeshInfoArrayViewAccessDetector
 
 private:
 	TArrayView<FMassInstancedStaticMeshInfo> InstancedStaticMeshInfos;
-	const FRWRecursiveAccessDetector* AccessDetector;
+	const FRWAccessDetector* AccessDetector;
 };
 
 typedef FMassInstancedStaticMeshInfoArrayViewAccessDetector FMassInstancedStaticMeshInfoArrayView;
diff --git a/Engine/Plugins/Runtime/MassGameplay/Source/MassRepresentation/Public/MassVisualizationComponent.h b/Engine/Plugins/Runtime/MassGameplay/Source/MassRepresentation/Public/MassVisualizationComponent.h
--- a/Engine/Plugins/Runtime/MassGameplay/Source/MassRepresentation/Public/MassVisualizationComponent.h
+++ b/Engine/Plugins/Runtime/MassGameplay/Source/MassRepresentation/Public/MassVisualizationComponent.h
@@ -109,7 +109,7 @@ protected:
 	/** The information of all the instanced static meshes. Make sure to use AddInstancedStaticMeshInfo to add elements to it */
 	UPROPERTY(Transient)
 	TArray<FMassInstancedStaticMeshInfo> InstancedStaticMeshInfos;
-	UE_MT_DECLARE_RW_RECURSIVE_ACCESS_DETECTOR(InstancedStaticMeshInfosDetector);
+	UE_MT_DECLARE_RW_ACCESS_DETECTOR(InstancedStaticMeshInfosDetector);
 
 	/** Indices to InstancedStaticMeshInfos that have been released and can be reused */
 	TArray<FStaticMeshInstanceVisualizationDescHandle> InstancedStaticMeshInfosFreeIndices;
