[MassGameplay]
* Fixed non unity compile errors

diff --git a/Engine/Plugins/Runtime/InstancedActors/Source/InstancedActors/Public/InstancedActorsInitializerProcessor.h b/Engine/Plugins/Runtime/InstancedActors/Source/InstancedActors/Public/InstancedActorsInitializerProcessor.h
--- a/Engine/Plugins/Runtime/InstancedActors/Source/InstancedActors/Public/InstancedActorsInitializerProcessor.h
+++ b/Engine/Plugins/Runtime/InstancedActors/Source/InstancedActors/Public/InstancedActorsInitializerProcessor.h
@@ -2,6 +2,7 @@
 
 #pragma once
 
+#include "MassEntityQuery.h"
 #include "MassProcessor.h"
 #include "InstancedActorsInitializerProcessor.generated.h"
 
diff --git a/Engine/Plugins/Runtime/MassGameplay/Source/MassEQS/Public/Processors/MassEnvQueryTestProcessor_MassEntityTags.h b/Engine/Plugins/Runtime/MassGameplay/Source/MassEQS/Public/Processors/MassEnvQueryTestProcessor_MassEntityTags.h
--- a/Engine/Plugins/Runtime/MassGameplay/Source/MassEQS/Public/Processors/MassEnvQueryTestProcessor_MassEntityTags.h
+++ b/Engine/Plugins/Runtime/MassGameplay/Source/MassEQS/Public/Processors/MassEnvQueryTestProcessor_MassEntityTags.h
@@ -2,6 +2,7 @@
 
 #pragma once
 
+#include "MassEntityQuery.h"
 #include "Processors/MassEnvQueryProcessorBase.h"
 #include "MassEnvQueryTestProcessor_MassEntityTags.generated.h"
 
diff --git a/Engine/Plugins/Runtime/MassGameplay/Source/MassMovement/Public/Example/MassSimpleMovementTrait.h b/Engine/Plugins/Runtime/MassGameplay/Source/MassMovement/Public/Example/MassSimpleMovementTrait.h
--- a/Engine/Plugins/Runtime/MassGameplay/Source/MassMovement/Public/Example/MassSimpleMovementTrait.h
+++ b/Engine/Plugins/Runtime/MassGameplay/Source/MassMovement/Public/Example/MassSimpleMovementTrait.h
@@ -2,6 +2,7 @@
 
 #pragma once
 
+#include "MassEntityQuery.h"
 #include "MassEntityTypes.h"
 #include "MassEntityTraitBase.h"
 #include "MassProcessor.h"
diff --git a/Engine/Plugins/Runtime/MassGameplay/Source/MassMovement/Public/Example/MassVelocityRandomizerTrait.h b/Engine/Plugins/Runtime/MassGameplay/Source/MassMovement/Public/Example/MassVelocityRandomizerTrait.h
--- a/Engine/Plugins/Runtime/MassGameplay/Source/MassMovement/Public/Example/MassVelocityRandomizerTrait.h
+++ b/Engine/Plugins/Runtime/MassGameplay/Source/MassMovement/Public/Example/MassVelocityRandomizerTrait.h
@@ -2,6 +2,7 @@
 
 #pragma once
 
+#include "MassEntityQuery.h"
 #include "MassEntityTraitBase.h"
 #include "MassObserverProcessor.h"
 #include "MassVelocityRandomizerTrait.generated.h"
diff --git a/Engine/Plugins/Runtime/MassGameplay/Source/MassRepresentation/Public/MassStationaryISMSwitcherProcessor.h b/Engine/Plugins/Runtime/MassGameplay/Source/MassRepresentation/Public/MassStationaryISMSwitcherProcessor.h
--- a/Engine/Plugins/Runtime/MassGameplay/Source/MassRepresentation/Public/MassStationaryISMSwitcherProcessor.h
+++ b/Engine/Plugins/Runtime/MassGameplay/Source/MassRepresentation/Public/MassStationaryISMSwitcherProcessor.h
@@ -2,6 +2,7 @@
 
 #pragma once 
 
+#include "MassEntityQuery.h"
 #include "MassProcessor.h"
 #include "MassStationaryISMSwitcherProcessor.generated.h"
 
diff --git a/Engine/Plugins/Runtime/MassGameplay/Source/MassSpawner/Private/MassTranslator.cpp b/Engine/Plugins/Runtime/MassGameplay/Source/MassSpawner/Private/MassTranslator.cpp
--- a/Engine/Plugins/Runtime/MassGameplay/Source/MassSpawner/Private/MassTranslator.cpp
+++ b/Engine/Plugins/Runtime/MassGameplay/Source/MassSpawner/Private/MassTranslator.cpp
@@ -2,6 +2,7 @@
 
 #include "MassTranslator.h"
 #include "MassCommonTypes.h"
+#include "MassEntityQuery.h"
 
 //----------------------------------------------------------------------//
 //  UMassTranslator
