[MassGameplay]
* Fixed another non unity

diff --git a/Engine/Plugins/Runtime/MassGameplay/Source/MassRepresentation/Public/MassUpdateISMProcessor.h b/Engine/Plugins/Runtime/MassGameplay/Source/MassRepresentation/Public/MassUpdateISMProcessor.h
--- a/Engine/Plugins/Runtime/MassGameplay/Source/MassRepresentation/Public/MassUpdateISMProcessor.h
+++ b/Engine/Plugins/Runtime/MassGameplay/Source/MassRepresentation/Public/MassUpdateISMProcessor.h
@@ -2,6 +2,7 @@
 
 #pragma once
 
+#include "MassEntityQuery.h"
 #include "MassProcessor.h"
 #include "MassUpdateISMProcessor.generated.h"
 
