[MassGameplay]
Added handler for premature destruction of Mass Representaion visualizer actor

diff --git a/Engine/Plugins/Runtime/MassGameplay/Source/MassRepresentation/Private/MassRepresentationSubsystem.cpp b/Engine/Plugins/Runtime/MassGameplay/Source/MassRepresentation/Private/MassRepresentationSubsystem.cpp
--- a/Engine/Plugins/Runtime/MassGameplay/Source/MassRepresentation/Private/MassRepresentationSubsystem.cpp
+++ b/Engine/Plugins/Runtime/MassGameplay/Source/MassRepresentation/Private/MassRepresentationSubsystem.cpp
@@ -379,22 +379,8 @@ void UMassRepresentationSubsystem::Initialize(FSubsystemCollectionBase& Collecti
 		ActorSpawnerSubsystem = World->GetSubsystem<UMassActorSpawnerSubsystem>();
 		WorldPartitionSubsystem = World->GetSubsystem<UWorldPartitionSubsystem>();
 
-		if (Visualizer == nullptr)
+		if (SpawnVisualizer(World))
 		{
-			FActorSpawnParameters SpawnInfo;
-			SpawnInfo.SpawnCollisionHandlingOverride = ESpawnActorCollisionHandlingMethod::AlwaysSpawn;
-			// The helper actor is only once per world so we can allow it to spawn during construction script.
-			SpawnInfo.bAllowDuringConstructionScript = true;
-			Visualizer = World->SpawnActor<AMassVisualizer>(SpawnInfo);
-			check(Visualizer);
-			VisualizationComponent = &Visualizer->GetVisualizationComponent();
-
-#if WITH_EDITOR
-			Visualizer->SetActorLabel(FString::Printf(TEXT("%sVisualizer"), *GetClass()->GetName()), /*bMarkDirty*/false);
-#endif
-
-			// bind to processing phase events only if VisualizationComponent gets set - the phase event
-			// handlers required that to be set, so no point in binding if it's not true
 			UMassSimulationSubsystem* SimSystem = World->GetSubsystem<UMassSimulationSubsystem>();
 			check(SimSystem);
 			SimSystem->GetOnProcessingPhaseStarted(EMassProcessingPhase::PrePhysics).AddUObject(this, &UMassRepresentationSubsystem::OnProcessingPhaseStarted, EMassProcessingPhase::PrePhysics);
@@ -432,6 +418,51 @@ void UMassRepresentationSubsystem::Deinitialize()
 	Super::Deinitialize();
 }
 
+void UMassRepresentationSubsystem::HandleVisualizerEndPlay(AActor* Actor, const EEndPlayReason::Type EndPlayReason)
+{
+	Visualizer = nullptr;
+	VisualizationComponent = nullptr;
+
+	if (UWorld* World = GetWorld())
+	{
+		// attempt to create a new visualizer in the case that the old one was destroyed prematurely
+		if (!SpawnVisualizer(World))
+		{
+			// visualizer could not be created, so try to unregister phase delegates
+			if (UMassSimulationSubsystem* SimSystem = World->GetSubsystem<UMassSimulationSubsystem>())
+			{
+				SimSystem->GetOnProcessingPhaseStarted(EMassProcessingPhase::PrePhysics).RemoveAll(this);
+				SimSystem->GetOnProcessingPhaseFinished(EMassProcessingPhase::PostPhysics).RemoveAll(this);
+			}
+		}
+	}
+}
+
+bool UMassRepresentationSubsystem::SpawnVisualizer(TNotNull<UWorld*> World)
+{
+	if (World->bIsTearingDown)
+	{
+		return false;
+	}
+
+	if (Visualizer == nullptr)
+	{
+		FActorSpawnParameters SpawnInfo;
+		SpawnInfo.SpawnCollisionHandlingOverride = ESpawnActorCollisionHandlingMethod::AlwaysSpawn;
+		// The helper actor is only once per world so we can allow it to spawn during construction script.
+		SpawnInfo.bAllowDuringConstructionScript = true;
+		Visualizer = World->SpawnActor<AMassVisualizer>(SpawnInfo);
+		check(Visualizer);
+		Visualizer->OnEndPlay.AddUniqueDynamic(this, &UMassRepresentationSubsystem::HandleVisualizerEndPlay);
+		VisualizationComponent = &Visualizer->GetVisualizationComponent();
+
+#if WITH_EDITOR
+		Visualizer->SetActorLabel(FString::Printf(TEXT("%sVisualizer"), *GetClass()->GetName()), /*bMarkDirty*/false);
+#endif
+	}
+	return true;
+}
+
 void UMassRepresentationSubsystem::OnProcessingPhaseStarted(const float /*DeltaSeconds*/, const EMassProcessingPhase Phase) const
 {
 	check(VisualizationComponent);
diff --git a/Engine/Plugins/Runtime/MassGameplay/Source/MassRepresentation/Public/MassRepresentationSubsystem.h b/Engine/Plugins/Runtime/MassGameplay/Source/MassRepresentation/Public/MassRepresentationSubsystem.h
--- a/Engine/Plugins/Runtime/MassGameplay/Source/MassRepresentation/Public/MassRepresentationSubsystem.h
+++ b/Engine/Plugins/Runtime/MassGameplay/Source/MassRepresentation/Public/MassRepresentationSubsystem.h
@@ -159,6 +159,8 @@ protected:
 	MASSREPRESENTATION_API virtual void Deinitialize() override;
 	// USubsystem END
 
+	MASSREPRESENTATION_API bool SpawnVisualizer(TNotNull<UWorld*> World);
+
 	/** Needed for batching the update of static mesh transform */
 	MASSREPRESENTATION_API void OnProcessingPhaseStarted(const float DeltaSeconds, const EMassProcessingPhase Phase) const;
 
@@ -218,4 +220,8 @@ protected:
 
 	/** Keeping track of all the mass agent this subsystem is responsible for spawning actors */
 	TMap<FMassEntityHandle, int32> HandledMassAgents;
+
+private:
+	UFUNCTION()
+	MASSREPRESENTATION_API void HandleVisualizerEndPlay(AActor* Actor, const EEndPlayReason::Type EndPlayReason);
 };
