[Mass] Made fragment value setting directly via archetypes properly handle the given fragment type not being a part of a given archetype's composition

diff --git a/Engine/Plugins/Runtime/MassGameplay/Source/MassGameplayTestSuite/Private/MassEntityTemplateTest.cpp b/Engine/Plugins/Runtime/MassGameplay/Source/MassGameplayTestSuite/Private/MassEntityTemplateTest.cpp
--- a/Engine/Plugins/Runtime/MassGameplay/Source/MassGameplayTestSuite/Private/MassEntityTemplateTest.cpp
+++ b/Engine/Plugins/Runtime/MassGameplay/Source/MassGameplayTestSuite/Private/MassEntityTemplateTest.cpp
@@ -5,6 +5,7 @@
 #include "MassEntityTestTypes.h"
 #include "MassEntityTemplate.h"
 #include "MassEntityTemplateRegistry.h"
+#include "MassObserverNotificationTypes.h"
 
 #define LOCTEXT_NAMESPACE "MassTest"
 
@@ -190,6 +191,42 @@ IMPLEMENT_AI_INSTANT_TEST(FEntityTemplate_Modified, "System.Mass.EntityTemplate.
 //};
 //IMPLEMENT_AI_INSTANT_TEST(FEntityTemplate_Modified, "System.Mass.EntityTemplate.Modification");
 
+
+struct FEntityTest_BatchSetMissingFragment : FEntityTemplateBase
+{
+	virtual bool InstantTest() override
+	{
+		CA_ASSUME(EntityManager);
+
+		constexpr int32 NumToSpawn = 3;
+
+		FMassEntityTemplateData TemplateData;
+		TemplateData.AddFragment<FTestFragment_Int>();
+		TemplateData.AddFragment_GetRef<FTestFragment_Float>().Value = 3.14f;
+		TemplateData.Sort();
+
+		const FMassEntityTemplateID TemplateID = FMassEntityTemplateIDFactory::Make(FGuid::NewGuid());
+		const TSharedRef<FMassEntityTemplate>& EntityTemplate = TemplateRegistry.FindOrAddTemplate(TemplateID, MoveTemp(TemplateData));
+
+		TArray<FMassEntityHandle> CreatedEntities;
+		TSharedRef<FMassEntityManager::FEntityCreationContext> CreationContext
+			= EntityManager->BatchCreateEntities(EntityTemplate->GetArchetype(), EntityTemplate->GetSharedFragmentValues(), NumToSpawn, CreatedEntities);
+
+		// forcing archetype change for the created entities. 
+		EntityManager->BatchChangeFragmentCompositionForEntities(CreationContext->GetEntityCollections(*EntityManager.Get())
+			, {}, FMassFragmentBitSet(*FTestFragment_Float::StaticStruct()));
+
+		{
+			AITEST_SCOPED_CHECK("Attempting to set value of fragment of type", 1);
+			TConstArrayView<FInstancedStruct> FragmentInstances = EntityTemplate->GetInitialFragmentValues();
+			EntityManager->BatchSetEntityFragmentValues(CreationContext->GetEntityCollections(*EntityManager.Get()), FragmentInstances);
+		}
+
+		return true;
+	}
+};
+IMPLEMENT_AI_INSTANT_TEST(FEntityTest_BatchSetMissingFragment, "System.Mass.EntityTemplate.BatchSetMissingFragment");
+
 } // FMassEntityTemplateTest
 
 UE_ENABLE_OPTIMIZATION_SHIP
diff --git a/Engine/Source/Runtime/MassEntity/Private/MassArchetypeData.cpp b/Engine/Source/Runtime/MassEntity/Private/MassArchetypeData.cpp
--- a/Engine/Source/Runtime/MassEntity/Private/MassArchetypeData.cpp
+++ b/Engine/Source/Runtime/MassEntity/Private/MassArchetypeData.cpp
@@ -446,19 +446,28 @@ void FMassArchetypeData::SetFragmentData(FMassArchetypeEntityCollection::FConstE
 	check(FragmentSource.IsValid());
 	const UScriptStruct* FragmentType = FragmentSource.GetScriptStruct();
 	check(FragmentType);
-	const int32 FragmentIndex = FragmentIndexMap.FindChecked(FragmentType);
-	const int32 FragmentTypeSize = FragmentType->GetStructureSize();
-	const uint8* FragmentSourceMemory = FragmentSource.GetMemory();
-	check(FragmentSourceMemory);
-	
-	for (FMassArchetypeChunkIterator ChunkIterator(EntityRangeContainer); ChunkIterator; ++ChunkIterator)
+	const int32* FragmentIndex = FragmentIndexMap.Find(FragmentType);
+	if (LIKELY(FragmentIndex))
 	{
-		uint8* FragmentMemory = (uint8*)FragmentConfigs[FragmentIndex].GetFragmentData(Chunks[ChunkIterator->ChunkIndex].GetRawMemory(), ChunkIterator->SubchunkStart);
-		for (int i = ChunkIterator->Length; i; --i, FragmentMemory += FragmentTypeSize)
+		const int32 FragmentTypeSize = FragmentType->GetStructureSize();
+		const uint8* FragmentSourceMemory = FragmentSource.GetMemory();
+		check(FragmentSourceMemory);
+	
+		for (FMassArchetypeChunkIterator ChunkIterator(EntityRangeContainer); ChunkIterator; ++ChunkIterator)
 		{
-			FragmentType->CopyScriptStruct(FragmentMemory, FragmentSourceMemory);
+			uint8* FragmentMemory = (uint8*)FragmentConfigs[*FragmentIndex].GetFragmentData(Chunks[ChunkIterator->ChunkIndex].GetRawMemory(), ChunkIterator->SubchunkStart);
+			for (int i = ChunkIterator->Length; i; --i, FragmentMemory += FragmentTypeSize)
+			{
+				FragmentType->CopyScriptStruct(FragmentMemory, FragmentSourceMemory);
+			}
 		}
 	}
+	else
+	{
+		UE_LOG(LogMass, Warning
+			, TEXT("Attempting to set value of fragment of type %s, while it's not part of the archetype's composition")
+			, *FragmentType->GetName());
+	}
 }
 
 void FMassArchetypeData::MoveEntityToAnotherArchetype(const FMassEntityHandle Entity, FMassArchetypeData& NewArchetype, const FMassArchetypeSharedFragmentValues* SharedFragmentValuesOverride)
