PR #13735: Export types from MassEQS

diff --git a/Engine/Plugins/Runtime/MassGameplay/Source/MassEQS/Public/MassEQSSubsystem.h b/Engine/Plugins/Runtime/MassGameplay/Source/MassEQS/Public/MassEQSSubsystem.h
--- a/Engine/Plugins/Runtime/MassGameplay/Source/MassEQS/Public/MassEQSSubsystem.h
+++ b/Engine/Plugins/Runtime/MassGameplay/Source/MassEQS/Public/MassEQSSubsystem.h
@@ -25,12 +25,12 @@ public:
 	FMassEQSRequestHandle PushRequest(const FEnvQueryInstance& QueryInstance, const int32 RequestQueueIndex, TUniquePtr<FMassEQSRequestData>&& Request);
 
 	/** Dequeue next Request off the RequestQueue, removing it from the queue, and transferring ownership of UniquePtr */
-	TUniquePtr<FMassEQSRequestData> PopRequest(const int32 RequestQueueIndex);
+	MASSEQS_API TUniquePtr<FMassEQSRequestData> PopRequest(const int32 RequestQueueIndex);
 
 	/** Send results from a finished/processed request to be stored in this subsystem */
-	void SubmitResults(FMassEQSRequestHandle RequestHandle, TUniquePtr<FMassEQSRequestData>&& Result);
+	MASSEQS_API void SubmitResults(FMassEQSRequestHandle RequestHandle, TUniquePtr<FMassEQSRequestData>&& Result);
 	/** Try to Acquire the results from a request from the QueryInstance, if it is available. */
-	TUniquePtr<FMassEQSRequestData> TryAcquireResults(FMassEQSRequestHandle RequestHandle);
+	MASSEQS_API TUniquePtr<FMassEQSRequestData> TryAcquireResults(FMassEQSRequestHandle RequestHandle);
 
 	/** Returns the index into RequestQueues for the input class */
 	int32 GetRequestQueueIndex(TSubclassOf<UEnvQueryNode> RequestClass);
@@ -92,4 +92,4 @@ struct TMassExternalSubsystemTraits<UMassEQSSubsystem> final
 		GameThreadOnly = false,
 		ThreadSafeWrite = true,
 	};
-};
\ No newline at end of file
+};
diff --git a/Engine/Plugins/Runtime/MassGameplay/Source/MassEQS/Public/MassEQSUtils.h b/Engine/Plugins/Runtime/MassGameplay/Source/MassEQS/Public/MassEQSUtils.h
--- a/Engine/Plugins/Runtime/MassGameplay/Source/MassEQS/Public/MassEQSUtils.h
+++ b/Engine/Plugins/Runtime/MassGameplay/Source/MassEQS/Public/MassEQSUtils.h
@@ -13,7 +13,7 @@ struct FMassEntityHandle;
 struct FMassEQSRequestData;
 
 /** Holds Utility functions for Mass EQS Needs */
-struct FMassEQSUtils
+struct MASSEQS_API FMassEQSUtils
 {
 
 	/** Returns the Item stored in QueryInstance/QueryResult Items[Index] as EntityInfo */
diff --git a/Engine/Plugins/Runtime/MassGameplay/Source/MassEQS/Public/Processors/MassEnvQueryProcessorBase.h b/Engine/Plugins/Runtime/MassGameplay/Source/MassEQS/Public/Processors/MassEnvQueryProcessorBase.h
--- a/Engine/Plugins/Runtime/MassGameplay/Source/MassEQS/Public/Processors/MassEnvQueryProcessorBase.h
+++ b/Engine/Plugins/Runtime/MassGameplay/Source/MassEQS/Public/Processors/MassEnvQueryProcessorBase.h
@@ -9,7 +9,7 @@ class UEnvQueryNode;
 
 /** Processor for completing MassEQSSubsystem Requests sent from UMassEnvQueryTest_MassEntityTags */
 UCLASS(Abstract, meta = (DisplayName = "Mass EQS Processor Base"))
-class UMassEnvQueryProcessorBase : public UMassProcessor
+class MASSEQS_API UMassEnvQueryProcessorBase : public UMassProcessor
 {
 	GENERATED_BODY()
 
diff --git a/Engine/Plugins/Runtime/MassGameplay/Source/MassEQS/Public/Tests/MassEnvQueryTest.h b/Engine/Plugins/Runtime/MassGameplay/Source/MassEQS/Public/Tests/MassEnvQueryTest.h
--- a/Engine/Plugins/Runtime/MassGameplay/Source/MassEQS/Public/Tests/MassEnvQueryTest.h
+++ b/Engine/Plugins/Runtime/MassGameplay/Source/MassEQS/Public/Tests/MassEnvQueryTest.h
@@ -16,16 +16,16 @@ public:
 	 * This will send this Test Request to MassEQSSubsystem the first time it is called,
 	 * and Try to complete Testing with the Result from MassEQSSubsystem on subsequent calls.
 	 */
-	virtual void RunTest(FEnvQueryInstance& QueryInstance) const final;
+	MASSEQS_API virtual void RunTest(FEnvQueryInstance& QueryInstance) const final;
 	virtual FORCEINLINE bool IsCurrentlyRunningAsync() const final { return MassEQSRequestHandler.IsPendingResults(); }
 
 	// Begin IMassEQSRequestInterface
-	virtual TUniquePtr<FMassEQSRequestData> GetRequestData(FEnvQueryInstance& QueryInstance) const override PURE_VIRTUAL(UMassEnvQueryTest::GetRequestData, return nullptr;)
-	virtual UClass* GetRequestClass() const override PURE_VIRTUAL(UMassEnvQueryTest::GetRequestClass, return nullptr;);
+	MASSEQS_API virtual TUniquePtr<FMassEQSRequestData> GetRequestData(FEnvQueryInstance& QueryInstance) const override PURE_VIRTUAL(UMassEnvQueryTest::GetRequestData, return nullptr;)
+	MASSEQS_API virtual UClass* GetRequestClass() const override PURE_VIRTUAL(UMassEnvQueryTest::GetRequestClass, return nullptr;);
 
-	virtual bool TryAcquireResults(FEnvQueryInstance& QueryInstance) const override PURE_VIRTUAL(UMassEnvQueryTest::GetRequestData, return false;)
+	MASSEQS_API virtual bool TryAcquireResults(FEnvQueryInstance& QueryInstance) const override PURE_VIRTUAL(UMassEnvQueryTest::GetRequestData, return false;)
 	// ~IMassEQSRequestInterface
 
 protected:
 	mutable FMassEQSRequestHandler MassEQSRequestHandler;
-};
\ No newline at end of file
+};
