[MassGameplay]
Fix crash when mass representation actor is destroyed while mass simulation is active. Remove checks in the representation API that can happen in this case and added reasonable default return values.
Fixed teardown crash in MassAgentComponent.

diff --git a/Engine/Plugins/Runtime/InstancedActors/Source/InstancedActors/Private/InstancedActorsRepresentationActorManagement.cpp b/Engine/Plugins/Runtime/InstancedActors/Source/InstancedActors/Private/InstancedActorsRepresentationActorManagement.cpp
--- a/Engine/Plugins/Runtime/InstancedActors/Source/InstancedActors/Private/InstancedActorsRepresentationActorManagement.cpp
+++ b/Engine/Plugins/Runtime/InstancedActors/Source/InstancedActors/Private/InstancedActorsRepresentationActorManagement.cpp
@@ -191,10 +191,15 @@ AActor* UInstancedActorsRepresentationActorManagement::FindOrInstantlySpawnActor
 		if (Representation.CurrentRepresentation == EMassRepresentationType::StaticMeshInstance)
 		{
 			FMassInstancedStaticMeshInfoArrayView ISMInfosView = RepresentationSubsystem.GetMutableInstancedStaticMeshInfos();
-			FMassInstancedStaticMeshInfo& ISMInfo = ISMInfosView[Representation.StaticMeshDescHandle.ToIndex()];
 
-			ISMInfo.RemoveInstance(EntityView.GetEntity(), Representation.PrevLODSignificance);
-			ISMInfo.RemoveInstance(EntityView.GetEntity(), RepresentationLOD.LODSignificance);
+			const int32 ISMInfoIndex = Representation.StaticMeshDescHandle.ToIndex();
+			if (ensureMsgf(ISMInfosView.IsValidIndex(ISMInfoIndex), TEXT("Invalid handle index %u for ISMInfosView"), ISMInfoIndex))
+			{
+				FMassInstancedStaticMeshInfo& ISMInfo = ISMInfosView[ISMInfoIndex];
+
+				ISMInfo.RemoveInstance(EntityView.GetEntity(), Representation.PrevLODSignificance);
+				ISMInfo.RemoveInstance(EntityView.GetEntity(), RepresentationLOD.LODSignificance);
+			}
 		}
 
 		Representation.PrevRepresentation = Representation.CurrentRepresentation;
diff --git a/Engine/Plugins/Runtime/InstancedActors/Source/InstancedActors/Private/InstancedActorsVisualizationSwitcherProcessor.cpp b/Engine/Plugins/Runtime/InstancedActors/Source/InstancedActors/Private/InstancedActorsVisualizationSwitcherProcessor.cpp
--- a/Engine/Plugins/Runtime/InstancedActors/Source/InstancedActors/Private/InstancedActorsVisualizationSwitcherProcessor.cpp
+++ b/Engine/Plugins/Runtime/InstancedActors/Source/InstancedActors/Private/InstancedActorsVisualizationSwitcherProcessor.cpp
@@ -69,11 +69,15 @@ void UInstancedActorsVisualizationSwitcherProcessor::SwitchEntityMeshDesc(FMassI
 				&& ensureMsgf(ISMInfosView.IsValidIndex(RepresentationFragment.StaticMeshDescHandle.ToIndex()), TEXT("Switching visualization while StaticMeshDescHandle is not a valid ISMInfo index, %d out of %d")
 					, RepresentationFragment.StaticMeshDescHandle.ToIndex(), ISMInfosView.Num()))
 			{
-				FMassInstancedStaticMeshInfo& ISMInfo = ISMInfosView[RepresentationFragment.StaticMeshDescHandle.ToIndex()];
-
-				// Note that we're using the PrevLODSignificance here, and the reason for it is that the Prev value matches the 
-				// PrevRepresentation - thus we need to remove from the "previously" used LODSignificance range.
-				ISMInfo.RemoveInstance(EntityHandle, RepresentationFragment.PrevLODSignificance);
+				const int32 ISMInfoIndex = RepresentationFragment.StaticMeshDescHandle.ToIndex();
+				if (ensureMsgf(ISMInfosView.IsValidIndex(ISMInfoIndex), TEXT("Invalid handle index %u for ISMInfosView"), ISMInfoIndex))
+				{
+					FMassInstancedStaticMeshInfo& ISMInfo = ISMInfosView[ISMInfoIndex];
+
+					// Note that we're using the PrevLODSignificance here, and the reason for it is that the Prev value matches the 
+					// PrevRepresentation - thus we need to remove from the "previously" used LODSignificance range.
+					ISMInfo.RemoveInstance(EntityHandle, RepresentationFragment.PrevLODSignificance);
+				}
 			}
 
 			// Set PrevRepresentation to None to match the new removed instance state and let 
diff --git a/Engine/Plugins/Runtime/MassGameplay/Source/MassActors/Private/MassAgentComponent.cpp b/Engine/Plugins/Runtime/MassGameplay/Source/MassActors/Private/MassAgentComponent.cpp
--- a/Engine/Plugins/Runtime/MassGameplay/Source/MassActors/Private/MassAgentComponent.cpp
+++ b/Engine/Plugins/Runtime/MassGameplay/Source/MassActors/Private/MassAgentComponent.cpp
@@ -139,9 +139,15 @@ void UMassAgentComponent::RegisterWithAgentSubsystem()
 
 void UMassAgentComponent::UnregisterWithAgentSubsystem()
 {
+	UWorld* World = GetWorld();
+	if (World && World->bIsTearingDown)
+	{
+		// no need to Unregister during world teardown
+		return;
+	}
 	if (State != EAgentComponentState::None)
 	{
-		if (UMassAgentSubsystem* AgentSubsystem = UWorld::GetSubsystem<UMassAgentSubsystem>(GetWorld()))
+		if (UMassAgentSubsystem* AgentSubsystem = UWorld::GetSubsystem<UMassAgentSubsystem>(World))
 		{
 			UE_VLOG(GetOwner(), LogMass, Verbose, TEXT("%s"), ANSI_TO_TCHAR(__FUNCTION__));
 			AgentSubsystem->ShutdownAgentComponent(*this);
@@ -364,6 +370,11 @@ void UMassAgentComponent::SwitchToState(EAgentComponentState NewState)
 void UMassAgentComponent::DebugCheckStateConsistency()
 {
 #if DO_CHECK
+	if (GetWorld() && GetWorld()->bIsTearingDown)
+	{
+		// Inconsistency is expected during teardown
+		return;
+	}
 	switch (State)
 	{
 		case EAgentComponentState::None:
diff --git a/Engine/Plugins/Runtime/MassGameplay/Source/MassRepresentation/Private/MassRepresentationSubsystem.cpp b/Engine/Plugins/Runtime/MassGameplay/Source/MassRepresentation/Private/MassRepresentationSubsystem.cpp
--- a/Engine/Plugins/Runtime/MassGameplay/Source/MassRepresentation/Private/MassRepresentationSubsystem.cpp
+++ b/Engine/Plugins/Runtime/MassGameplay/Source/MassRepresentation/Private/MassRepresentationSubsystem.cpp
@@ -24,51 +24,82 @@
 
 FStaticMeshInstanceVisualizationDescHandle UMassRepresentationSubsystem::FindOrAddStaticMeshDesc(const FStaticMeshInstanceVisualizationDesc& Desc)
 {
-	check(VisualizationComponent);
-	return VisualizationComponent->FindOrAddVisualDesc(Desc);
+	if (LIKELY(VisualizationComponent))
+	{
+		return VisualizationComponent->FindOrAddVisualDesc(Desc);
+	}
+
+	return FStaticMeshInstanceVisualizationDescHandle();
 }
 
 FStaticMeshInstanceVisualizationDescHandle UMassRepresentationSubsystem::AddVisualDescWithISMComponent(const FStaticMeshInstanceVisualizationDesc& Desc, UInstancedStaticMeshComponent& ISMComponent)
 {
-	check(VisualizationComponent);
-	return VisualizationComponent->AddVisualDescWithISMComponent(Desc, ISMComponent);
+	if (LIKELY(VisualizationComponent))
+	{
+		return VisualizationComponent->AddVisualDescWithISMComponent(Desc, ISMComponent);
+	}
+
+	return FStaticMeshInstanceVisualizationDescHandle();
 }
 
 FStaticMeshInstanceVisualizationDescHandle UMassRepresentationSubsystem::AddVisualDescWithISMComponents(const FStaticMeshInstanceVisualizationDesc& Desc, TArrayView<TObjectPtr<UInstancedStaticMeshComponent>> ISMComponents)
 {
-	check(VisualizationComponent);
-	return VisualizationComponent->AddVisualDescWithISMComponents(Desc, ISMComponents);
+	if (LIKELY(VisualizationComponent))
+	{
+		return VisualizationComponent->AddVisualDescWithISMComponents(Desc, ISMComponents);
+	}
+
+	return FStaticMeshInstanceVisualizationDescHandle();
 }
 
 const FMassISMCSharedData* UMassRepresentationSubsystem::GetISMCSharedDataForDescriptionIndex(const int32 DescriptionIndex) const
 {
-	check(VisualizationComponent);
-	return VisualizationComponent->GetISMCSharedDataForDescriptionIndex(DescriptionIndex);
+	if (LIKELY(VisualizationComponent))
+	{
+		return VisualizationComponent->GetISMCSharedDataForDescriptionIndex(DescriptionIndex);
+	}
+
+	return nullptr;
 }
 
 const FMassISMCSharedData* UMassRepresentationSubsystem::GetISMCSharedDataForInstancedStaticMesh(const UInstancedStaticMeshComponent* ISMC) const
 {
-	check(VisualizationComponent);
-	return VisualizationComponent->GetISMCSharedDataForInstancedStaticMesh(ISMC);
+	if (LIKELY(VisualizationComponent))
+	{
+		return VisualizationComponent->GetISMCSharedDataForInstancedStaticMesh(ISMC);
+	}
+	return nullptr;
 }
 
 void UMassRepresentationSubsystem::RemoveVisualDesc(const FStaticMeshInstanceVisualizationDescHandle VisualizationHandle)
 {
-	check(VisualizationComponent);
-	return VisualizationComponent->RemoveVisualDesc(VisualizationHandle);
+	if (LIKELY(VisualizationComponent))
+	{
+		VisualizationComponent->RemoveVisualDesc(VisualizationHandle);
+	}
 }
 
 FMassInstancedStaticMeshInfoArrayView UMassRepresentationSubsystem::GetMutableInstancedStaticMeshInfos()
 {
-	check(VisualizationComponent);
-	FMassInstancedStaticMeshInfoArrayView View = VisualizationComponent->GetMutableVisualInfos();
-	return MoveTemp(View);
+	if (LIKELY(VisualizationComponent))
+	{
+		FMassInstancedStaticMeshInfoArrayView View = VisualizationComponent->GetMutableVisualInfos();
+		return MoveTemp(View);
+	}
+
+	// This can happen if the visualizer actor is destroyed prematurely.
+	// It can also happen during world teardown if this is called after actors are destroyed.
+	// This is not hit during normal execution, but we still need a default return value.
+	FMassInstancedStaticMeshInfoArrayView EmptyView = MAKE_MASS_INSTANCED_STATIC_MESH_INFO_ARRAY_VIEW(MakeArrayView<FMassInstancedStaticMeshInfo>(nullptr, 0), InstancedStaticMeshInfosDetector);
+	return MoveTemp(EmptyView);
 }
 
 void UMassRepresentationSubsystem::DirtyStaticMeshInstances()
 {
-	check(VisualizationComponent);
-	return VisualizationComponent->DirtyVisuals();
+	if (LIKELY(VisualizationComponent))
+	{
+		return VisualizationComponent->DirtyVisuals();
+	}
 }
 
 int16 UMassRepresentationSubsystem::FindOrAddTemplateActor(const TSubclassOf<AActor>& ActorClass)
@@ -379,22 +410,8 @@ void UMassRepresentationSubsystem::Initialize(FSubsystemCollectionBase& Collecti
 		ActorSpawnerSubsystem = World->GetSubsystem<UMassActorSpawnerSubsystem>();
 		WorldPartitionSubsystem = World->GetSubsystem<UWorldPartitionSubsystem>();
 
-		if (Visualizer == nullptr)
+		if (SpawnVisualizer(World))
 		{
-			FActorSpawnParameters SpawnInfo;
-			SpawnInfo.SpawnCollisionHandlingOverride = ESpawnActorCollisionHandlingMethod::AlwaysSpawn;
-			// The helper actor is only once per world so we can allow it to spawn during construction script.
-			SpawnInfo.bAllowDuringConstructionScript = true;
-			Visualizer = World->SpawnActor<AMassVisualizer>(SpawnInfo);
-			check(Visualizer);
-			VisualizationComponent = &Visualizer->GetVisualizationComponent();
-
-#if WITH_EDITOR
-			Visualizer->SetActorLabel(FString::Printf(TEXT("%sVisualizer"), *GetClass()->GetName()), /*bMarkDirty*/false);
-#endif
-
-			// bind to processing phase events only if VisualizationComponent gets set - the phase event
-			// handlers required that to be set, so no point in binding if it's not true
 			UMassSimulationSubsystem* SimSystem = World->GetSubsystem<UMassSimulationSubsystem>();
 			check(SimSystem);
 			SimSystem->GetOnProcessingPhaseStarted(EMassProcessingPhase::PrePhysics).AddUObject(this, &UMassRepresentationSubsystem::OnProcessingPhaseStarted, EMassProcessingPhase::PrePhysics);
@@ -432,6 +449,51 @@ void UMassRepresentationSubsystem::Deinitialize()
 	Super::Deinitialize();
 }
 
+void UMassRepresentationSubsystem::HandleVisualizerEndPlay(AActor* Actor, const EEndPlayReason::Type EndPlayReason)
+{
+	Visualizer = nullptr;
+	VisualizationComponent = nullptr;
+
+	if (UWorld* World = GetWorld())
+	{
+		// attempt to create a new visualizer in the case that the old one was destroyed prematurely
+		if (!SpawnVisualizer(World))
+		{
+			// visualizer could not be created, so try to unregister phase delegates
+			if (UMassSimulationSubsystem* SimSystem = World->GetSubsystem<UMassSimulationSubsystem>())
+			{
+				SimSystem->GetOnProcessingPhaseStarted(EMassProcessingPhase::PrePhysics).RemoveAll(this);
+				SimSystem->GetOnProcessingPhaseFinished(EMassProcessingPhase::PostPhysics).RemoveAll(this);
+			}
+		}
+	}
+}
+
+bool UMassRepresentationSubsystem::SpawnVisualizer(TNotNull<UWorld*> World)
+{
+	if (World->bIsTearingDown)
+	{
+		return false;
+	}
+
+	if (Visualizer == nullptr)
+	{
+		FActorSpawnParameters SpawnInfo;
+		SpawnInfo.SpawnCollisionHandlingOverride = ESpawnActorCollisionHandlingMethod::AlwaysSpawn;
+		// The helper actor is only once per world so we can allow it to spawn during construction script.
+		SpawnInfo.bAllowDuringConstructionScript = true;
+		Visualizer = World->SpawnActor<AMassVisualizer>(SpawnInfo);
+		check(Visualizer);
+		Visualizer->OnEndPlay.AddUniqueDynamic(this, &UMassRepresentationSubsystem::HandleVisualizerEndPlay);
+		VisualizationComponent = &Visualizer->GetVisualizationComponent();
+
+#if WITH_EDITOR
+		Visualizer->SetActorLabel(FString::Printf(TEXT("%sVisualizer"), *GetClass()->GetName()), /*bMarkDirty*/false);
+#endif
+	}
+	return true;
+}
+
 void UMassRepresentationSubsystem::OnProcessingPhaseStarted(const float /*DeltaSeconds*/, const EMassProcessingPhase Phase) const
 {
 	check(VisualizationComponent);
diff --git a/Engine/Plugins/Runtime/MassGameplay/Source/MassRepresentation/Private/MassStationaryISMRepresentationFragmentDestructor.cpp b/Engine/Plugins/Runtime/MassGameplay/Source/MassRepresentation/Private/MassStationaryISMRepresentationFragmentDestructor.cpp
--- a/Engine/Plugins/Runtime/MassGameplay/Source/MassRepresentation/Private/MassStationaryISMRepresentationFragmentDestructor.cpp
+++ b/Engine/Plugins/Runtime/MassGameplay/Source/MassRepresentation/Private/MassStationaryISMRepresentationFragmentDestructor.cpp
@@ -42,13 +42,17 @@ void UMassStationaryISMRepresentationFragmentDestructor::Execute(FMassEntityMana
 			FMassRepresentationFragment& Representation = RepresentationList[EntityIt];
 			if (Representation.CurrentRepresentation == EMassRepresentationType::StaticMeshInstance)
 			{
-				FMassInstancedStaticMeshInfo& ISMInfo = ISMInfosView[Representation.StaticMeshDescHandle.ToIndex()];
-				if (FMassLODSignificanceRange* OldRange = ISMInfo.GetLODSignificanceRange(Representation.PrevLODSignificance))
+				const int32 ISMInfoIndex = Representation.StaticMeshDescHandle.ToIndex();
+				if (ensureMsgf(ISMInfosView.IsValidIndex(ISMInfoIndex), TEXT("Invalid handle index %u for ISMInfosView"), ISMInfoIndex))
 				{
-					const FMassEntityHandle EntityHandle = Context.GetEntity(EntityIt);
-					if (OldRange)
+					FMassInstancedStaticMeshInfo& ISMInfo = ISMInfosView[ISMInfoIndex];
+					if (FMassLODSignificanceRange* OldRange = ISMInfo.GetLODSignificanceRange(Representation.PrevLODSignificance))
 					{
-						OldRange->RemoveInstance(EntityHandle);
+						const FMassEntityHandle EntityHandle = Context.GetEntity(EntityIt);
+						if (OldRange)
+						{
+							OldRange->RemoveInstance(EntityHandle);
+						}
 					}
 				}
 				Representation.CurrentRepresentation = EMassRepresentationType::None;
diff --git a/Engine/Plugins/Runtime/MassGameplay/Source/MassRepresentation/Private/MassUpdateISMProcessor.cpp b/Engine/Plugins/Runtime/MassGameplay/Source/MassRepresentation/Private/MassUpdateISMProcessor.cpp
--- a/Engine/Plugins/Runtime/MassGameplay/Source/MassRepresentation/Private/MassUpdateISMProcessor.cpp
+++ b/Engine/Plugins/Runtime/MassGameplay/Source/MassRepresentation/Private/MassUpdateISMProcessor.cpp
@@ -54,7 +54,11 @@ void UMassUpdateISMProcessor::Execute(FMassEntityManager& EntityManager, FMassEx
 
 			if (Representation.CurrentRepresentation == EMassRepresentationType::StaticMeshInstance)
 			{
-				UpdateISMTransform(Context.GetEntity(EntityIt), ISMInfo[Representation.StaticMeshDescHandle.ToIndex()], TransformFragment.GetTransform(), Representation.PrevTransform, RepresentationLOD.LODSignificance, Representation.PrevLODSignificance);
+				const int32 ISMInfoIndex = Representation.StaticMeshDescHandle.ToIndex();
+				if (ensureMsgf(ISMInfo.IsValidIndex(ISMInfoIndex), TEXT("Invalid handle index %u for ISMInfosView"), ISMInfoIndex))
+				{
+					UpdateISMTransform(Context.GetEntity(EntityIt), ISMInfo[ISMInfoIndex], TransformFragment.GetTransform(), Representation.PrevTransform, RepresentationLOD.LODSignificance, Representation.PrevLODSignificance);
+				}
 			}
 			Representation.PrevTransform = TransformFragment.GetTransform();
 			Representation.PrevLODSignificance = RepresentationLOD.LODSignificance;
diff --git a/Engine/Plugins/Runtime/MassGameplay/Source/MassRepresentation/Public/MassRepresentationSubsystem.h b/Engine/Plugins/Runtime/MassGameplay/Source/MassRepresentation/Public/MassRepresentationSubsystem.h
--- a/Engine/Plugins/Runtime/MassGameplay/Source/MassRepresentation/Public/MassRepresentationSubsystem.h
+++ b/Engine/Plugins/Runtime/MassGameplay/Source/MassRepresentation/Public/MassRepresentationSubsystem.h
@@ -159,6 +159,8 @@ protected:
 	MASSREPRESENTATION_API virtual void Deinitialize() override;
 	// USubsystem END
 
+	MASSREPRESENTATION_API bool SpawnVisualizer(TNotNull<UWorld*> World);
+
 	/** Needed for batching the update of static mesh transform */
 	MASSREPRESENTATION_API void OnProcessingPhaseStarted(const float DeltaSeconds, const EMassProcessingPhase Phase) const;
 
@@ -218,4 +220,11 @@ protected:
 
 	/** Keeping track of all the mass agent this subsystem is responsible for spawning actors */
 	TMap<FMassEntityHandle, int32> HandledMassAgents;
+
+private:
+	UFUNCTION()
+	MASSREPRESENTATION_API void HandleVisualizerEndPlay(AActor* Actor, const EEndPlayReason::Type EndPlayReason);
+
+	/** Used for default return value from GetMutableInstancedStaticMeshInfos */
+	UE_MT_DECLARE_RW_ACCESS_DETECTOR(InstancedStaticMeshInfosDetector);
 };
