PR #13332: Added std::is_trivially_copyable to the CFragment concept.

diff --git a/Engine/Plugins/AI/MassAI/Source/MassNavMeshNavigation/Public/MassNavMeshNavigationFragments.h b/Engine/Plugins/AI/MassAI/Source/MassNavMeshNavigation/Public/MassNavMeshNavigationFragments.h
--- a/Engine/Plugins/AI/MassAI/Source/MassNavMeshNavigation/Public/MassNavMeshNavigationFragments.h
+++ b/Engine/Plugins/AI/MassAI/Source/MassNavMeshNavigation/Public/MassNavMeshNavigationFragments.h
@@ -123,6 +123,15 @@ struct FMassNavMeshCachedPathFragment : public FMassFragment
 	static constexpr uint8 NumLeadingPoints = 1;
 };
 
+template<>
+struct TMassFragmentTraits<FMassNavMeshCachedPathFragment> final
+{
+	enum
+	{
+		AuthorAcceptsItsNotTriviallyCopyable = true
+	};
+};
+
 USTRUCT()
 struct FMassNavMeshBoundaryFragment : public FMassFragment
 {
diff --git a/Engine/Plugins/AI/MassAI/Source/MassNavigation/Public/Avoidance/MassAvoidanceFragments.h b/Engine/Plugins/AI/MassAI/Source/MassNavigation/Public/Avoidance/MassAvoidanceFragments.h
--- a/Engine/Plugins/AI/MassAI/Source/MassNavigation/Public/Avoidance/MassAvoidanceFragments.h
+++ b/Engine/Plugins/AI/MassAI/Source/MassNavigation/Public/Avoidance/MassAvoidanceFragments.h
@@ -4,6 +4,7 @@
 
 #include "MassEntityTypes.h"
 #include "MassEntityHandle.h"
+#include "MassExternalSubsystemTraits.h"
 #if UE_ENABLE_INCLUDE_ORDER_DEPRECATED_IN_5_6
 #include "MassNavigationSubsystem.h"
 #endif // UE_ENABLE_INCLUDE_ORDER_DEPRECATED_IN_5_6
@@ -205,6 +206,15 @@ struct FMassNavigationEdgesFragment : public FMassFragment
 	bool bExtrudedEdges = false;
 };
 
+template<>
+struct TMassFragmentTraits<FMassNavigationEdgesFragment> final
+{
+	enum
+	{
+		AuthorAcceptsItsNotTriviallyCopyable = true
+	};
+};
+
 USTRUCT()
 struct FMassAvoidanceEntitiesToIgnoreFragment : public FMassFragment
 {
@@ -213,3 +223,12 @@ struct FMassAvoidanceEntitiesToIgnoreFragment : public FMassFragment
 	static constexpr int EntityCount = 4;
 	TArray<FMassEntityHandle, TFixedAllocator<EntityCount>> EntitiesToIgnore;
 };
+
+template<>
+struct TMassFragmentTraits<FMassAvoidanceEntitiesToIgnoreFragment> final
+{
+	enum
+	{
+		AuthorAcceptsItsNotTriviallyCopyable = true
+	};
+};
\ No newline at end of file
diff --git a/Engine/Plugins/Runtime/MassGameplay/Source/MassActors/Public/MassActorSubsystem.h b/Engine/Plugins/Runtime/MassGameplay/Source/MassActors/Public/MassActorSubsystem.h
--- a/Engine/Plugins/Runtime/MassGameplay/Source/MassActors/Public/MassActorSubsystem.h
+++ b/Engine/Plugins/Runtime/MassGameplay/Source/MassActors/Public/MassActorSubsystem.h
@@ -45,6 +45,15 @@ struct FMassActorInstanceFragment : public FMassFragment
 	FActorInstanceHandle Handle;
 };
 
+template<>
+struct TMassFragmentTraits<FMassActorInstanceFragment> final
+{
+	enum
+	{
+		AuthorAcceptsItsNotTriviallyCopyable = true
+	};
+};
+
 namespace UE::Mass::Signals
 {
 	/** Signal use when the actor instance handle is set or cleared in the associated fragment. */
diff --git a/Engine/Plugins/Runtime/MassGameplay/Source/MassReplication/Public/MassReplicationFragments.h b/Engine/Plugins/Runtime/MassGameplay/Source/MassReplication/Public/MassReplicationFragments.h
--- a/Engine/Plugins/Runtime/MassGameplay/Source/MassReplication/Public/MassReplicationFragments.h
+++ b/Engine/Plugins/Runtime/MassGameplay/Source/MassReplication/Public/MassReplicationFragments.h
@@ -51,6 +51,15 @@ struct FMassReplicationViewerInfoFragment : public FMassFragment
 	TArray<float> DistanceToViewerSq;
 };
 
+template<>
+struct TMassFragmentTraits<FMassReplicationViewerInfoFragment> final
+{
+	enum
+	{
+		AuthorAcceptsItsNotTriviallyCopyable = true
+	};
+};
+
 USTRUCT()
 struct FMassReplicationLODFragment : public FMassFragment
 {
diff --git a/Engine/Plugins/Runtime/MassGameplay/Source/MassSmartObjects/Public/MassSmartObjectFragments.h b/Engine/Plugins/Runtime/MassGameplay/Source/MassSmartObjects/Public/MassSmartObjectFragments.h
--- a/Engine/Plugins/Runtime/MassGameplay/Source/MassSmartObjects/Public/MassSmartObjectFragments.h
+++ b/Engine/Plugins/Runtime/MassGameplay/Source/MassSmartObjects/Public/MassSmartObjectFragments.h
@@ -36,6 +36,15 @@ struct FMassSmartObjectUserFragment : public FMassFragment
 	double InteractionCooldownEndTime = 0.;
 };
 
+template<>
+struct TMassFragmentTraits<FMassSmartObjectUserFragment> final
+{
+	enum
+	{
+		AuthorAcceptsItsNotTriviallyCopyable = true
+	};
+};
+
 /** Fragment used to process time based smartobject interactions */
 USTRUCT()
 struct FMassSmartObjectTimedBehaviorFragment : public FMassFragment
diff --git a/Engine/Plugins/Runtime/MassGameplay/Source/MassSmartObjects/Public/MassSmartObjectRequest.h b/Engine/Plugins/Runtime/MassGameplay/Source/MassSmartObjects/Public/MassSmartObjectRequest.h
--- a/Engine/Plugins/Runtime/MassGameplay/Source/MassSmartObjects/Public/MassSmartObjectRequest.h
+++ b/Engine/Plugins/Runtime/MassGameplay/Source/MassSmartObjects/Public/MassSmartObjectRequest.h
@@ -135,6 +135,15 @@ struct FMassSmartObjectWorldLocationRequestFragment : public FMassFragment
 	FGameplayTagQuery ActivityRequirements;
 };
 
+template<>
+struct TMassFragmentTraits<FMassSmartObjectWorldLocationRequestFragment> final
+{
+	enum
+	{
+		AuthorAcceptsItsNotTriviallyCopyable = true
+	};
+};
+
 /**
  * Fragment used to build a list potential smart objects to use. Once added to an entity
  * this will be processed by the candidates finder processor to fill a SmartObjectCandidates
@@ -157,6 +166,15 @@ struct FMassSmartObjectLaneLocationRequestFragment : public FMassFragment
 	FGameplayTagQuery ActivityRequirements;
 };
 
+template<>
+struct TMassFragmentTraits<FMassSmartObjectLaneLocationRequestFragment> final
+{
+	enum
+	{
+		AuthorAcceptsItsNotTriviallyCopyable = true
+	};
+};
+
 /**
  * Special tag to mark processed requests
  */
diff --git a/Engine/Plugins/Runtime/MassGameplay/Source/MassSpawner/Public/MassEntityTemplate.h b/Engine/Plugins/Runtime/MassGameplay/Source/MassSpawner/Public/MassEntityTemplate.h
--- a/Engine/Plugins/Runtime/MassGameplay/Source/MassSpawner/Public/MassEntityTemplate.h
+++ b/Engine/Plugins/Runtime/MassGameplay/Source/MassSpawner/Public/MassEntityTemplate.h
@@ -129,13 +129,13 @@ struct FMassEntityTemplateData
 	template<typename T>
 	void AddFragment()
 	{
-		static_assert(UE::Mass::CFragment<T>, "Given struct doesn't represent a valid fragment type. Make sure to inherit from FMassFragment or one of its child-types.");
+		static_assert(UE::Mass::CFragment<T>, MASS_INVALID_FRAGMENT_MSG);
 		Composition.GetFragments().Add<T>();
 	}
 
 	void AddFragment(const UScriptStruct& FragmentType)
 	{
-		checkf(UE::Mass::IsA<FMassFragment>(&FragmentType), TEXT("Given struct doesn't represent a valid fragment type. Make sure to inherit from FMassFragment or one of its child-types."));
+		checkf(UE::Mass::IsA<FMassFragment>(&FragmentType), TEXT(MASS_INVALID_FRAGMENT_MSG_F), *FragmentType.GetName());
 		Composition.GetFragments().Add(FragmentType);
 	}
 
@@ -143,7 +143,8 @@ struct FMassEntityTemplateData
 	void AddFragment(FConstStructView Fragment)
 	{
 		const UScriptStruct* FragmentType = Fragment.GetScriptStruct();
-		checkf(UE::Mass::IsA<FMassFragment>(FragmentType), TEXT("Given struct doesn't represent a valid fragment type. Make sure to inherit from FMassFragment or one of its child-types."));
+		checkf(UE::Mass::IsA<FMassFragment>(FragmentType), TEXT(MASS_INVALID_FRAGMENT_MSG_F), *GetNameSafe(FragmentType));
+
 		if (!Composition.GetFragments().Contains(*FragmentType))
 		{
 			Composition.GetFragments().Add(*FragmentType);
@@ -158,7 +159,7 @@ struct FMassEntityTemplateData
 	template<typename T>
 	T& AddFragment_GetRef()
 	{
-		static_assert(UE::Mass::CFragment<T>, "Given struct doesn't represent a valid fragment type. Make sure to inherit from FMassFragment or one of its child-types.");
+		static_assert(UE::Mass::CFragment<T>, MASS_INVALID_FRAGMENT_MSG);
 		if (!Composition.GetFragments().Contains<T>())
 		{
 			Composition.GetFragments().Add<T>();
@@ -175,7 +176,7 @@ struct FMassEntityTemplateData
 	template<typename T>
 	T* GetMutableFragment()
 	{
-		static_assert(UE::Mass::CFragment<T>, "Given struct doesn't represent a valid fragment type. Make sure to inherit from FMassFragment or one of its child-types.");
+		static_assert(UE::Mass::CFragment<T>, MASS_INVALID_FRAGMENT_MSG);
 		FInstancedStruct* Fragment = InitialFragmentValues.FindByPredicate(FStructTypeEqualOperator(T::StaticStruct()));
 		return Fragment ? &Fragment->template GetMutable<T>() : (T*)nullptr;
 	}
diff --git a/Engine/Source/Developer/MassEntityTestSuite/Public/MassEntityTestTypes.h b/Engine/Source/Developer/MassEntityTestSuite/Public/MassEntityTestTypes.h
--- a/Engine/Source/Developer/MassEntityTestSuite/Public/MassEntityTestTypes.h
+++ b/Engine/Source/Developer/MassEntityTestSuite/Public/MassEntityTestTypes.h
@@ -81,6 +81,15 @@ struct FTestFragment_Array : public FMassFragment
 	}
 };
 
+template<>
+struct TMassFragmentTraits<FTestFragment_Array> final
+{
+	enum
+	{
+		AuthorAcceptsItsNotTriviallyCopyable = true
+	};
+};
+
 USTRUCT()
 struct FFragmentWithSharedPtr : public FMassFragment
 {
@@ -93,6 +102,15 @@ struct FFragmentWithSharedPtr : public FMassFragment
 	{}
 };
 
+template<>
+struct TMassFragmentTraits<FFragmentWithSharedPtr> final
+{
+	enum
+	{
+		AuthorAcceptsItsNotTriviallyCopyable = true
+	};
+};
+
 USTRUCT()
 struct FTestChunkFragment_Int : public FMassChunkFragment
 {
diff --git a/Engine/Source/Runtime/MassEntity/Public/MassCommandBuffer.h b/Engine/Source/Runtime/MassEntity/Public/MassCommandBuffer.h
--- a/Engine/Source/Runtime/MassEntity/Public/MassCommandBuffer.h
+++ b/Engine/Source/Runtime/MassEntity/Public/MassCommandBuffer.h
@@ -171,28 +171,28 @@ struct FMassCommandBuffer
 	template<typename T>
 	void AddFragment(FMassEntityHandle Entity)
 	{
-		static_assert(UE::Mass::CFragment<T>, "Given struct type is not a valid fragment type.");
+		static_assert(UE::Mass::CFragment<T>, MASS_INVALID_FRAGMENT_MSG);
 		PushCommand<FMassCommandAddFragmentsInternal<EMassCommandCheckTime::CompileTimeCheck, T>>(Entity);
 	}
 
 	template<typename T>
 	void AddFragment_RuntimeCheck(FMassEntityHandle Entity)
 	{
-		checkf(UE::Mass::IsA<FMassFragment>(T::StaticStruct()), TEXT("Given struct type is not a valid fragment type."));
+		checkf(UE::Mass::IsA<FMassFragment>(T::StaticStruct()), TEXT(MASS_INVALID_FRAGMENT_MSG_F), *T::StaticStruct()->GetName());
 		PushCommand<FMassCommandAddFragmentsInternal<EMassCommandCheckTime::RuntimeCheck, T>>(Entity);
 	}
 
 	template<typename T>
 	void RemoveFragment(FMassEntityHandle Entity)
 	{
-		static_assert(UE::Mass::CFragment<T>, "Given struct type is not a valid fragment type.");
+		static_assert(UE::Mass::CFragment<T>, MASS_INVALID_FRAGMENT_MSG);
 		PushCommand<FMassCommandRemoveFragmentsInternal<EMassCommandCheckTime::CompileTimeCheck, T>>(Entity);
 	}
 
 	template<typename T>
 	void RemoveFragment_RuntimeCheck(FMassEntityHandle Entity)
 	{
-		checkf(UE::Mass::IsA<FMassFragment>(T::StaticStruct()), TEXT("Given struct type is not a valid fragment type."));
+		checkf(UE::Mass::IsA<FMassFragment>(T::StaticStruct()), TEXT(MASS_INVALID_FRAGMENT_MSG_F), *T::StaticStruct()->GetName());
 		PushCommand<FMassCommandRemoveFragmentsInternal<EMassCommandCheckTime::RuntimeCheck, T>>(Entity);
 	}
 
diff --git a/Engine/Source/Runtime/MassEntity/Public/MassEntityConcepts.h b/Engine/Source/Runtime/MassEntity/Public/MassEntityConcepts.h
--- a/Engine/Source/Runtime/MassEntity/Public/MassEntityConcepts.h
+++ b/Engine/Source/Runtime/MassEntity/Public/MassEntityConcepts.h
@@ -3,25 +3,33 @@
 #pragma once
 
 #include "MassEntityElementTypes.h"
+#include "MassExternalSubsystemTraits.h"
 #include "Subsystems/Subsystem.h"
 
 
 namespace UE::Mass
 {
 	template<typename T>
-	concept CFragment = TIsDerivedFrom<typename TRemoveReference<T>::Type, FMassFragment>::Value;
-
+	using Clean = typename TRemoveReference<T>::Type;
+	
+	template<typename T>
+	concept CFragment = TIsDerivedFrom<Clean<T>, FMassFragment>::Value &&
+		(
+			std::is_trivially_copyable_v<Clean<T>> ||
+			static_cast<bool>(TMassFragmentTraits<Clean<T>>::AuthorAcceptsItsNotTriviallyCopyable)
+		);
+	
 	template<typename T>
-	concept CTag = TIsDerivedFrom<typename TRemoveReference<T>::Type, FMassTag>::Value;
+	concept CTag = TIsDerivedFrom<Clean<T>, FMassTag>::Value;
 
 	template<typename T>
-	concept CChunkFragment = TIsDerivedFrom<typename TRemoveReference<T>::Type, FMassChunkFragment>::Value;
+	concept CChunkFragment = TIsDerivedFrom<Clean<T>, FMassChunkFragment>::Value;
 
 	template<typename T>
-	concept CSharedFragment = TIsDerivedFrom<typename TRemoveReference<T>::Type, FMassSharedFragment>::Value;
+	concept CSharedFragment = TIsDerivedFrom<Clean<T>, FMassSharedFragment>::Value;
 
 	template<typename T>
-	concept CConstSharedFragment = TIsDerivedFrom<typename TRemoveReference<T>::Type, FMassConstSharedFragment>::Value;
+	concept CConstSharedFragment = TIsDerivedFrom<Clean<T>, FMassConstSharedFragment>::Value;
 
 	template<typename T>
 	concept CNonTag = CFragment<T> || CChunkFragment<T> || CSharedFragment<T> || CConstSharedFragment<T>;
@@ -30,7 +38,7 @@ namespace UE::Mass
 	concept CElement = CNonTag<T> || CTag<T>;
 
 	template<typename T>
-	concept CSubsystem = TIsDerivedFrom<typename TRemoveReference<T>::Type, USubsystem>::Value;
+	concept CSubsystem = TIsDerivedFrom<Clean<T>, USubsystem>::Value;
 
 	namespace Private
 	{
diff --git a/Engine/Source/Runtime/MassEntity/Public/MassEntityElementTypes.h b/Engine/Source/Runtime/MassEntity/Public/MassEntityElementTypes.h
--- a/Engine/Source/Runtime/MassEntity/Public/MassEntityElementTypes.h
+++ b/Engine/Source/Runtime/MassEntity/Public/MassEntityElementTypes.h
@@ -11,6 +11,13 @@ struct FMassFragment
 	GENERATED_BODY()
 };
 
+// these are the messages we'll print out when static checks whether a given type is a fragment fails
+#define _MASS_INVALID_FRAGMENT_CORE_MESSAGE "Make sure to inherit from FMassFragment or one of its child-types and ensure that the struct is trivially copyable, \
+or opt out by specializing TMassFragmentTraits for this type and setting AuthorAcceptsItsNotTriviallyCopyable = true"
+#define MASS_INVALID_FRAGMENT_MSG  "Given struct doesn't represent a valid fragment type." _MASS_INVALID_FRAGMENT_CORE_MESSAGE
+#define MASS_INVALID_FRAGMENT_MSG_F  "Type %s is not a valid fragment type." _MASS_INVALID_FRAGMENT_CORE_MESSAGE
+
+
 // This is the base class for types that will only be tested for presence/absence, i.e. Tags.
 // Subclasses should never contain any member properties.
 USTRUCT()
diff --git a/Engine/Source/Runtime/MassEntity/Public/MassEntityManager.h b/Engine/Source/Runtime/MassEntity/Public/MassEntityManager.h
--- a/Engine/Source/Runtime/MassEntity/Public/MassEntityManager.h
+++ b/Engine/Source/Runtime/MassEntity/Public/MassEntityManager.h
@@ -570,16 +570,14 @@ public:
 	template<typename FragmentType>
 	FragmentType& GetFragmentDataChecked(FMassEntityHandle EntityHandle) const
 	{
-		static_assert(UE::Mass::CFragment<FragmentType>
-			, "Given struct doesn't represent a valid fragment type. Make sure to inherit from FMassFragment or one of its child-types.");
+		static_assert(UE::Mass::CFragment<FragmentType>, MASS_INVALID_FRAGMENT_MSG);
 		return *((FragmentType*)InternalGetFragmentDataChecked(EntityHandle, FragmentType::StaticStruct()));
 	}
 
 	template<typename FragmentType>
 	FragmentType* GetFragmentDataPtr(FMassEntityHandle EntityHandle) const
 	{
-		static_assert(UE::Mass::CFragment<FragmentType>
-			, "Given struct doesn't represent a valid fragment type. Make sure to inherit from FMassFragment or one of its child-types.");
+		static_assert(UE::Mass::CFragment<FragmentType>, MASS_INVALID_FRAGMENT_MSG);
 		return (FragmentType*)InternalGetFragmentDataPtr(EntityHandle, FragmentType::StaticStruct());
 	}
 
diff --git a/Engine/Source/Runtime/MassEntity/Public/MassEntityView.h b/Engine/Source/Runtime/MassEntity/Public/MassEntityView.h
--- a/Engine/Source/Runtime/MassEntity/Public/MassEntityView.h
+++ b/Engine/Source/Runtime/MassEntity/Public/MassEntityView.h
@@ -54,8 +54,7 @@ struct FMassEntityView
 	{
 		static_assert(!UE::Mass::CTag<T>,
 			"Given struct doesn't represent a valid fragment type but a tag. Use HasTag instead.");
-		static_assert(UE::Mass::CFragment<T>,
-			"Given struct doesn't represent a valid fragment type. Make sure to inherit from FMassFragment or one of its child-types.");
+		static_assert(UE::Mass::CFragment<T>, MASS_INVALID_FRAGMENT_MSG);
 
 		return *((T*)GetFragmentPtrChecked(*T::StaticStruct()));
 	}
@@ -66,8 +65,7 @@ struct FMassEntityView
 	{
 		static_assert(!UE::Mass::CTag<T>,
 			"Given struct doesn't represent a valid fragment type but a tag. Use HasTag instead.");
-		static_assert(UE::Mass::CFragment<T>,
-			"Given struct doesn't represent a valid fragment type. Make sure to inherit from FMassFragment or one of its child-types.");
+		static_assert(UE::Mass::CFragment<T>, MASS_INVALID_FRAGMENT_MSG);
 
 		return (T*)GetFragmentPtr(*T::StaticStruct());
 	}
diff --git a/Engine/Source/Runtime/MassEntity/Public/MassExecutionContext.h b/Engine/Source/Runtime/MassEntity/Public/MassExecutionContext.h
--- a/Engine/Source/Runtime/MassEntity/Public/MassExecutionContext.h
+++ b/Engine/Source/Runtime/MassEntity/Public/MassExecutionContext.h
@@ -303,7 +303,7 @@ public:
 	template<typename T>
 	bool DoesArchetypeHaveFragment() const
 	{
-		static_assert(UE::Mass::CFragment<T>, "Given struct is not of a valid fragment type.");
+		static_assert(UE::Mass::CFragment<T>, MASS_INVALID_FRAGMENT_MSG);
 		return CurrentArchetypeCompositionDescriptor.GetFragments().Contains<T>();
 	}
 
diff --git a/Engine/Source/Runtime/MassEntity/Public/MassExternalSubsystemTraits.h b/Engine/Source/Runtime/MassEntity/Public/MassExternalSubsystemTraits.h
--- a/Engine/Source/Runtime/MassEntity/Public/MassExternalSubsystemTraits.h
+++ b/Engine/Source/Runtime/MassEntity/Public/MassExternalSubsystemTraits.h
@@ -36,7 +36,7 @@ struct TMassExternalSubsystemTraits final
 };
 
 /** 
- * Shared Fragments' traits.
+ * Shared Fragment traits.
  * @see TMassExternalSubsystemTraits
  */
 template <typename T>
@@ -48,3 +48,20 @@ struct TMassSharedFragmentTraits final
 		ThreadSafeWrite = false,
 	};
 };
+
+/** 
+ * Fragment traits.
+ * @see TMassExternalSubsystemTraits
+ */
+template <typename T>
+struct TMassFragmentTraits final
+{
+	enum
+	{
+		// Fragment types are best kept trivially copyable for performance reasons.
+		// To enforce that we test this trait when checking if a given type is a valid fragment type.
+		// This test can be skipped by specifically opting out, which also makes the 
+		// type author's acknowledge that they know and accept that the type contains unorthodox elements.
+		AuthorAcceptsItsNotTriviallyCopyable = false
+	};
+};
\ No newline at end of file
diff --git a/Engine/Source/Runtime/MassEntity/Public/MassRequirements.h b/Engine/Source/Runtime/MassEntity/Public/MassRequirements.h
--- a/Engine/Source/Runtime/MassEntity/Public/MassRequirements.h
+++ b/Engine/Source/Runtime/MassEntity/Public/MassRequirements.h
@@ -215,7 +215,7 @@ struct FMassFragmentRequirements
 		checkf(FragmentRequirements.FindByPredicate([](const FMassFragmentRequirementDescription& Item) { return Item.StructType == T::StaticStruct(); }) == nullptr
 			, TEXT("Duplicated requirements are not supported. %s already present"), *T::StaticStruct()->GetName());
 
-		static_assert(UE::Mass::CFragment<T>, "Given struct doesn't represent a valid fragment type. Make sure to inherit from FMassFragment or one of its child-types.");
+		static_assert(UE::Mass::CFragment<T>, MASS_INVALID_FRAGMENT_MSG);
 		
 		if (Presence != EMassFragmentPresence::None)
 		{
