[Mass] reconfigured the FMassProcessingContext created for initializer-processors in UMassSpawnerSubsystem::DoSpawning.

The original configuration was flushing commands if any have been issued, but done so in a way that forced flushing of all the system's commands accumulated up until that point. That resulted in different behavior in entity spawning depending on whether there were initializers for given entity available or not.
New configuration results in keeping the behavior consistent, regardless of initializers presence, which is to not flush commands and if any are issued they will be appended to the main buffer when FMassProcessingContext goes out of scope.

diff --git a/Engine/Plugins/Runtime/MassGameplay/Source/MassSpawner/Private/MassSpawnerSubsystem.cpp b/Engine/Plugins/Runtime/MassGameplay/Source/MassSpawner/Private/MassSpawnerSubsystem.cpp
--- a/Engine/Plugins/Runtime/MassGameplay/Source/MassSpawner/Private/MassSpawnerSubsystem.cpp
+++ b/Engine/Plugins/Runtime/MassGameplay/Source/MassSpawner/Private/MassSpawnerSubsystem.cpp
@@ -148,7 +148,7 @@ TSharedPtr<FMassEntityManager::FEntityCreationContext> UMassSpawnerSubsystem::Do
 
 	if (SpawnDataInitializer)
 	{
-		FMassProcessingContext ProcessingContext(EntityManager, /*TimeDelta=*/0.0f);
+		FMassProcessingContext ProcessingContext(EntityManager, /*TimeDelta=*/0.0f, /*bInFlushCommandBuffer=*/false);
 		ProcessingContext.AuxData = SpawnData;
 		UE::Mass::Executor::RunProcessorsView(MakeArrayView(&SpawnDataInitializer, 1), ProcessingContext, CreationContext->GetEntityCollections(*EntityManager.Get()));
 	}
