PR #12786: Destroy Mass puppet actors immediately when destructed
Along with the fix: Fixed bug with MassRepresentation not calling Disabled in some cases

diff --git a/Engine/Plugins/Runtime/MassGameplay/Source/MassRepresentation/Private/MassRepresentationActorManagement.cpp b/Engine/Plugins/Runtime/MassGameplay/Source/MassRepresentation/Private/MassRepresentationActorManagement.cpp
--- a/Engine/Plugins/Runtime/MassGameplay/Source/MassRepresentation/Private/MassRepresentationActorManagement.cpp
+++ b/Engine/Plugins/Runtime/MassGameplay/Source/MassRepresentation/Private/MassRepresentationActorManagement.cpp
@@ -111,19 +111,19 @@ EMassActorSpawnRequestAction UMassRepresentationActorManagement::OnPostActorSpaw
 	return EMassActorSpawnRequestAction::Remove;
 }
 
-void UMassRepresentationActorManagement::ReleaseAnyActorOrCancelAnySpawning(FMassEntityManager& EntityManager, const FMassEntityHandle MassAgent)
+void UMassRepresentationActorManagement::ReleaseAnyActorOrCancelAnySpawning(FMassEntityManager& EntityManager, const FMassEntityHandle MassAgent, const bool bImmediate /*= false*/)
 {
 	FMassEntityView EntityView(EntityManager, MassAgent);
 	FMassActorFragment& ActorInfo = EntityView.GetFragmentData<FMassActorFragment>();
 	FMassRepresentationFragment& Representation = EntityView.GetFragmentData<FMassRepresentationFragment>();
 	UMassRepresentationSubsystem* RepresentationSubsystem = EntityView.GetSharedFragmentData<FMassRepresentationSubsystemSharedFragment>().RepresentationSubsystem;
 	check(RepresentationSubsystem);
-	ReleaseAnyActorOrCancelAnySpawning(*RepresentationSubsystem, MassAgent, ActorInfo, Representation);
+	ReleaseAnyActorOrCancelAnySpawning(*RepresentationSubsystem, MassAgent, ActorInfo, Representation, /*CachedActorSubsystem=*/nullptr, bImmediate);
 }
 
 void UMassRepresentationActorManagement::ReleaseAnyActorOrCancelAnySpawning(UMassRepresentationSubsystem& RepresentationSubsystem
 	, const FMassEntityHandle MassAgent, FMassActorFragment& ActorInfo, FMassRepresentationFragment& Representation
-	, UMassActorSubsystem* CachedActorSubsystem)
+	, UMassActorSubsystem* CachedActorSubsystem, bool bImmediate /*= false*/)
 {
 	// This method can only release owned by mass actors
 	AActor* Actor = ActorInfo.GetOwnedByMassMutable();
@@ -137,11 +137,11 @@ void UMassRepresentationActorManagement::ReleaseAnyActorOrCancelAnySpawning(UMas
 	// Try releasing both as we can have a low res actor and a high res spawning request
 	if (Representation.HighResTemplateActorIndex != INDEX_NONE)
 	{
-		RepresentationSubsystem.ReleaseTemplateActorOrCancelSpawning(MassAgent, Representation.HighResTemplateActorIndex, Actor, Representation.ActorSpawnRequestHandle);
+		RepresentationSubsystem.ReleaseTemplateActorOrCancelSpawning(MassAgent, Representation.HighResTemplateActorIndex, Actor, Representation.ActorSpawnRequestHandle, bImmediate);
 	}	
 	if (Representation.LowResTemplateActorIndex != Representation.HighResTemplateActorIndex && Representation.LowResTemplateActorIndex != INDEX_NONE)
 	{
-		RepresentationSubsystem.ReleaseTemplateActorOrCancelSpawning(MassAgent, Representation.LowResTemplateActorIndex, Actor, Representation.ActorSpawnRequestHandle);
+		RepresentationSubsystem.ReleaseTemplateActorOrCancelSpawning(MassAgent, Representation.LowResTemplateActorIndex, Actor, Representation.ActorSpawnRequestHandle, bImmediate);
 	}
 	check(!Representation.ActorSpawnRequestHandle.IsValid());
 }
diff --git a/Engine/Plugins/Runtime/MassGameplay/Source/MassRepresentation/Private/MassRepresentationProcessor.cpp b/Engine/Plugins/Runtime/MassGameplay/Source/MassRepresentation/Private/MassRepresentationProcessor.cpp
--- a/Engine/Plugins/Runtime/MassGameplay/Source/MassRepresentation/Private/MassRepresentationProcessor.cpp
+++ b/Engine/Plugins/Runtime/MassGameplay/Source/MassRepresentation/Private/MassRepresentationProcessor.cpp
@@ -136,6 +136,10 @@ void UMassRepresentationProcessor::UpdateRepresentation(FMassExecutionContext& C
 
 		auto DisableActorForISM = [&](AActor*& Actor)
 		{
+			if (Actor != nullptr)
+			{
+				RepresentationActorManagement->SetActorEnabled(EMassActorEnabledType::Disabled, *Actor, EntityIt, Context.Defer());
+			}
 			if (!Actor || ActorInfo.IsOwnedByMass())
 			{
 				// Execute only if the high res is different than the low res Actor 
@@ -143,26 +147,23 @@ void UMassRepresentationProcessor::UpdateRepresentation(FMassExecutionContext& C
 				if (Representation.HighResTemplateActorIndex != Representation.LowResTemplateActorIndex || !RepresentationParams.bKeepLowResActors)
 				{
 					// Try releasing the high actor or any high res spawning request
-					if (ReleaseActorOrCancelSpawning(*RepresentationSubsystem, MassActorSubsystem, MassAgentEntityHandle, ActorInfo, Representation.HighResTemplateActorIndex, Representation.ActorSpawnRequestHandle, Context.Defer()))
-					{
-						Actor = ActorInfo.GetOwnedByMassMutable();
-					}
+					ReleaseActorOrCancelSpawning(*RepresentationSubsystem, MassActorSubsystem, MassAgentEntityHandle, ActorInfo, Representation.HighResTemplateActorIndex, Representation.ActorSpawnRequestHandle, Context.Defer());
 					// Do not do the same with low res if indicated so
-					if (!RepresentationParams.bKeepLowResActors && ReleaseActorOrCancelSpawning(*RepresentationSubsystem, MassActorSubsystem, MassAgentEntityHandle, ActorInfo, Representation.LowResTemplateActorIndex, Representation.ActorSpawnRequestHandle, Context.Defer()))
+					if (!RepresentationParams.bKeepLowResActors)
 					{
-						Actor = ActorInfo.GetOwnedByMassMutable();
+						ReleaseActorOrCancelSpawning(*RepresentationSubsystem, MassActorSubsystem, MassAgentEntityHandle, ActorInfo, Representation.LowResTemplateActorIndex, Representation.ActorSpawnRequestHandle, Context.Defer());
 					}
 				}
 				// If we already queued spawn request but have changed our mind, continue with it but once we get the actor back, disable it immediately
 				if (Representation.ActorSpawnRequestHandle.IsValid())
 				{
 					Actor = RepresentationActorManagement->GetOrSpawnActor(*RepresentationSubsystem, CachedEntityManager, MassAgentEntityHandle, TransformFragment.GetTransform(), Representation.LowResTemplateActorIndex, Representation.ActorSpawnRequestHandle, RepresentationActorManagement->GetSpawnPriority(RepresentationLOD));
+					if (Actor != nullptr)
+					{
+						RepresentationActorManagement->SetActorEnabled(EMassActorEnabledType::Disabled, *Actor, EntityIt, Context.Defer());
+					}
 				}
 			}
-			if (Actor != nullptr)
-			{
-				RepresentationActorManagement->SetActorEnabled(EMassActorEnabledType::Disabled, *Actor, EntityIt, Context.Defer());
-			}
 		};
 
 		// === 3. Switch the in-game instance representation depending on the EMassRepresentationType;
@@ -478,7 +479,7 @@ void UMassRepresentationFragmentDestructor::Execute(FMassEntityManager& EntityMa
 
 			const FMassEntityHandle MassAgentEntityHandle = Context.GetEntity(EntityIt);
 
-			UMassRepresentationActorManagement::ReleaseAnyActorOrCancelAnySpawning(*RepresentationSubsystem, MassAgentEntityHandle, ActorInfo, Representation, ActorSubsystem);
+			UMassRepresentationActorManagement::ReleaseAnyActorOrCancelAnySpawning(*RepresentationSubsystem, MassAgentEntityHandle, ActorInfo, Representation, ActorSubsystem, /*bImmediate=*/true);
 		}
 	});
 }
diff --git a/Engine/Plugins/Runtime/MassGameplay/Source/MassRepresentation/Private/MassRepresentationSubsystem.cpp b/Engine/Plugins/Runtime/MassGameplay/Source/MassRepresentation/Private/MassRepresentationSubsystem.cpp
--- a/Engine/Plugins/Runtime/MassGameplay/Source/MassRepresentation/Private/MassRepresentationSubsystem.cpp
+++ b/Engine/Plugins/Runtime/MassGameplay/Source/MassRepresentation/Private/MassRepresentationSubsystem.cpp
@@ -279,11 +279,12 @@ void UMassRepresentationSubsystem::ReleaseAllResources()
 	VisualizationComponent->ClearAllVisualInstances();
 }
 
-bool UMassRepresentationSubsystem::ReleaseTemplateActorOrCancelSpawning(const FMassEntityHandle MassAgent, const int16 TemplateActorIndex, AActor* ActorToRelease, FMassActorSpawnRequestHandle& SpawnRequestHandle)
+bool UMassRepresentationSubsystem::ReleaseTemplateActorOrCancelSpawning(const FMassEntityHandle MassAgent, const int16 TemplateActorIndex, AActor* ActorToRelease
+	, FMassActorSpawnRequestHandle& SpawnRequestHandle, const bool bImmediate /*= false*/)
 {
 	// First try to cancel the spawning request, then try to release the actor
 	if( CancelSpawningInternal(TemplateActorIndex, SpawnRequestHandle) ||
-		ReleaseTemplateActorInternal(TemplateActorIndex, ActorToRelease, false))
+		ReleaseTemplateActorInternal(TemplateActorIndex, ActorToRelease, bImmediate))
 	{ 
 		const int32 RefCount = --(HandledMassAgents.FindChecked(MassAgent));
 		checkf(RefCount >= 0, TEXT("RefCount are expected to be greater than or equal to 0"));
diff --git a/Engine/Plugins/Runtime/MassGameplay/Source/MassRepresentation/Public/MassRepresentationActorManagement.h b/Engine/Plugins/Runtime/MassGameplay/Source/MassRepresentation/Public/MassRepresentationActorManagement.h
--- a/Engine/Plugins/Runtime/MassGameplay/Source/MassRepresentation/Public/MassRepresentationActorManagement.h
+++ b/Engine/Plugins/Runtime/MassGameplay/Source/MassRepresentation/Public/MassRepresentationActorManagement.h
@@ -90,9 +90,10 @@ public:
 	 *          So any reference to fragment might become invalid.
 	 * @param EntityManager to use to retrieve the mass agent fragments
 	 * @param MassAgent is the handle to the associated mass agent
+	 * @param bImmediate whether to perform the actor destruction immediately, otherwise it will be queued for processing later
 	 * @return True if actor was release or spawning request was canceled
 	 */
-	static UE_API void ReleaseAnyActorOrCancelAnySpawning(FMassEntityManager& EntityManager, const FMassEntityHandle MassAgent);
+	static UE_API void ReleaseAnyActorOrCancelAnySpawning(FMassEntityManager& EntityManager, const FMassEntityHandle MassAgent, bool bImmediate = false);
 
 	/**
 	 * static Release an actor or cancel its spawning
@@ -104,9 +105,10 @@ public:
 	 * @param Representation fragment containing the current and previous visual state
 	 * @param ActorSubsystem passed over to FMassActorFragment::ResetAndUpdateHandleMap, used to avoid fetching the subsystem
 	 *	from UWorld every time.
+	 * @param bImmediate whether to perform the actor destruction immediately, otherwise it will be queued for processing later
 	 */
 	static UE_API void ReleaseAnyActorOrCancelAnySpawning(UMassRepresentationSubsystem& RepresentationSubsystem, const FMassEntityHandle MassAgent
-		, FMassActorFragment& ActorInfo, FMassRepresentationFragment& Representation, UMassActorSubsystem* ActorSubsystem = nullptr);
+		, FMassActorFragment& ActorInfo, FMassRepresentationFragment& Representation, UMassActorSubsystem* ActorSubsystem = nullptr, bool bImmediate = false);
 };
 
 #undef UE_API
diff --git a/Engine/Plugins/Runtime/MassGameplay/Source/MassRepresentation/Public/MassRepresentationSubsystem.h b/Engine/Plugins/Runtime/MassGameplay/Source/MassRepresentation/Public/MassRepresentationSubsystem.h
--- a/Engine/Plugins/Runtime/MassGameplay/Source/MassRepresentation/Public/MassRepresentationSubsystem.h
+++ b/Engine/Plugins/Runtime/MassGameplay/Source/MassRepresentation/Public/MassRepresentationSubsystem.h
@@ -124,9 +124,10 @@ public:
 	 * @param TemplateActorIndex is the template type of the actor to release in case it was successfully spawned
 	 * @param ActorToRelease is the actual actor to release if any
 	 * @param SpawnRequestHandle [IN/OUT] previously requested spawn, gets invalidated as a result of this call.
+	 * @param bImmediate whether to perform the actor destruction immediately, otherwise it will be queued for processing later
 	 * @return True if actor was released or spawning request was canceled
 	 */
-	MASSREPRESENTATION_API bool ReleaseTemplateActorOrCancelSpawning(const FMassEntityHandle MassAgent, const int16 TemplateActorIndex, AActor* ActorToRelease, FMassActorSpawnRequestHandle& SpawnRequestHandle);
+	MASSREPRESENTATION_API bool ReleaseTemplateActorOrCancelSpawning(const FMassEntityHandle MassAgent, const int16 TemplateActorIndex, AActor* ActorToRelease, FMassActorSpawnRequestHandle& SpawnRequestHandle, bool bImmediate = false);
 
 
 	/**
