[Mass] fixed inconsistencies in how CancelSpawningInternal is being used in relation to recent addition of "bImmediate" mode in surrounding context.

#rnx

diff --git a/Engine/Plugins/Runtime/MassGameplay/Source/MassRepresentation/Private/MassRepresentationSubsystem.cpp b/Engine/Plugins/Runtime/MassGameplay/Source/MassRepresentation/Private/MassRepresentationSubsystem.cpp
--- a/Engine/Plugins/Runtime/MassGameplay/Source/MassRepresentation/Private/MassRepresentationSubsystem.cpp
+++ b/Engine/Plugins/Runtime/MassGameplay/Source/MassRepresentation/Private/MassRepresentationSubsystem.cpp
@@ -283,8 +283,8 @@ bool UMassRepresentationSubsystem::ReleaseTemplateActorOrCancelSpawning(const FM
 	, FMassActorSpawnRequestHandle& SpawnRequestHandle, const bool bImmediate /*= false*/)
 {
 	// First try to cancel the spawning request, then try to release the actor
-	if( CancelSpawningInternal(TemplateActorIndex, SpawnRequestHandle) ||
-		ReleaseTemplateActorInternal(TemplateActorIndex, ActorToRelease, bImmediate))
+	if (CancelSpawningInternal(TemplateActorIndex, SpawnRequestHandle, bImmediate) 
+		|| ReleaseTemplateActorInternal(TemplateActorIndex, ActorToRelease, bImmediate))
 	{ 
 		const int32 RefCount = --(HandledMassAgents.FindChecked(MassAgent));
 		checkf(RefCount >= 0, TEXT("RefCount are expected to be greater than or equal to 0"));
@@ -348,7 +348,7 @@ bool UMassRepresentationSubsystem::ReleaseTemplateActorInternal(const int16 Temp
 	return true;
 }
 
-bool UMassRepresentationSubsystem::CancelSpawningInternal(const int16 TemplateActorIndex, FMassActorSpawnRequestHandle& SpawnRequestHandle)
+bool UMassRepresentationSubsystem::CancelSpawningInternal(const int16 TemplateActorIndex, FMassActorSpawnRequestHandle& SpawnRequestHandle, const bool bImmediateActorRelease)
 {
 	UE_MT_SCOPED_READ_ACCESS(TemplateActorsMTAccessDetector);
 	if (!TemplateActors.IsValidIndex(TemplateActorIndex))
@@ -375,7 +375,7 @@ bool UMassRepresentationSubsystem::CancelSpawningInternal(const int16 TemplateAc
 	if (SpawnRequest.SpawnStatus == ESpawnRequestStatus::Succeeded)
 	{
 		check(SpawnRequest.SpawnedActor);
-		ReleaseTemplateActorInternal(TemplateActorIndex, SpawnRequest.SpawnedActor, false);
+		ReleaseTemplateActorInternal(TemplateActorIndex, SpawnRequest.SpawnedActor, bImmediateActorRelease);
 	}
 
 	// Remove the spawn request
diff --git a/Engine/Plugins/Runtime/MassGameplay/Source/MassRepresentation/Public/MassRepresentationSubsystem.h b/Engine/Plugins/Runtime/MassGameplay/Source/MassRepresentation/Public/MassRepresentationSubsystem.h
--- a/Engine/Plugins/Runtime/MassGameplay/Source/MassRepresentation/Public/MassRepresentationSubsystem.h
+++ b/Engine/Plugins/Runtime/MassGameplay/Source/MassRepresentation/Public/MassRepresentationSubsystem.h
@@ -169,7 +169,7 @@ protected:
 	MASSREPRESENTATION_API void OnMassAgentComponentEntityDetaching(const UMassAgentComponent& AgentComponent);
 
 	MASSREPRESENTATION_API bool ReleaseTemplateActorInternal(const int16 TemplateActorIndex, AActor* ActorToRelease, bool bImmediate);
-	MASSREPRESENTATION_API bool CancelSpawningInternal(const int16 TemplateActorIndex, FMassActorSpawnRequestHandle& SpawnRequestHandle);
+	MASSREPRESENTATION_API bool CancelSpawningInternal(const int16 TemplateActorIndex, FMassActorSpawnRequestHandle& SpawnRequestHandle, bool bImmediateActorRelease = false);
 
 	static MASSREPRESENTATION_API void AddReferencedObjects(UObject* InThis, FReferenceCollector& Collector);
 
