PR #13067: Mass: Add optional ability for applications to disable invalid mesh desc error messages

diff --git a/Engine/Plugins/Runtime/MassGameplay/Source/MassRepresentation/Private/MassVisualizationTrait.cpp b/Engine/Plugins/Runtime/MassGameplay/Source/MassRepresentation/Private/MassVisualizationTrait.cpp
--- a/Engine/Plugins/Runtime/MassGameplay/Source/MassRepresentation/Private/MassVisualizationTrait.cpp
+++ b/Engine/Plugins/Runtime/MassGameplay/Source/MassRepresentation/Private/MassVisualizationTrait.cpp
@@ -184,7 +184,8 @@ bool UMassVisualizationTrait::ValidateParams() const
 	{
 		// the SM config provided is not valid. We need to check if EMassRepresentationType::StaticMeshInstance
 		// is being used as any of the LODRepresentations. If so then we need to clear those out and report an error
-		if (StaticMeshInstanceDesc.IsValid() == false)
+		if (bRequireValidStaticMeshInstanceDesc == true
+			&& StaticMeshInstanceDesc.IsValid() == false)
 		{
 			for (int32 LODIndex = 0; LODIndex < EMassLOD::Max; ++LODIndex)
 			{
diff --git a/Engine/Plugins/Runtime/MassGameplay/Source/MassRepresentation/Public/MassVisualizationTrait.h b/Engine/Plugins/Runtime/MassGameplay/Source/MassRepresentation/Public/MassVisualizationTrait.h
--- a/Engine/Plugins/Runtime/MassGameplay/Source/MassRepresentation/Public/MassVisualizationTrait.h
+++ b/Engine/Plugins/Runtime/MassGameplay/Source/MassRepresentation/Public/MassVisualizationTrait.h
@@ -55,6 +55,17 @@ public:
 	/** the property is marked like this to ensure it won't show up in UI */
 	UPROPERTY(EditDefaultsOnly, Category = "Mass|Visual")
 	bool bCanModifyRepresentationSubsystemClass = true;
+
+	/**
+	 * When True, ValidateParams() will require a valid StaticMeshInstanceDesc.
+	 *
+	 * When False, ValidateParams() skips checking the StaticMeshInstanceDesc.
+	 *
+	 * Use False when you expect a given Trait to have invalid StaticMeshInstanceDesc
+	 * settings, and you don't want errors logged when it occurs.
+	 */
+	UPROPERTY(EditDefaultsOnly, Category = "Mass|Visual", AdvancedDisplay)
+	bool bRequireValidStaticMeshInstanceDesc = true;
 #endif // WITH_EDITORONLY_DATA
 
 protected:
