[MassGameplay]
Added additional null guards around UMassRepresentationSubsystem::VisualizationComponent which may become invalid while the UMassRepresentationSubsystem is still active.

#DisabledValidations nonUnityAndStaticAnalysis

diff --git a/Engine/Plugins/Runtime/MassGameplay/Source/MassRepresentation/Private/MassRepresentationSubsystem.cpp b/Engine/Plugins/Runtime/MassGameplay/Source/MassRepresentation/Private/MassRepresentationSubsystem.cpp
--- a/Engine/Plugins/Runtime/MassGameplay/Source/MassRepresentation/Private/MassRepresentationSubsystem.cpp
+++ b/Engine/Plugins/Runtime/MassGameplay/Source/MassRepresentation/Private/MassRepresentationSubsystem.cpp
@@ -276,7 +276,10 @@ void UMassRepresentationSubsystem::ReleaseAllResources()
 	}
 
 	// Release all static meshes resources
-	VisualizationComponent->ClearAllVisualInstances();
+	if (LIKELY(VisualizationComponent))
+	{
+		VisualizationComponent->ClearAllVisualInstances();
+	}
 }
 
 bool UMassRepresentationSubsystem::ReleaseTemplateActorOrCancelSpawning(const FMassEntityHandle MassAgent, const int16 TemplateActorIndex, AActor* ActorToRelease
@@ -497,9 +500,10 @@ bool UMassRepresentationSubsystem::SpawnVisualizer(TNotNull<UWorld*> World)
 
 void UMassRepresentationSubsystem::OnProcessingPhaseStarted(const float /*DeltaSeconds*/, const EMassProcessingPhase Phase) const
 {
-	check(VisualizationComponent);
-	switch (Phase)
+	if (LIKELY(VisualizationComponent))
 	{
+		switch (Phase)
+		{
 		case EMassProcessingPhase::PrePhysics:
 			VisualizationComponent->BeginVisualChanges();
 			break;
@@ -509,6 +513,7 @@ void UMassRepresentationSubsystem::OnProcessingPhaseStarted(const float /*DeltaS
 		default:
 			check(false); // Need to handle this case
 			break;
+		}
 	}
 }
 
