[MassGameplay]
Prevent entity spawning after world teardown

diff --git a/Engine/Plugins/Runtime/MassGameplay/Source/MassSpawner/Private/MassSpawner.cpp b/Engine/Plugins/Runtime/MassGameplay/Source/MassSpawner/Private/MassSpawner.cpp
--- a/Engine/Plugins/Runtime/MassGameplay/Source/MassSpawner/Private/MassSpawner.cpp
+++ b/Engine/Plugins/Runtime/MassGameplay/Source/MassSpawner/Private/MassSpawner.cpp
@@ -227,6 +227,11 @@ void AMassSpawner::EndPlay(const EEndPlayReason::Type EndPlayReason)
 {
 	UMassSimulationSubsystem::GetOnSimulationStarted().Remove(SimulationStartedHandle);
 
+	if (StreamingHandle.IsValid() && StreamingHandle->IsActive())
+	{
+		StreamingHandle->CancelHandle();
+	}
+
 	DoDespawning();
 
 	Super::EndPlay(EndPlayReason);
@@ -368,7 +373,15 @@ void AMassSpawner::DoSpawning()
 
 void AMassSpawner::OnSpawnDataGenerationFinished(TConstArrayView<FMassEntitySpawnDataGeneratorResult> Results, FMassSpawnDataGenerator* FinishedGenerator)
 {
-	// @todo: this can be potentially expensive copy for the instanced structs, could there be a way to use move gere instead?
+	if (UWorld* World = GetWorld())
+	{
+		if (World->bIsTearingDown || World->IsCleanedUp())
+		{
+			return;
+		}
+	}
+
+	// @todo: this can be potentially expensive copy for the instanced structs, could there be a way to use move here instead?
 	AllGeneratedResults.Append(Results.GetData(), Results.Num());
 
 	bool bAllSpawnPointsGenerated = true;
