[MassDebugger] added MT detector around SelectedEntity, SelectedEntityInfo and Shapes

diff --git a/Engine/Plugins/AI/MassAI/Source/MassAIDebug/Private/GameplayDebuggerCategory_Mass.cpp b/Engine/Plugins/AI/MassAI/Source/MassAIDebug/Private/GameplayDebuggerCategory_Mass.cpp
--- a/Engine/Plugins/AI/MassAI/Source/MassAIDebug/Private/GameplayDebuggerCategory_Mass.cpp
+++ b/Engine/Plugins/AI/MassAI/Source/MassAIDebug/Private/GameplayDebuggerCategory_Mass.cpp
@@ -406,31 +406,32 @@ void FGameplayDebuggerCategory_Mass::CollectData(APlayerController* OwnerPC, AAc
 	//@todo could shave off some perf cost if UMassDebuggerSubsystem used FGameplayDebuggerShape directly
 	if (bShowShapes && Debugger)
 	{
-		const TArray<UMassDebuggerSubsystem::FShapeDesc>* Shapes = Debugger->GetShapes();
-		check(Shapes);
 		// EMassEntityDebugShape::Box
-		for (const UMassDebuggerSubsystem::FShapeDesc& Desc : Shapes[uint8(EMassEntityDebugShape::Box)])
-		{
-			AddShape(FGameplayDebuggerShape::MakeBox(Desc.Location, FVector(Desc.Size), FColor::Blue));
-		}
+		Debugger->ForEachShape(EMassEntityDebugShape::Box, [this](const UMassDebuggerSubsystem::FShapeDesc& Desc)
+			{
+				AddShape(FGameplayDebuggerShape::MakeBox(Desc.Location, FVector(Desc.Size), FColor::Blue));
+			});
+
 		// EMassEntityDebugShape::Cone
 		// note that we're modifying the Size here because MakeCone is using the third param as Cone's "height", while all mass debugger shapes are created with agent radius
 		// FGameplayDebuggerShape::Draw is using 0.25 rad for cone angle, so that's what we'll use here
 		static const float Tan025Rad = FMath::Tan(0.25f);
-		for (const UMassDebuggerSubsystem::FShapeDesc& Desc : Shapes[uint8(EMassEntityDebugShape::Cone)])
-		{
-			AddShape(FGameplayDebuggerShape::MakeCone(Desc.Location, FVector::UpVector, Desc.Size / Tan025Rad, FColor::Orange));
-		}
+		Debugger->ForEachShape(EMassEntityDebugShape::Cone, [this](const UMassDebuggerSubsystem::FShapeDesc& Desc)
+			{
+				AddShape(FGameplayDebuggerShape::MakeCone(Desc.Location, FVector::UpVector, Desc.Size / Tan025Rad, FColor::Orange));
+			});
+
 		// EMassEntityDebugShape::Cylinder
-		for (const UMassDebuggerSubsystem::FShapeDesc& Desc : Shapes[uint8(EMassEntityDebugShape::Cylinder)])
-		{
-			AddShape(FGameplayDebuggerShape::MakeCylinder(Desc.Location, Desc.Size, Desc.Size * 2, FColor::Yellow));
-		}
+		Debugger->ForEachShape(EMassEntityDebugShape::Cylinder, [this](const UMassDebuggerSubsystem::FShapeDesc& Desc)
+			{
+				AddShape(FGameplayDebuggerShape::MakeCylinder(Desc.Location, Desc.Size, Desc.Size * 2, FColor::Yellow));
+			});
+
 		// EMassEntityDebugShape::Capsule
-		for (const UMassDebuggerSubsystem::FShapeDesc& Desc : Shapes[uint8(EMassEntityDebugShape::Capsule)])
-		{
-			AddShape(FGameplayDebuggerShape::MakeCapsule(Desc.Location, Desc.Size, Desc.Size * 2, FColor::Green));
-		}
+		Debugger->ForEachShape(EMassEntityDebugShape::Capsule, [this](const UMassDebuggerSubsystem::FShapeDesc& Desc)
+			{
+				AddShape(FGameplayDebuggerShape::MakeCapsule(Desc.Location, Desc.Size, Desc.Size * 2, FColor::Green));
+			});
 	}
 
 	if (bShowAgentFragments)
diff --git a/Engine/Plugins/Runtime/MassGameplay/Source/MassGameplayDebug/Private/AssignDebugVisProcessor.cpp b/Engine/Plugins/Runtime/MassGameplay/Source/MassGameplayDebug/Private/AssignDebugVisProcessor.cpp
--- a/Engine/Plugins/Runtime/MassGameplay/Source/MassGameplayDebug/Private/AssignDebugVisProcessor.cpp
+++ b/Engine/Plugins/Runtime/MassGameplay/Source/MassGameplayDebug/Private/AssignDebugVisProcessor.cpp
@@ -1,25 +1,33 @@
 // Copyright Epic Games, Inc. All Rights Reserved.
 
 #include "AssignDebugVisProcessor.h"
+#include UE_INLINE_GENERATED_CPP_BY_NAME(AssignDebugVisProcessor)
+
+#if WITH_MASSGAMEPLAY_DEBUG
 #include "MassGameplayDebugTypes.h"
 #include "MassDebugVisualizationComponent.h"
 #include "MassDebuggerSubsystem.h"
 #include "MassExecutionContext.h"
+#endif // WITH_MASSGAMEPLAY_DEBUG
 
 //----------------------------------------------------------------------//
 // UAssignDebugVisProcessor
 //----------------------------------------------------------------------//
-
-#include UE_INLINE_GENERATED_CPP_BY_NAME(AssignDebugVisProcessor)
 UAssignDebugVisProcessor::UAssignDebugVisProcessor()
+#if WITH_MASSGAMEPLAY_DEBUG
 	: EntityQuery(*this)
+#endif // WITH_MASSGAMEPLAY_DEBUG
 {
 	bAutoRegisterWithProcessingPhases = false;
+
+#if WITH_MASSGAMEPLAY_DEBUG
 	bRequiresGameThreadExecution = true; // due to UMassDebuggerSubsystem
 	ObservedType = FSimDebugVisFragment::StaticStruct();
 	ObservedOperations = EMassObservedOperationFlags::Add;
+#endif // WITH_MASSGAMEPLAY_DEBUG
 }
 
+#if WITH_MASSGAMEPLAY_DEBUG
 void UAssignDebugVisProcessor::ConfigureQueries(const TSharedRef<FMassEntityManager>& EntityManager)
 {
 	EntityQuery.AddRequirement<FSimDebugVisFragment>(EMassFragmentAccess::ReadWrite);
@@ -65,3 +73,4 @@ void UAssignDebugVisProcessor::Execute(FMassEntityManager& EntityManager, FMassE
 	}
 #endif // WITH_EDITORONLY_DATA
 }
+#endif // WITH_MASSGAMEPLAY_DEBUG
diff --git a/Engine/Plugins/Runtime/MassGameplay/Source/MassGameplayDebug/Private/DebugVisLocationProcessor.cpp b/Engine/Plugins/Runtime/MassGameplay/Source/MassGameplayDebug/Private/DebugVisLocationProcessor.cpp
--- a/Engine/Plugins/Runtime/MassGameplay/Source/MassGameplayDebug/Private/DebugVisLocationProcessor.cpp
+++ b/Engine/Plugins/Runtime/MassGameplay/Source/MassGameplayDebug/Private/DebugVisLocationProcessor.cpp
@@ -1,6 +1,9 @@
 // Copyright Epic Games, Inc. All Rights Reserved.
 
 #include "DebugVisLocationProcessor.h"
+#include UE_INLINE_GENERATED_CPP_BY_NAME(DebugVisLocationProcessor)
+
+#if WITH_MASSGAMEPLAY_DEBUG
 #include "MassDebuggerSubsystem.h"
 #include "MassDebugVisualizationComponent.h"
 #include "MassCommonFragments.h"
@@ -10,9 +13,6 @@
 #include "MassMovementFragments.h"
 #include "MassDebugDrawHelpers.h"
 
-#include UE_INLINE_GENERATED_CPP_BY_NAME(DebugVisLocationProcessor)
-
-
 namespace UE::Mass::Debug
 {
 	bool bDebugDrawAllEntities = false;
@@ -26,19 +26,27 @@ namespace UE::Mass::Debug
 		};
 	}
 }
+#endif // WITH_MASSGAMEPLAY_DEBUG
 
 //----------------------------------------------------------------------//
 // UDebugVisLocationProcessor
 //----------------------------------------------------------------------//
 UDebugVisLocationProcessor::UDebugVisLocationProcessor()
+#if WITH_MASSGAMEPLAY_DEBUG
 	: EntityQuery(*this)
 	, AllLocationEntitiesQuery(*this)
+#endif // WITH_MASSGAMEPLAY_DEBUG
 {
+#if WITH_MASSGAMEPLAY_DEBUG
 	ExecutionOrder.ExecuteAfter.Add(UE::Mass::ProcessorGroupNames::SyncWorldToMass);
 	bRequiresGameThreadExecution = true; // due to UMassDebuggerSubsystem access
 	QueryBasedPruning = EMassQueryBasedPruning::Never;
+#else
+	bAutoRegisterWithProcessingPhases = false;
+#endif // WITH_MASSGAMEPLAY_DEBUG
 }
 
+#if WITH_MASSGAMEPLAY_DEBUG
 void UDebugVisLocationProcessor::ConfigureQueries(const TSharedRef<FMassEntityManager>& EntityManager)
 {
 	EntityQuery.AddRequirement<FTransformFragment>(EMassFragmentAccess::ReadOnly);
@@ -101,12 +109,7 @@ void UDebugVisLocationProcessor::Execute(FMassEntityManager& EntityManager, FMas
 
 			AllLocationEntitiesQuery.ForEachEntityChunk(Context, [&LineBatcher](const FMassExecutionContext& Context)
 				{
-					const FColor ArchetypeColor =
-#if WITH_MASSENTITY_DEBUG
-						Context.DebugGetArchetypeColor();
-#else
-						FColor::Green;
-#endif // WITH_MASSENTITY_DEBUG
+					const FColor ArchetypeColor = Context.DebugGetArchetypeColor();
 
 					const FVector DebugDrawBoxExtent(20.);
 					constexpr float ArrowLength = 30.f;
@@ -126,17 +129,25 @@ void UDebugVisLocationProcessor::Execute(FMassEntityManager& EntityManager, FMas
 		}
 	}
 }
+#endif // WITH_MASSGAMEPLAY_DEBUG
 
 //----------------------------------------------------------------------//
 //  UMassProcessor_UpdateDebugVis
 //----------------------------------------------------------------------//
 UMassProcessor_UpdateDebugVis::UMassProcessor_UpdateDebugVis()
+#if WITH_MASSGAMEPLAY_DEBUG
 	: EntityQuery(*this)
+#endif // WITH_MASSGAMEPLAY_DEBUG
 {
+#if WITH_MASSGAMEPLAY_DEBUG
 	ExecutionOrder.ExecuteAfter.Add(UE::Mass::ProcessorGroupNames::UpdateWorldFromMass);
 	bRequiresGameThreadExecution = true; // due to UMassDebuggerSubsystem access
+#else
+	bAutoRegisterWithProcessingPhases = false;
+#endif // WITH_MASSGAMEPLAY_DEBUG
 }
 
+#if WITH_MASSGAMEPLAY_DEBUG
 void UMassProcessor_UpdateDebugVis::ConfigureQueries(const TSharedRef<FMassEntityManager>& EntityManager) 
 {
 	EntityQuery.AddRequirement<FTransformFragment>(EMassFragmentAccess::ReadOnly);
@@ -166,3 +177,4 @@ void UMassProcessor_UpdateDebugVis::Execute(FMassEntityManager& EntityManager, F
 			}
 		});
 }
+#endif // WITH_MASSGAMEPLAY_DEBUG
diff --git a/Engine/Plugins/Runtime/MassGameplay/Source/MassGameplayDebug/Private/MassDebugVisualizationTrait.cpp b/Engine/Plugins/Runtime/MassGameplay/Source/MassGameplayDebug/Private/MassDebugVisualizationTrait.cpp
--- a/Engine/Plugins/Runtime/MassGameplay/Source/MassGameplayDebug/Private/MassDebugVisualizationTrait.cpp
+++ b/Engine/Plugins/Runtime/MassGameplay/Source/MassGameplayDebug/Private/MassDebugVisualizationTrait.cpp
@@ -1,25 +1,23 @@
 // Copyright Epic Games, Inc. All Rights Reserved.
 
 #include "MassDebugVisualizationTrait.h"
+#include UE_INLINE_GENERATED_CPP_BY_NAME(MassDebugVisualizationTrait)
+
+#if WITH_MASSGAMEPLAY_DEBUG
 #include "MassDebuggerSubsystem.h"
 #include "MassDebugVisualizationComponent.h"
 #include "MassEntityTemplateRegistry.h"
 #include "MassCommonFragments.h"
 #include "Engine/World.h"
-
-#include UE_INLINE_GENERATED_CPP_BY_NAME(MassDebugVisualizationTrait)
+#endif // WITH_MASSGAMEPLAY_DEBUG
 
 void UMassDebugVisualizationTrait::BuildTemplate(FMassEntityTemplateBuildContext& BuildContext, const UWorld& World) const
 {
+#if WITH_MASSGAMEPLAY_DEBUG
 #if WITH_EDITORONLY_DATA
 	const UStaticMesh* const DebugMesh = DebugShape.Mesh;
-#else
-	const UStaticMesh* const DebugMesh = nullptr;
-#endif
-	
 	if (DebugMesh || BuildContext.IsInspectingData())
 	{
-#if WITH_EDITORONLY_DATA
 		FSimDebugVisFragment& DebugVisFragment = BuildContext.AddFragment_GetRef<FSimDebugVisFragment>();
 		if (!BuildContext.IsInspectingData())
 		{
@@ -34,20 +32,20 @@ void UMassDebugVisualizationTrait::BuildTemplate(FMassEntityTemplateBuildContext
 				// @todo this path requires a fragment destructor that will remove the mesh from the debugger.
 			}
 		}
-#endif // WITH_EDITORONLY_DATA
 	}
+#endif // WITH_EDITORONLY_DATA
+
 	// add fragments needed whenever we have debugging capabilities
-#if !(UE_BUILD_SHIPPING || UE_BUILD_TEST)
 	BuildContext.AddTag<FMassDebuggableTag>();
+
 #if WITH_EDITORONLY_DATA
 	BuildContext.AddFragment_GetRef<FDataFragment_DebugVis>().Shape = DebugShape.WireShape;
 #else
 	// DebugShape unavailable, will used default instead
 	BuildContext.AddFragment<FDataFragment_DebugVis>();
 #endif // WITH_EDITORONLY_DATA
-	BuildContext.AddFragment<FAgentRadiusFragment>();
 
+	BuildContext.AddFragment<FAgentRadiusFragment>();
 	BuildContext.AddFragment<FTransformFragment>();
-#endif // if !(UE_BUILD_SHIPPING || UE_BUILD_TEST)
-
+#endif // WITH_MASSGAMEPLAY_DEBUG
 }
diff --git a/Engine/Plugins/Runtime/MassGameplay/Source/MassGameplayDebug/Private/MassDebuggerSubsystem.cpp b/Engine/Plugins/Runtime/MassGameplay/Source/MassGameplayDebug/Private/MassDebuggerSubsystem.cpp
--- a/Engine/Plugins/Runtime/MassGameplay/Source/MassGameplayDebug/Private/MassDebuggerSubsystem.cpp
+++ b/Engine/Plugins/Runtime/MassGameplay/Source/MassGameplayDebug/Private/MassDebuggerSubsystem.cpp
@@ -2,41 +2,31 @@
 
 #include "MassDebuggerSubsystem.h"
 #include "EngineUtils.h"
-#include "MassCommonTypes.h"
-#include "MassSimulationSubsystem.h"
+#include "MassDebugger.h"
 #include "MassDebugVisualizationComponent.h"
-#include "MassEntityManager.h"
 #include "MassDebugVisualizer.h"
-#include "MassDebugger.h"
+#include "MassEntityManager.h"
+#include "MassSimulationSubsystem.h"
 
 #include UE_INLINE_GENERATED_CPP_BY_NAME(MassDebuggerSubsystem)
 
-
+#if WITH_MASSGAMEPLAY_DEBUG
 void UMassDebuggerSubsystem::Initialize(FSubsystemCollectionBase& Collection)
 {
 	Collection.InitializeDependency(UMassSimulationSubsystem::StaticClass());
 	Super::Initialize(Collection);
 
-	UMassSimulationSubsystem* SimSystem = UWorld::GetSubsystem<UMassSimulationSubsystem>(GetWorld());
-	check(SimSystem);
-	SimSystem->GetOnProcessingPhaseStarted(EMassProcessingPhase::PrePhysics).AddUObject(this, &UMassDebuggerSubsystem::OnProcessingPhaseStarted);
-
-#if WITH_MASSENTITY_DEBUG
 	OnEntitySelectedHandle = FMassDebugger::OnEntitySelectedDelegate.AddUObject(this, &UMassDebuggerSubsystem::OnEntitySelected);
-#endif // WITH_MASSENTITY_DEBUG
 
 	OverrideSubsystemTraits<UMassDebuggerSubsystem>(Collection);
 }
 
 void UMassDebuggerSubsystem::Deinitialize()
 {
-#if WITH_MASSENTITY_DEBUG
 	FMassDebugger::OnEntitySelectedDelegate.Remove(OnEntitySelectedHandle);
-#endif // WITH_MASSENTITY_DEBUG
 	Super::Deinitialize();
 }
 
-#if WITH_MASSENTITY_DEBUG
 void UMassDebuggerSubsystem::OnEntitySelected(const FMassEntityManager& EntityManager, const FMassEntityHandle EntityHandle)
 {
 	if (EntityManager.GetWorld() == GetWorld())
@@ -44,36 +34,40 @@ void UMassDebuggerSubsystem::OnEntitySelected(const FMassEntityManager& EntityMa
 		SetSelectedEntity(EntityHandle);
 	}
 }
-#endif // WITH_MASSENTITY_DEBUG
 
-void UMassDebuggerSubsystem::ResetDebugShapes()
+void UMassDebuggerSubsystem::ForEachShape(EMassEntityDebugShape Shape, const TFunctionRef<void(const FShapeDesc&)> Function) const
 {
-	// get ready to receive new debug info
-	for (TArray<FShapeDesc>& Array : Shapes)
+	UE_MT_SCOPED_READ_ACCESS(MTDetector);
+
+	for (const FShapeDesc& Desc : Shapes[static_cast<uint8>(Shape)])
 	{
-		Array.Reset();
+		Function(Desc);
 	}
 }
 
-void UMassDebuggerSubsystem::PreTickProcessors()
+void UMassDebuggerSubsystem::ResetDebugShapes()
 {
-	Entities.Reset();
-	Locations.Reset();
-}
+	UE_MT_SCOPED_WRITE_ACCESS(MTDetector);
 
-void UMassDebuggerSubsystem::OnProcessingPhaseStarted(const float DeltaSeconds)
-{
-	PreTickProcessors();
+	// get ready to receive new debug info
+	for (TArray<FShapeDesc>& Array : Shapes)
+	{
+		Array.Reset();
+	}
 }
 
 void UMassDebuggerSubsystem::SetSelectedEntity(const FMassEntityHandle InSelectedEntity)
 {
+	UE_MT_SCOPED_WRITE_ACCESS(MTDetector);
+
 	SelectedEntity = InSelectedEntity;
 	SelectedEntityDetails.Empty();
 }
 
 void UMassDebuggerSubsystem::AppendSelectedEntityInfo(const FString& Info)
 {
+	UE_MT_SCOPED_WRITE_ACCESS(MTDetector);
+
 	if (UpdateFrameNumber != GFrameNumber)
 	{
 		SelectedEntityDetails.Empty();
@@ -84,11 +78,11 @@ void UMassDebuggerSubsystem::AppendSelectedEntityInfo(const FString& Info)
 }
 
 UMassDebugVisualizationComponent* UMassDebuggerSubsystem::GetVisualizationComponent()
-{ 
+{
 #if WITH_EDITORONLY_DATA
 	if (VisualizationComponent == nullptr)
 	{
-		if (!ensureMsgf(DebugVisualizer == nullptr, TEXT("If we do have a DebugVisualizer but don't have VisualizationComponent then somethin's wrong")))
+		if (!ensureMsgf(DebugVisualizer == nullptr, TEXT("If we do have a DebugVisualizer but don't have VisualizationComponent then something's wrong")))
 		{
 			VisualizationComponent = &DebugVisualizer->GetDebugVisComponent();
 		}
@@ -101,7 +95,7 @@ UMassDebugVisualizationComponent* UMassDebuggerSubsystem::GetVisualizationCompon
 			}
 		}
 	}
-	ensureMsgf(VisualizationComponent,  TEXT("In editor builds we always expect to have a visualizer component available"));
+	ensureMsgf(VisualizationComponent, TEXT("In editor builds we always expect to have a visualizer component available"));
 #endif // WITH_EDITORONLY_DATA
 	return VisualizationComponent;
 }
@@ -131,4 +125,5 @@ AMassDebugVisualizer& UMassDebuggerSubsystem::GetOrSpawnDebugVisualizer(UWorld&
 
 	return *DebugVisualizer;
 }
-#endif
+#endif // WITH_EDITORONLY_DATA
+#endif //WITH_MASSGAMEPLAY_DEBUG
\ No newline at end of file
diff --git a/Engine/Plugins/Runtime/MassGameplay/Source/MassGameplayDebug/Public/AssignDebugVisProcessor.h b/Engine/Plugins/Runtime/MassGameplay/Source/MassGameplayDebug/Public/AssignDebugVisProcessor.h
--- a/Engine/Plugins/Runtime/MassGameplay/Source/MassGameplayDebug/Public/AssignDebugVisProcessor.h
+++ b/Engine/Plugins/Runtime/MassGameplay/Source/MassGameplayDebug/Public/AssignDebugVisProcessor.h
@@ -2,8 +2,9 @@
 
 #pragma once
 
-#include "MassObserverProcessor.h"
+#include "MassCommonTypes.h"
 #include "MassEntityQuery.h"
+#include "MassObserverProcessor.h"
 #include "AssignDebugVisProcessor.generated.h"
 
 #define UE_API MASSGAMEPLAYDEBUG_API
@@ -18,11 +19,14 @@ class UAssignDebugVisProcessor : public UMassObserverProcessor
 	GENERATED_BODY()
 public:
 	UE_API UAssignDebugVisProcessor();
+
+#if WITH_MASSGAMEPLAY_DEBUG
 	UE_API virtual void ConfigureQueries(const TSharedRef<FMassEntityManager>& EntityManager) override;
 	UE_API virtual void Execute(FMassEntityManager& EntityManager, FMassExecutionContext& Context) override;
 
 protected:
 	FMassEntityQuery EntityQuery;
+#endif // WITH_MASSGAMEPLAY_DEBUG
 };
 
 #undef UE_API
diff --git a/Engine/Plugins/Runtime/MassGameplay/Source/MassGameplayDebug/Public/DebugVisLocationProcessor.h b/Engine/Plugins/Runtime/MassGameplay/Source/MassGameplayDebug/Public/DebugVisLocationProcessor.h
--- a/Engine/Plugins/Runtime/MassGameplay/Source/MassGameplayDebug/Public/DebugVisLocationProcessor.h
+++ b/Engine/Plugins/Runtime/MassGameplay/Source/MassGameplayDebug/Public/DebugVisLocationProcessor.h
@@ -2,6 +2,7 @@
 
 #pragma once
 
+#include "MassCommonTypes.h"
 #include "MassProcessor.h"
 #include "MassEntityQuery.h"
 #include "DebugVisLocationProcessor.generated.h"
@@ -18,31 +19,34 @@ class UDebugVisLocationProcessor : public UMassProcessor
 
 public:
 	UE_API UDebugVisLocationProcessor();
-	
+
+#if WITH_MASSGAMEPLAY_DEBUG
 protected:
 	UE_API virtual void ConfigureQueries(const TSharedRef<FMassEntityManager>& EntityManager) override;
 	UE_API virtual void Execute(FMassEntityManager& EntityManager, FMassExecutionContext& Context) override;
 
 	FMassEntityQuery EntityQuery;
 	FMassEntityQuery AllLocationEntitiesQuery;
+#endif // WITH_MASSGAMEPLAY_DEBUG
 };
 
-//----------------------------------------------------------------------//
-// new one 
-//----------------------------------------------------------------------//
-//class UMassDebugger;
-
+/**
+ * Processor gathering debug data from entities to push debug shapes to the MassDebugger subsystem
+ */
 UCLASS(MinimalAPI)
 class UMassProcessor_UpdateDebugVis : public UMassProcessor
 {
 	GENERATED_BODY()
 public:
 	UE_API UMassProcessor_UpdateDebugVis();
+
+#if WITH_MASSGAMEPLAY_DEBUG
 	UE_API virtual void ConfigureQueries(const TSharedRef<FMassEntityManager>& EntityManager) override;
 	UE_API virtual void Execute(FMassEntityManager& EntityManager, FMassExecutionContext& Context) override;
 
 protected:
 	FMassEntityQuery EntityQuery;
+#endif // WITH_MASSGAMEPLAY_DEBUG
 };
 
 #undef UE_API
diff --git a/Engine/Plugins/Runtime/MassGameplay/Source/MassGameplayDebug/Public/MassDebuggerSubsystem.h b/Engine/Plugins/Runtime/MassGameplay/Source/MassGameplayDebug/Public/MassDebuggerSubsystem.h
--- a/Engine/Plugins/Runtime/MassGameplay/Source/MassGameplayDebug/Public/MassDebuggerSubsystem.h
+++ b/Engine/Plugins/Runtime/MassGameplay/Source/MassGameplayDebug/Public/MassDebuggerSubsystem.h
@@ -5,6 +5,7 @@
 #include "MassGameplayDebugTypes.h"
 #include "MassSubsystemBase.h"
 #include "MassExternalSubsystemTraits.h"
+#include "Misc/MTTransactionallySafeAccessDetector.h"
 #include "MassDebuggerSubsystem.generated.h"
 
 #define UE_API MASSGAMEPLAYDEBUG_API
@@ -17,29 +18,66 @@ UCLASS(MinimalAPI)
 class UMassDebuggerSubsystem : public UMassSubsystemBase
 {
 	GENERATED_BODY()
+
 public:
+
+#if WITH_MASSGAMEPLAY_DEBUG
+
 	struct FShapeDesc
 	{
 		FVector Location = {}; // no init on purpose, value will come from constructor
 		float Size = {};
-		FShapeDesc(const FVector InLocation, const float InSize) : Location(InLocation), Size(InSize) {}
+		FShapeDesc(const FVector& InLocation, const float InSize) : Location(InLocation), Size(InSize)
+		{
+		}
 	};
 
 	// Methods to optimize the collection of data to only when category is enabled
-	bool IsCollectingData() const { return bCollectingData; }
-	void SetCollectingData() { bCollectingData = true; }
-	void DataCollected() { bCollectingData = false; }
+	bool IsCollectingData() const
+	{
+		return bCollectingData;
+	}
+
+	void SetCollectingData()
+	{
+		bCollectingData = true;
+	}
+
+	void DataCollected()
+	{
+		bCollectingData = false;
+	}
+
+	void AddShape(EMassEntityDebugShape Shape, const FVector& Location, const float Size)
+	{
+		UE_MT_SCOPED_WRITE_ACCESS(MTDetector);
+		Shapes[static_cast<uint8>(Shape)].Add(FShapeDesc(Location, Size));
+	}
+
+	UE_API void ForEachShape(EMassEntityDebugShape Shape, const TFunctionRef<void(const FShapeDesc&)> Function) const;
+
+	UE_DEPRECATED(5.8, "Use ForEachShape instead.")
+	const TArray<FShapeDesc>* GetShapes() const
+	{
+		return Shapes;
+	}
 
-	void AddShape(EMassEntityDebugShape Shape, FVector Location, float Size) { Shapes[uint8(Shape)].Add(FShapeDesc(Location, Size)); }
-	const TArray<FShapeDesc>* GetShapes() const { return Shapes; }
 	UE_API void ResetDebugShapes();
 
-	FMassEntityHandle GetSelectedEntity() const { return SelectedEntity; }
 	UE_API void SetSelectedEntity(const FMassEntityHandle InSelectedEntity);
+	FMassEntityHandle GetSelectedEntity() const
+	{
+		UE_MT_SCOPED_READ_ACCESS(MTDetector);
+		return SelectedEntity;
+	}
 
 	UE_API void AppendSelectedEntityInfo(const FString& Info);
-	const FString& GetSelectedEntityInfo() const { return SelectedEntityDetails; }
-	
+	const FString& GetSelectedEntityInfo() const
+	{
+		UE_MT_SCOPED_READ_ACCESS(MTDetector);
+		return SelectedEntityDetails;
+	}
+
 	/** Fetches the UMassDebugVisualizationComponent owned by lazily created DebugVisualizer */
 	UE_API UMassDebugVisualizationComponent* GetVisualizationComponent();
 
@@ -48,32 +86,32 @@ public:
 #endif
 
 protected:
-	// USubsystem BEGIN
+	//~ USubsystem BEGIN
 	UE_API virtual void Initialize(FSubsystemCollectionBase& Collection) override;
 	UE_API virtual void Deinitialize() override;
-	// USubsystem END
-	
-	UE_API void OnProcessingPhaseStarted(const float DeltaSeconds);
-	UE_API void PreTickProcessors();
+	//~ USubsystem END
+
 	UE_API void OnEntitySelected(const FMassEntityManager& EntityManager, const FMassEntityHandle EntityHandle);
 
-protected:
-	bool bCollectingData = false;
+	TArray<FShapeDesc> Shapes[static_cast<uint8>(EMassEntityDebugShape::MAX)];
 
-	TArray<FShapeDesc> Shapes[uint8(EMassEntityDebugShape::MAX)];
-	TArray<FMassEntityHandle> Entities;
-	TArray<FVector> Locations;
 	FMassEntityHandle SelectedEntity;
 	FString SelectedEntityDetails;
+
 	uint64 UpdateFrameNumber = INDEX_NONE;
 
+	FDelegateHandle OnEntitySelectedHandle;
+
+	UE_MT_DECLARE_TS_RW_ACCESS_DETECTOR(MTDetector);
+
+	bool bCollectingData = false;
+#endif // WITH_MASSGAMEPLAY_DEBUG
+
 	UPROPERTY(Transient)
 	TObjectPtr<UMassDebugVisualizationComponent> VisualizationComponent;
 
 	UPROPERTY(Transient)
 	TObjectPtr<AMassDebugVisualizer> DebugVisualizer;
-
-	FDelegateHandle OnEntitySelectedHandle;
 };
 
 template<>
