Mass fix FNames in headers

diff --git a/Engine/Plugins/Runtime/MassGameplay/Source/MassCommon/Private/MassCommonTypes.cpp b/Engine/Plugins/Runtime/MassGameplay/Source/MassCommon/Private/MassCommonTypes.cpp
--- a/Engine/Plugins/Runtime/MassGameplay/Source/MassCommon/Private/MassCommonTypes.cpp
+++ b/Engine/Plugins/Runtime/MassGameplay/Source/MassCommon/Private/MassCommonTypes.cpp
@@ -3,3 +3,14 @@
 #include "MassCommonTypes.h"
 
 #include UE_INLINE_GENERATED_CPP_BY_NAME(MassCommonTypes)
+
+namespace UE::Mass::ProcessorGroupNames
+{
+	const FName UpdateWorldFromMass = TEXT("UpdateWorldFromMass");
+	const FName SyncWorldToMass = TEXT("SyncWorldToMass");
+	const FName Behavior = TEXT("Behavior");
+	const FName Tasks = TEXT("Tasks");
+	const FName Avoidance = TEXT("Avoidance");
+	const FName ApplyForces = TEXT("ApplyForces");
+	const FName Movement = TEXT("Movement");
+}
diff --git a/Engine/Plugins/Runtime/MassGameplay/Source/MassCommon/Public/MassCommonTypes.h b/Engine/Plugins/Runtime/MassGameplay/Source/MassCommon/Public/MassCommonTypes.h
--- a/Engine/Plugins/Runtime/MassGameplay/Source/MassCommon/Public/MassCommonTypes.h
+++ b/Engine/Plugins/Runtime/MassGameplay/Source/MassCommon/Public/MassCommonTypes.h
@@ -17,13 +17,13 @@ class UMaterialInterface;
 
 namespace UE::Mass::ProcessorGroupNames
 {
-	const FName UpdateWorldFromMass = FName(TEXT("UpdateWorldFromMass"));
-	const FName SyncWorldToMass = FName(TEXT("SyncWorldToMass"));
-	const FName Behavior = FName(TEXT("Behavior"));
-	const FName Tasks = FName(TEXT("Tasks"));
-	const FName Avoidance = FName(TEXT("Avoidance"));
-	const FName ApplyForces = FName(TEXT("ApplyForces"));
-	const FName Movement = FName(TEXT("Movement"));
+	extern MASSCOMMON_API const FName UpdateWorldFromMass;
+	extern MASSCOMMON_API const FName SyncWorldToMass;
+	extern MASSCOMMON_API const FName Behavior;
+	extern MASSCOMMON_API const FName Tasks;
+	extern MASSCOMMON_API const FName Avoidance;
+	extern MASSCOMMON_API const FName ApplyForces;
+	extern MASSCOMMON_API const FName Movement;
 }
 
 USTRUCT()
diff --git a/Engine/Plugins/Runtime/MassGameplay/Source/MassLOD/Private/MassLODSubsystem.cpp b/Engine/Plugins/Runtime/MassGameplay/Source/MassLOD/Private/MassLODSubsystem.cpp
--- a/Engine/Plugins/Runtime/MassGameplay/Source/MassLOD/Private/MassLODSubsystem.cpp
+++ b/Engine/Plugins/Runtime/MassGameplay/Source/MassLOD/Private/MassLODSubsystem.cpp
@@ -18,6 +18,12 @@
 
 #include UE_INLINE_GENERATED_CPP_BY_NAME(MassLODSubsystem)
 
+namespace UE::Mass::ProcessorGroupNames
+{
+	const FName LODCollector = TEXT("LODCollector");
+	const FName LOD = TEXT("LOD");
+}
+
 namespace UE::MassLOD
 {
 	FColor LODColors[] =
diff --git a/Engine/Plugins/Runtime/MassGameplay/Source/MassLOD/Public/MassLODTypes.h b/Engine/Plugins/Runtime/MassGameplay/Source/MassLOD/Public/MassLODTypes.h
--- a/Engine/Plugins/Runtime/MassGameplay/Source/MassLOD/Public/MassLODTypes.h
+++ b/Engine/Plugins/Runtime/MassGameplay/Source/MassLOD/Public/MassLODTypes.h
@@ -27,8 +27,9 @@ namespace UE::MassLOD
 
 namespace UE::Mass::ProcessorGroupNames
 {
-	const FName LODCollector = FName(TEXT("LODCollector"));
-	const FName LOD = FName(TEXT("LOD"));
+	// Defined in MassLODSubsystem.cpp
+	extern MASSLOD_API const FName LODCollector;
+	extern MASSLOD_API const FName LOD;
 }
 
 // We are not using enum class here because we are doing so many arithmetic operation and comparison on them 
@@ -54,4 +55,4 @@ enum class EMassVisibility : uint8
 	CulledByFrustum, // Not in camera frustum but within visibility distance
 	CulledByDistance, // Too far whether in or out of frustum
 	Max
-};
\ No newline at end of file
+};
