Mass Replication improvements/fixes:
- Added NotReplicated attribute to properties in FMassClientBubbleSerializerBase to make sure they're not replicated and don't break Iris FastArray replication
- Removed friend class requirement in TClientBubbleHandlerBase for derived classes by exposing a few methods (since the friend solution is not scalable to outside/third party uses)
- Adapted TMassClientBubblePathHandler and TMassClientBubbleTransformHandler to use the new methods

#rnx

diff --git a/Engine/Plugins/AI/MassAI/Source/MassAIReplication/Public/MassReplicationPathHandlers.h b/Engine/Plugins/AI/MassAI/Source/MassAIReplication/Public/MassReplicationPathHandlers.h
--- a/Engine/Plugins/AI/MassAI/Source/MassAIReplication/Public/MassReplicationPathHandlers.h
+++ b/Engine/Plugins/AI/MassAI/Source/MassAIReplication/Public/MassReplicationPathHandlers.h
@@ -120,10 +120,7 @@ void TMassClientBubblePathHandler<AgentArrayItem>::SetBubblePathData(const FMass
 	const FMassMoveTargetFragment& MoveTargetFragment,
 	const FMassZoneGraphLaneLocationFragment& LaneLocationFragment)
 {
-	check(OwnerHandler.AgentHandleManager.IsValidHandle(Handle));
-
-	const int32 AgentsIdx = OwnerHandler.AgentLookupArray[Handle.GetIndex()].AgentsIdx;
-	AgentArrayItem& Item = (*OwnerHandler.Agents)[AgentsIdx];
+	AgentArrayItem& Item = OwnerHandler.GetAgentItem(Handle);
 
 	checkf(Item.Agent.GetNetID().IsValid(), TEXT("Pos should not be updated on FCrowdFastArrayItem's that have an Invalid ID! First Add the Agent!"));
 
@@ -134,7 +131,7 @@ void TMassClientBubblePathHandler<AgentArrayItem>::SetBubblePathData(const FMass
 	{
 		ReplicatedPath = FReplicatedAgentPathData(PathRequestFragment, MoveTargetFragment, LaneLocationFragment);
 
-		OwnerHandler.Serializer->MarkItemDirty(Item);
+		OwnerHandler.MarkItemDirty(Item);
 	}
 }
 #endif //UE_REPLICATION_COMPILE_SERVER_CODE
@@ -173,7 +170,7 @@ void TMassClientBubblePathHandler<AgentArrayItem>::ClearFragmentViewsForSpawnQue
 template<typename AgentArrayItem>
 void TMassClientBubblePathHandler<AgentArrayItem>::SetSpawnedEntityData(const FMassEntityView& EntityView, const FReplicatedAgentPathData& ReplicatedPathData, const int32 EntityIdx) const
 {
-	UWorld* World = OwnerHandler.Serializer->GetWorld();
+	UWorld* World = OwnerHandler.GetWorld();
 	check(World);
 	ReplicatedPathData.InitEntity(*World, EntityView, LaneLocationList[EntityIdx], MoveTargetList[EntityIdx], PathRequestList[EntityIdx]);
 }
@@ -183,7 +180,7 @@ void TMassClientBubblePathHandler<AgentArrayItem>::SetSpawnedEntityData(const FM
 template<typename AgentArrayItem>
 void TMassClientBubblePathHandler<AgentArrayItem>::SetModifiedEntityData(const FMassEntityView& EntityView, const FReplicatedAgentPathData& ReplicatedPathData) const
 {
-	UWorld* World = OwnerHandler.Serializer->GetWorld();
+	UWorld* World = OwnerHandler.GetWorld();
 	check(World);
 	ReplicatedPathData.ApplyToEntity(*World, EntityView);
 }
diff --git a/Engine/Plugins/Runtime/MassGameplay/Source/MassReplication/Public/MassClientBubbleHandler.h b/Engine/Plugins/Runtime/MassGameplay/Source/MassReplication/Public/MassClientBubbleHandler.h
--- a/Engine/Plugins/Runtime/MassGameplay/Source/MassReplication/Public/MassClientBubbleHandler.h
+++ b/Engine/Plugins/Runtime/MassGameplay/Source/MassReplication/Public/MassClientBubbleHandler.h
@@ -129,12 +129,6 @@ template<typename AgentArrayItem>
 class TClientBubbleHandlerBase : public IClientBubbleHandlerInterface
 {
 public:
-	template<typename T>
-	friend class TMassClientBubblePathHandler;
-
-	template<typename T>
-	friend class TMassClientBubbleTransformHandler;
-
 	typedef TFunctionRef<void(FMassEntityQuery&)> FAddRequirementsForSpawnQueryFunction;
 	typedef TFunctionRef<void(FMassExecutionContext&)> FCacheFragmentViewsForSpawnQueryFunction;
 	typedef TFunctionRef<void(const FMassEntityView&, const typename AgentArrayItem::FReplicatedAgentType&, const int32)> FSetSpawnedEntityDataFunction;
@@ -144,6 +138,7 @@ public:
 	virtual void Initialize(TArray<AgentArrayItem>& InAgents, FMassClientBubbleSerializerBase& InSerializer);
 
 	virtual void InitializeForWorld(UWorld& InWorld) override;
+	UWorld* GetWorld() const;
 
 #if UE_REPLICATION_COMPILE_SERVER_CODE
 	FMassReplicatedAgentHandle AddAgent(FMassEntityHandle Entity, typename AgentArrayItem::FReplicatedAgentType& Agent);
@@ -151,6 +146,8 @@ public:
 	bool RemoveAgent(FMassNetworkID NetID);
 	bool RemoveAgent(FMassReplicatedAgentHandle AgentHandle);
 	void RemoveAgentChecked(FMassReplicatedAgentHandle AgentHandle);
+	void MarkItemDirty(FFastArraySerializerItem& Item) const;
+	AgentArrayItem& GetAgentItem(const FMassReplicatedAgentHandle Handle) const;
 
 	/** Gets an agent safely */
 	const typename AgentArrayItem::FReplicatedAgentType* GetAgent(FMassReplicatedAgentHandle Handle) const;
@@ -238,6 +235,12 @@ void TClientBubbleHandlerBase<AgentArrayItem>::InitializeForWorld(UWorld& InWorl
 	Serializer->InitializeForWorld(InWorld);
 }
 
+template <typename AgentArrayItem>
+UWorld* TClientBubbleHandlerBase<AgentArrayItem>::GetWorld() const
+{
+	return Serializer->GetWorld();
+}
+
 #if UE_REPLICATION_COMPILE_SERVER_CODE
 template<typename AgentArrayItem>
 FMassReplicatedAgentHandle TClientBubbleHandlerBase<AgentArrayItem>::AddAgent(FMassEntityHandle Entity, typename AgentArrayItem::FReplicatedAgentType& Agent)
@@ -299,6 +302,12 @@ void TClientBubbleHandlerBase<AgentArrayItem>::RemoveAgentChecked(FMassReplicate
 	RemoveAgentImpl(Handle);
 }
 
+template <typename AgentArrayItem>
+void TClientBubbleHandlerBase<AgentArrayItem>::MarkItemDirty(FFastArraySerializerItem& Item) const
+{
+	Serializer->MarkItemDirty(Item);
+}
+
 template<typename AgentArrayItem>
 void TClientBubbleHandlerBase<AgentArrayItem>::RemoveAgentImpl(FMassReplicatedAgentHandle Handle)
 {
@@ -327,6 +336,16 @@ void TClientBubbleHandlerBase<AgentArrayItem>::RemoveAgentImpl(FMassReplicatedAg
 	LookUpData.Invalidate();
 }
 
+template <typename AgentArrayItem>
+AgentArrayItem& TClientBubbleHandlerBase<AgentArrayItem>::GetAgentItem(const FMassReplicatedAgentHandle Handle) const
+{
+	check(AgentHandleManager.IsValidHandle(Handle));
+        
+	const int32 AgentsIdx = AgentLookupArray[Handle.GetIndex()].AgentsIdx;
+	
+	return (*Agents)[AgentsIdx];
+}
+
 template<typename AgentArrayItem>
 const typename AgentArrayItem::FReplicatedAgentType* TClientBubbleHandlerBase<AgentArrayItem>::GetAgent(FMassReplicatedAgentHandle Handle) const
 {
diff --git a/Engine/Plugins/Runtime/MassGameplay/Source/MassReplication/Public/MassClientBubbleSerializerBase.h b/Engine/Plugins/Runtime/MassGameplay/Source/MassReplication/Public/MassClientBubbleSerializerBase.h
--- a/Engine/Plugins/Runtime/MassGameplay/Source/MassReplication/Public/MassClientBubbleSerializerBase.h
+++ b/Engine/Plugins/Runtime/MassGameplay/Source/MassReplication/Public/MassClientBubbleSerializerBase.h
@@ -45,15 +45,15 @@ public:
 	UE_API void ShutDown();
 
 private:
-	UPROPERTY(Transient)
+	UPROPERTY(Transient, NotReplicated)
 	TObjectPtr<UWorld> World = nullptr;
 
-	UPROPERTY(Transient)
+	UPROPERTY(Transient, NotReplicated)
 	TObjectPtr<UMassSpawnerSubsystem> SpawnerSubsystem = nullptr;
 
 	TSharedPtr<FMassEntityManager> EntityManager;
 
-	UPROPERTY(Transient)
+	UPROPERTY(Transient, NotReplicated)
 	TObjectPtr<UMassReplicationSubsystem> ReplicationSubsystem = nullptr;
 
 	/** Pointer to the IClientBubbleHandlerInterface derived class in the class derived from this one */
diff --git a/Engine/Plugins/Runtime/MassGameplay/Source/MassReplication/Public/MassReplicationTransformHandlers.h b/Engine/Plugins/Runtime/MassGameplay/Source/MassReplication/Public/MassReplicationTransformHandlers.h
--- a/Engine/Plugins/Runtime/MassGameplay/Source/MassReplication/Public/MassReplicationTransformHandlers.h
+++ b/Engine/Plugins/Runtime/MassGameplay/Source/MassReplication/Public/MassReplicationTransformHandlers.h
@@ -107,12 +107,9 @@ protected:
 template<typename AgentArrayItem>
 void TMassClientBubbleTransformHandler<AgentArrayItem>::SetBubblePositionYawFromTransform(const FMassReplicatedAgentHandle Handle, const FTransform& Transform)
 {
-	check(OwnerHandler.AgentHandleManager.IsValidHandle(Handle));
-
-	const int32 AgentsIdx = OwnerHandler.AgentLookupArray[Handle.GetIndex()].AgentsIdx;
 	bool bMarkDirty = false;
-
-	AgentArrayItem& Item = (*OwnerHandler.Agents)[AgentsIdx];
+	
+	AgentArrayItem& Item = OwnerHandler.GetAgentItem(Handle);
 
 	checkf(Item.Agent.GetNetID().IsValid(), TEXT("Pos should not be updated on FCrowdFastArrayItem's that have an Invalid ID! First Add the Agent!"));
 
@@ -138,7 +135,7 @@ void TMassClientBubbleTransformHandler<AgentArrayItem>::SetBubblePositionYawFrom
 
 	if (bMarkDirty)
 	{
-		OwnerHandler.Serializer->MarkItemDirty(Item);
+		OwnerHandler.MarkItemDirty(Item);
 	}
 
 }
