[Mass] Minor fixes based on BugHawk's findings.

diff --git a/Engine/Plugins/Runtime/MassGameplay/Source/MassSpawner/Public/MassEntityTemplateRegistry.h b/Engine/Plugins/Runtime/MassGameplay/Source/MassSpawner/Public/MassEntityTemplateRegistry.h
--- a/Engine/Plugins/Runtime/MassGameplay/Source/MassSpawner/Public/MassEntityTemplateRegistry.h
+++ b/Engine/Plugins/Runtime/MassGameplay/Source/MassSpawner/Public/MassEntityTemplateRegistry.h
@@ -220,7 +220,7 @@ struct FMassEntityTemplateBuildContext
 	void AddTranslator()
 	{
 		TypeAdded(*T::StaticClass());
-		GetDefault<T>()->AppendRequiredTags(TemplateData.GetMutableTags());
+		TemplateData.GetMutableTags() += GetDefault<T>()->GetRequiredTags();
 	}
 
 	//----------------------------------------------------------------------//
diff --git a/Engine/Plugins/Runtime/MassGameplay/Source/MassSpawner/Public/MassTranslator.h b/Engine/Plugins/Runtime/MassGameplay/Source/MassSpawner/Public/MassTranslator.h
--- a/Engine/Plugins/Runtime/MassGameplay/Source/MassSpawner/Public/MassTranslator.h
+++ b/Engine/Plugins/Runtime/MassGameplay/Source/MassSpawner/Public/MassTranslator.h
@@ -38,15 +38,16 @@ public:
 	 * 
 	 *  @todo might want this function on the MassProcessor level. TBD
 	 * 
-	 *  @param InOutTags tag types will be appended to this array. No uniqueness checks are performed. */
+	 *  @param InOutTags - this translator's expected tags will be added to this parameter */
+	UE_DEPRECATED(5.8, "This function is deprecated, please fetch the tags directly with GetRequiredTags and append then wherever you want.")
 	void AppendRequiredTags(FMassTagBitSet& InOutTags) const 
 	{
 		InOutTags += RequiredTags; 
 	}
 
-	void AppendRequiredTags(FMassElementBitSet& InOutTags) const 
+	FMassTagBitSet GetRequiredTags() const 
 	{
-		InOutTags += RequiredTags; 
+		return RequiredTags; 
 	}
 
 protected:
diff --git a/Engine/Source/Runtime/MassEntity/Private/MassEntityTypes.cpp b/Engine/Source/Runtime/MassEntity/Private/MassEntityTypes.cpp
--- a/Engine/Source/Runtime/MassEntity/Private/MassEntityTypes.cpp
+++ b/Engine/Source/Runtime/MassEntity/Private/MassEntityTypes.cpp
@@ -80,8 +80,10 @@ FMassArchetypeSharedFragmentValues FMassArchetypeSharedFragmentValues::CreateCom
 			return SharedStruct.GetScriptStruct() == nullptr
 				|| NewSharedElementComposition.Contains(*SharedStruct.GetScriptStruct()) == false;
 		};
+
 	NewSharedValues.SharedFragments.RemoveAllSwap(RemovePredicate);
 	NewSharedValues.ConstSharedFragments.RemoveAllSwap(RemovePredicate);
+	NewSharedValues.StoredElementsBitSet = NewSharedElementComposition;
 
 	NewSharedValues.Sort();
 	return NewSharedValues;
diff --git a/Engine/Source/Runtime/MassEntity/Private/MassObserverManager.cpp b/Engine/Source/Runtime/MassEntity/Private/MassObserverManager.cpp
--- a/Engine/Source/Runtime/MassEntity/Private/MassObserverManager.cpp
+++ b/Engine/Source/Runtime/MassEntity/Private/MassObserverManager.cpp
@@ -256,9 +256,10 @@ void FMassObserverManager::DeInitialize()
 		ModulesUnloadedHandle.Reset();
 	}
 
-	for (int32 i = 0; i < (int32)EMassObservedOperation::MAX; ++i)
+	for (int32 OperationIndex = 0; OperationIndex < (int32)EMassObservedOperation::MAX; ++OperationIndex)
 	{
-		(*ElementObservers[i]).Empty();
+		(*ElementObservers[OperationIndex]).Empty();
+		ObservedElements[OperationIndex].Reset();
 	}
 }
 
