Add support for replicating entity destruction in MassReplication

#rnx
#Synced-CL 45873191

diff --git a/Engine/Plugins/Runtime/MassGameplay/Source/MassReplication/Private/MassReplicationProcessor.cpp b/Engine/Plugins/Runtime/MassGameplay/Source/MassReplication/Private/MassReplicationProcessor.cpp
--- a/Engine/Plugins/Runtime/MassGameplay/Source/MassReplication/Private/MassReplicationProcessor.cpp
+++ b/Engine/Plugins/Runtime/MassGameplay/Source/MassReplication/Private/MassReplicationProcessor.cpp
@@ -398,11 +398,11 @@ void UMassReplicationProcessor::Execute(FMassEntityManager& EntityManager, FMass
 		}
 		ClientReplicationInfo.HandledEntities = MoveTemp(EntitiesInRange);
 
-		// Cleanup any AgentData that isn't relevant anymore (that is EMassLOD::OFF)
+		// Cleanup any AgentData that isn't relevant anymore (that is EMassLOD::OFF or pending destruction)
 		for (FMassReplicationAgentDataMap::TIterator It = ClientReplicationInfo.AgentsData.CreateIterator(); It; ++It)
 		{
 			FMassReplicatedAgentData& AgentData = It.Value();
-			if (AgentData.LOD == EMassLOD::Off)
+			if (AgentData.bPendingDestruction || AgentData.LOD == EMassLOD::Off)
 			{
 				checkf(!AgentData.Handle.IsValid(), TEXT("This replicated agent should have been removed from this client and was not"));
 				It.RemoveCurrent();
@@ -411,3 +411,34 @@ void UMassReplicationProcessor::Execute(FMassEntityManager& EntityManager, FMass
 	}
 #endif //UE_REPLICATION_COMPILE_SERVER_CODE
 }
+
+
+UMassReplicationEntityDestructionObserver::UMassReplicationEntityDestructionObserver()
+	: EntityQuery(*this)
+{
+	ExecutionFlags = static_cast<uint8>(EProcessorExecutionFlags::Server);
+	ObservedType = FMassReplicatedAgentFragment::StaticStruct();
+	ObservedOperations = EMassObservedOperationFlags::DestroyEntity;
+
+	// Make sure we can only run in Game Thread to sync properly with UMassReplicationProcessor
+	bRequiresGameThreadExecution = true;
+}
+
+void UMassReplicationEntityDestructionObserver::ConfigureQueries(const TSharedRef<FMassEntityManager>& EntityManager)
+{
+	EntityQuery.AddSubsystemRequirement<UMassReplicationSubsystem>(EMassFragmentAccess::ReadWrite);
+}
+
+void UMassReplicationEntityDestructionObserver::Execute(FMassEntityManager& EntityManager, FMassExecutionContext& Context)
+{
+#if UE_REPLICATION_COMPILE_SERVER_CODE
+	EntityQuery.ForEachEntityChunk(Context, [this](FMassExecutionContext& Context)
+	{
+		UMassReplicationSubsystem& ReplicationSubsystem = Context.GetMutableSubsystemChecked<UMassReplicationSubsystem>();
+		for (FMassExecutionContext::FEntityIterator EntityIt = Context.CreateEntityIterator(); EntityIt; ++EntityIt)
+		{
+			ReplicationSubsystem.NotifyEntityDestroyed(Context.GetEntity(EntityIt));
+		}
+	});
+#endif
+}
\ No newline at end of file
diff --git a/Engine/Plugins/Runtime/MassGameplay/Source/MassReplication/Private/MassReplicationSubsystem.cpp b/Engine/Plugins/Runtime/MassGameplay/Source/MassReplication/Private/MassReplicationSubsystem.cpp
--- a/Engine/Plugins/Runtime/MassGameplay/Source/MassReplication/Private/MassReplicationSubsystem.cpp
+++ b/Engine/Plugins/Runtime/MassGameplay/Source/MassReplication/Private/MassReplicationSubsystem.cpp
@@ -585,7 +585,20 @@ void UMassReplicationSubsystem::DebugCheckArraysAreInSync()
 	}
 
  #endif // UE_DEBUG_REPLICATION
-};
+}
+
+#if UE_REPLICATION_COMPILE_SERVER_CODE
+void UMassReplicationSubsystem::NotifyEntityDestroyed(FMassEntityHandle EntityHandle)
+{
+	for (FMassClientReplicationInfo& ClientReplicationInfo : ClientsReplicationInfo)
+	{
+		if (FMassReplicatedAgentData* AgentDataPtr = ClientReplicationInfo.AgentsData.Find(EntityHandle))
+		{
+			AgentDataPtr->bPendingDestruction = true;
+		}
+	}
+}
+#endif // UE_REPLICATION_COMPILE_SERVER_CODE
 
 void UMassReplicationSubsystem::AddClient(FMassViewerHandle ViewerHandle, APlayerController& InController)
 {
diff --git a/Engine/Plugins/Runtime/MassGameplay/Source/MassReplication/Public/MassReplicationProcessor.h b/Engine/Plugins/Runtime/MassGameplay/Source/MassReplication/Public/MassReplicationProcessor.h
--- a/Engine/Plugins/Runtime/MassGameplay/Source/MassReplication/Public/MassReplicationProcessor.h
+++ b/Engine/Plugins/Runtime/MassGameplay/Source/MassReplication/Public/MassReplicationProcessor.h
@@ -50,6 +50,20 @@ protected:
 	FMassEntityQuery EntityQuery;
 };
 
+UCLASS(MinimalAPI)
+class UMassReplicationEntityDestructionObserver : public UMassObserverProcessor
+{
+	GENERATED_BODY()
+
+public:
+	MASSREPLICATION_API UMassReplicationEntityDestructionObserver();
+
+protected:
+	MASSREPLICATION_API virtual void ConfigureQueries(const TSharedRef<FMassEntityManager>& EntityManager) override;
+	MASSREPLICATION_API virtual void Execute(FMassEntityManager& EntityManager, FMassExecutionContext& Context) override;
+
+	FMassEntityQuery EntityQuery;
+};
 
 struct FMassReplicationContext
 {
@@ -104,20 +118,16 @@ void UMassReplicatorBase::CalculateClientReplication(FMassExecutionContext& Cont
 	CacheViews(Context);
 
 	const double Time = ReplicationContext.World.GetRealTimeSeconds();
-
+	const FMassClientHandle& ClientHandle = RepSharedFragment.CurrentClientHandle;
+	check(ClientHandle.IsValid());
+	checkSlow(RepSharedFragment.BubbleInfos[ClientHandle.GetIndex()] != nullptr);
+	
 	for (FMassExecutionContext::FEntityIterator EntityIt = Context.CreateEntityIterator(); EntityIt; ++EntityIt)
 	{
 		FMassReplicatedAgentFragment& AgentFragment = ReplicatedAgentList[EntityIt];
-
-		const FMassClientHandle& ClientHandle = RepSharedFragment.CurrentClientHandle;
-		check(ClientHandle.IsValid());
-
-		checkSlow(RepSharedFragment.BubbleInfos[ClientHandle.GetIndex()] != nullptr);
-
 		FMassReplicatedAgentData& AgentData = AgentFragment.AgentData;
 
 		const EMassLOD::Type LOD = ViewerLODList[EntityIt].LOD;
-
 		if (LOD < EMassLOD::Off)
 		{
 			AgentData.LOD = LOD;
@@ -152,5 +162,17 @@ void UMassReplicatorBase::CalculateClientReplication(FMassExecutionContext& Cont
 			}
 		}
 	}
+
+	// Handle destroyed entities
+	FMassClientReplicationInfo& ClientReplicationInfo = ReplicationContext.ReplicationSubsystem.GetMutableClientReplicationInfoChecked(ClientHandle);
+	for (FMassReplicationAgentDataMap::TIterator It = ClientReplicationInfo.AgentsData.CreateIterator(); It; ++It)
+	{
+		FMassReplicatedAgentData& AgentData = It.Value();
+		if (AgentData.Handle.IsValid() && AgentData.bPendingDestruction)
+		{
+			RemoveEntity(Context, AgentData.Handle, ClientHandle);
+			AgentData.Invalidate();
+		}
+	}
 #endif //UE_REPLICATION_COMPILE_SERVER_CODE
 }
diff --git a/Engine/Plugins/Runtime/MassGameplay/Source/MassReplication/Public/MassReplicationSubsystem.h b/Engine/Plugins/Runtime/MassGameplay/Source/MassReplication/Public/MassReplicationSubsystem.h
--- a/Engine/Plugins/Runtime/MassGameplay/Source/MassReplication/Public/MassReplicationSubsystem.h
+++ b/Engine/Plugins/Runtime/MassGameplay/Source/MassReplication/Public/MassReplicationSubsystem.h
@@ -209,6 +209,10 @@ public:
 	const FReplicationHashGrid2D& GetGrid() const { return ReplicationGrid; }
 	FReplicationHashGrid2D& GetGridMutable() { return ReplicationGrid; }
 
+#if UE_REPLICATION_COMPILE_SERVER_CODE
+	void NotifyEntityDestroyed(FMassEntityHandle EntityHandle);
+#endif
+
 #if UE_REPLICATION_COMPILE_CLIENT_CODE
 
 	/** 
diff --git a/Engine/Plugins/Runtime/MassGameplay/Source/MassReplication/Public/MassReplicationTypes.h b/Engine/Plugins/Runtime/MassGameplay/Source/MassReplication/Public/MassReplicationTypes.h
--- a/Engine/Plugins/Runtime/MassGameplay/Source/MassReplication/Public/MassReplicationTypes.h
+++ b/Engine/Plugins/Runtime/MassGameplay/Source/MassReplication/Public/MassReplicationTypes.h
@@ -151,11 +151,13 @@ struct FMassReplicatedAgentData
 		Handle.Invalidate();
 		LastUpdateTime = 0.;
 		LOD = EMassLOD::Off;
+		bPendingDestruction = false;
 	}
 
 	FMassReplicatedAgentHandle Handle;
 	double LastUpdateTime = 0.;
 	EMassLOD::Type LOD = EMassLOD::Off;
+	bool bPendingDestruction = false;
 };
 
 #if UE_REPLICATION_COMPILE_CLIENT_CODE
