Deprecate GET_MASSSIMULATION_CONFIG_VALUE and GET_MASSDEBUGGER_CONFIG_VALUE
- These are now joining GET_MASS_CONFIG_VALUE (CL 46231295)
- There's no need to hide the use of these settings classes. It's unnecessary boilerplate and leads to duplicate lookups of the class instance when there multiple uses.

#Synced-CL 48050961

diff --git a/Engine/Plugins/Runtime/MassGameplay/Source/MassActors/Private/MassActorSpawnerSubsystem.cpp b/Engine/Plugins/Runtime/MassGameplay/Source/MassActors/Private/MassActorSpawnerSubsystem.cpp
--- a/Engine/Plugins/Runtime/MassGameplay/Source/MassActors/Private/MassActorSpawnerSubsystem.cpp
+++ b/Engine/Plugins/Runtime/MassGameplay/Source/MassActors/Private/MassActorSpawnerSubsystem.cpp
@@ -544,7 +544,9 @@ void UMassActorSpawnerSubsystem::OnPrePhysicsPhaseStarted(const float DeltaSecon
 	// 
 	// Note: MassRepresentationProcessor relies on actor spawns being processed before it runs so it can confirm the
 	// spawn and clean up the previous representation  
-	ProcessPendingSpawningRequest(GET_MASSSIMULATION_CONFIG_VALUE(DesiredActorSpawningTimeSlicePerTick));
+	const UMassSimulationSettings* SimSettings = GetDefault<UMassSimulationSettings>();
+	check(SimSettings);
+	ProcessPendingSpawningRequest(SimSettings->DesiredActorSpawningTimeSlicePerTick);
 
 	CSV_CUSTOM_STAT(MassActors, NumSpawned, NumActorSpawned, ECsvCustomStatOp::Accumulate);
 	CSV_CUSTOM_STAT(MassActors, NumPooled, NumActorPooled, ECsvCustomStatOp::Accumulate);
@@ -557,7 +559,9 @@ void UMassActorSpawnerSubsystem::OnPrePhysicsPhaseFinished(const float DeltaSeco
 	// Note: MassRepresentationProcessor relies on actor destruction processing after it runs so it can clean up
 	// unwanted actor representations that it has replaced for this frame. It also relies on this running before physics
 	// so unwanted representations don't interfere with new physics enabled actors 
-	ProcessPendingDestruction(GET_MASSSIMULATION_CONFIG_VALUE(DesiredActorDestructionTimeSlicePerTick));
+	const UMassSimulationSettings* SimSettings = GetDefault<UMassSimulationSettings>();
+	check(SimSettings);
+	ProcessPendingDestruction(SimSettings->DesiredActorDestructionTimeSlicePerTick);
 }
 
 void UMassActorSpawnerSubsystem::AddReferencedObjects(UObject* InThis, FReferenceCollector& Collector)
diff --git a/Engine/Plugins/Runtime/MassGameplay/Source/MassRepresentation/Private/MassRepresentationSubsystem.cpp b/Engine/Plugins/Runtime/MassGameplay/Source/MassRepresentation/Private/MassRepresentationSubsystem.cpp
--- a/Engine/Plugins/Runtime/MassGameplay/Source/MassRepresentation/Private/MassRepresentationSubsystem.cpp
+++ b/Engine/Plugins/Runtime/MassGameplay/Source/MassRepresentation/Private/MassRepresentationSubsystem.cpp
@@ -428,8 +428,10 @@ void UMassRepresentationSubsystem::Initialize(FSubsystemCollectionBase& Collecti
 		AgentSystem->GetOnMassAgentComponentEntityDetaching().AddUObject(this, &UMassRepresentationSubsystem::OnMassAgentComponentEntityDetaching);
 	}
 
-	RetryMovedDistanceSq = FMath::Square(GET_MASSSIMULATION_CONFIG_VALUE(DesiredActorFailedSpawningRetryMoveDistance));
-	RetryTimeInterval = GET_MASSSIMULATION_CONFIG_VALUE(DesiredActorFailedSpawningRetryTimeInterval);
+	const UMassSimulationSettings* SimSettings = GetDefault<UMassSimulationSettings>();
+	check(SimSettings);
+	RetryMovedDistanceSq = FMath::Square(SimSettings->DesiredActorFailedSpawningRetryMoveDistance);
+	RetryTimeInterval = SimSettings->DesiredActorFailedSpawningRetryTimeInterval;
 }
 
 void UMassRepresentationSubsystem::Deinitialize()
diff --git a/Engine/Plugins/Runtime/MassGameplay/Source/MassSimulation/Public/MassSimulationSettings.h b/Engine/Plugins/Runtime/MassGameplay/Source/MassSimulation/Public/MassSimulationSettings.h
--- a/Engine/Plugins/Runtime/MassGameplay/Source/MassSimulation/Public/MassSimulationSettings.h
+++ b/Engine/Plugins/Runtime/MassGameplay/Source/MassSimulation/Public/MassSimulationSettings.h
@@ -5,7 +5,9 @@
 #include "MassSettings.h"
 #include "MassSimulationSettings.generated.h"
 
-#define GET_MASSSIMULATION_CONFIG_VALUE(a) (GetMutableDefault<UMassSimulationSettings>()->a)
+#define GET_MASSSIMULATION_CONFIG_VALUE(a) \
+	UE_DEPRECATED_MACRO(5.8, "Use GetDefault<UMassSimulationSettings>() (or GetMutableDefault) directly.") \
+	(GetMutableDefault<UMassSimulationSettings>()->a)
 
 /**
  * Implements the settings for MassSimulation
diff --git a/Engine/Source/Editor/MassEntityDebugger/Private/MassDebuggerModel.cpp b/Engine/Source/Editor/MassEntityDebugger/Private/MassDebuggerModel.cpp
--- a/Engine/Source/Editor/MassEntityDebugger/Private/MassDebuggerModel.cpp
+++ b/Engine/Source/Editor/MassEntityDebugger/Private/MassDebuggerModel.cpp
@@ -45,7 +45,10 @@ namespace UE::Mass::Debugger::Private
 	void MakeDisplayName(const FString& InName, FString& OutDisplayName)
 	{
 		OutDisplayName = InName;
-		if (GET_MASSDEBUGGER_CONFIG_VALUE(bStripMassPrefix) == false)
+
+		const UMassDebuggerSettings* DebuggerSettings = GetDefault<UMassDebuggerSettings>();
+		check(DebuggerSettings);
+		if (DebuggerSettings->bStripMassPrefix == false)
 		{
 			return;
 		}
diff --git a/Engine/Source/Editor/MassEntityDebugger/Public/MassDebuggerSettings.h b/Engine/Source/Editor/MassEntityDebugger/Public/MassDebuggerSettings.h
--- a/Engine/Source/Editor/MassEntityDebugger/Public/MassDebuggerSettings.h
+++ b/Engine/Source/Editor/MassEntityDebugger/Public/MassDebuggerSettings.h
@@ -6,7 +6,9 @@
 #include "MassDebuggerSettings.generated.h"
 
 
-#define GET_MASSDEBUGGER_CONFIG_VALUE(a) (GetMutableDefault<UMassDebuggerSettings>()->a)
+#define GET_MASSDEBUGGER_CONFIG_VALUE(a) \
+	UE_DEPRECATED_MACRO(5.8, "Use GetDefault<UMassDebuggerSettings>() (or GetMutableDefault) directly.") \
+	(GetMutableDefault<UMassDebuggerSettings>()->a)
 
 UCLASS(MinimalAPI, config = EditorPerProjectUserSettings, DisplayName = "Mass Debugger")
 class UMassDebuggerSettings : public UMassModuleSettings
