[MassGameplay] fixing puppet misspelling in MassAgentComponent

#rnx
#Synced-CL 48503291

diff --git a/Engine/Plugins/Runtime/MassGameplay/Source/MassActors/Private/MassAgentComponent.cpp b/Engine/Plugins/Runtime/MassGameplay/Source/MassActors/Private/MassAgentComponent.cpp
--- a/Engine/Plugins/Runtime/MassGameplay/Source/MassActors/Private/MassAgentComponent.cpp
+++ b/Engine/Plugins/Runtime/MassGameplay/Source/MassActors/Private/MassAgentComponent.cpp
@@ -533,7 +533,7 @@ void UMassAgentComponent::PausePuppet(const bool bPause)
 
 void UMassAgentComponent::PuppetReplicationPending()
 {
-	checkf(IsNetSimulating(), TEXT("Expecting a replicated pupet"));
+	checkf(IsNetSimulating(), TEXT("Expecting a replicated puppet"));
 	MASSAGENT_CHECK(State == EAgentComponentState::None,
 		TEXT("%s is expecting to be in state[None] but is in %s"), ANSI_TO_TCHAR(__FUNCTION__), *UEnum::GetValueAsString(State));
 
@@ -542,7 +542,7 @@ void UMassAgentComponent::PuppetReplicationPending()
 
 void UMassAgentComponent::SetReplicatedPuppetHandle(FMassEntityHandle NewHandle)
 {
-	checkf(IsNetSimulating(), TEXT("Expecting a replicated pupet"));
+	checkf(IsNetSimulating(), TEXT("Expecting a replicated puppet"));
 	MASSAGENT_CHECK(State == EAgentComponentState::PuppetPendingReplication ||
 		State == EAgentComponentState::PuppetReplicatedOrphan,
 		TEXT("%s is expecting to be in state [PuppetPendingReplication|PuppetReplicatedOrphan] but is in %s"), ANSI_TO_TCHAR(__FUNCTION__), *UEnum::GetValueAsString(State));
@@ -552,7 +552,7 @@ void UMassAgentComponent::SetReplicatedPuppetHandle(FMassEntityHandle NewHandle)
 
 void UMassAgentComponent::ClearReplicatedPuppetHandle()
 {
-	checkf(IsNetSimulating(), TEXT("Expecting a replicated pupet"));
+	checkf(IsNetSimulating(), TEXT("Expecting a replicated puppet"));
 	MASSAGENT_CHECK(State == EAgentComponentState::PuppetPaused,
 		TEXT("%s is expecting to be in state [PuppetPaused] but is in %s"), ANSI_TO_TCHAR(__FUNCTION__), *UEnum::GetValueAsString(State));
 
@@ -567,7 +567,7 @@ void UMassAgentComponent::ClearReplicatedPuppetHandle()
 
 void UMassAgentComponent::MakePuppetAReplicatedOrphan()
 {
-	checkf(IsNetSimulating(), TEXT("Expecting a replicated pupet"));
+	checkf(IsNetSimulating(), TEXT("Expecting a replicated puppet"));
 
 	SwitchToState(EAgentComponentState::PuppetReplicatedOrphan);
 }
