Fixup API macros

#rnx

#Synced-CL 48600773

#Synced-CL 48605435

diff --git a/Engine/Plugins/Experimental/ChaosMover/Source/ChaosMover/Public/ChaosMover/ChaosMoverSimulationTypes.h b/Engine/Plugins/Experimental/ChaosMover/Source/ChaosMover/Public/ChaosMover/ChaosMoverSimulationTypes.h
--- a/Engine/Plugins/Experimental/ChaosMover/Source/ChaosMover/Public/ChaosMover/ChaosMoverSimulationTypes.h
+++ b/Engine/Plugins/Experimental/ChaosMover/Source/ChaosMover/Public/ChaosMover/ChaosMoverSimulationTypes.h
@@ -430,7 +430,7 @@ struct FAsyncLocalOnlyInstantMovementEffect : public FInstantMovementEffect
 {
 	GENERATED_BODY()
 
-	CHAOSMOVER_API FAsyncLocalOnlyInstantMovementEffect() {}
+	FAsyncLocalOnlyInstantMovementEffect() {}
 	virtual ~FAsyncLocalOnlyInstantMovementEffect() = default;
 
 	FName OptionalName;
@@ -450,7 +450,7 @@ struct FDebugTeleportToInstantMovementEffect : public FInstantMovementEffect
 {
 	GENERATED_BODY()
 
-	CHAOSMOVER_API FDebugTeleportToInstantMovementEffect() {}
+	FDebugTeleportToInstantMovementEffect() {}
 	virtual ~FDebugTeleportToInstantMovementEffect() {}
 
 	// A component can be set to MAX_FLT to signify "keep that current component unchanged"
diff --git a/Engine/Plugins/FX/Niagara/Source/Niagara/Private/NiagaraDataManager.cpp b/Engine/Plugins/FX/Niagara/Source/Niagara/Private/NiagaraDataManager.cpp
--- a/Engine/Plugins/FX/Niagara/Source/Niagara/Private/NiagaraDataManager.cpp
+++ b/Engine/Plugins/FX/Niagara/Source/Niagara/Private/NiagaraDataManager.cpp
@@ -6,6 +6,8 @@
 #include "NiagaraGpuComputeDispatchInterface.h"
 #include "NiagaraComponent.h"
 
+#include UE_INLINE_GENERATED_CPP_BY_NAME(NiagaraDataManager)
+
 
 void UNiagaraDataManager::Init(FNiagaraWorldManager* InWorldManager)
 {
diff --git a/Engine/Plugins/FX/Niagara/Source/Niagara/Private/NiagaraScriptRuntimeCookedData.cpp b/Engine/Plugins/FX/Niagara/Source/Niagara/Private/NiagaraScriptRuntimeCookedData.cpp
--- a/Engine/Plugins/FX/Niagara/Source/Niagara/Private/NiagaraScriptRuntimeCookedData.cpp
+++ b/Engine/Plugins/FX/Niagara/Source/Niagara/Private/NiagaraScriptRuntimeCookedData.cpp
@@ -4,6 +4,8 @@
 
 #include "NiagaraScript.h"
 
+#include UE_INLINE_GENERATED_CPP_BY_NAME(NiagaraScriptRuntimeCookedData)
+
 #if WITH_EDITORONLY_DATA
 void FNiagaraScriptRuntimeCookedData::Init(UNiagaraScript& InScript, const FNiagaraScriptRuntimeCompiledData& InScriptRuntimeCompiledData)
 {
@@ -35,4 +37,4 @@ bool FNiagaraScriptRuntimeCookedData::ValidateWithScript(UNiagaraScript& InScrip
 		}
 	}
 	return true;
-}
\ No newline at end of file
+}
diff --git a/Engine/Plugins/Mutable/Source/MutableRuntime/Public/MuR/External/Value.h b/Engine/Plugins/Mutable/Source/MutableRuntime/Public/MuR/External/Value.h
--- a/Engine/Plugins/Mutable/Source/MutableRuntime/Public/MuR/External/Value.h
+++ b/Engine/Plugins/Mutable/Source/MutableRuntime/Public/MuR/External/Value.h
@@ -41,18 +41,18 @@ namespace UE::Mutable
 		
 		FValue& operator=(const FValue& Other) = delete;
 
-		MUTABLERUNTIME_API FValue(FValue&& Other)
+		FValue(FValue&& Other)
 		{
 			Ptr = MoveTemp(Other.Ptr);
 		}
 
-		MUTABLERUNTIME_API FValue& operator=(FValue&& Other)
+		FValue& operator=(FValue&& Other)
 		{
 			Ptr = MoveTemp(Other.Ptr);
 			return *this;
 		}
 
-		MUTABLERUNTIME_API FValue(const UScriptStruct* ScriptStruct)
+		FValue(const UScriptStruct* ScriptStruct)
 		{
 			Ptr = MakeShared<FInstancedStruct>();
 			Ptr->InitializeAs(ScriptStruct);
@@ -98,12 +98,12 @@ namespace UE::Mutable
 		
 		FValueConst& operator=(const FValueConst& Other) = delete;
 
-		MUTABLERUNTIME_API FValueConst(FValueConst&& Other)
+		FValueConst(FValueConst&& Other)
 		{
 			Ptr = MoveTemp(Other.Ptr);
 		}
 
-		MUTABLERUNTIME_API FValueConst& operator=(FValueConst&& Other)
+		FValueConst& operator=(FValueConst&& Other)
 		{
 			Ptr = MoveTemp(Other.Ptr);
 			return *this;
diff --git a/Engine/Plugins/Runtime/MassGameplay/Source/MassActors/Public/MassActorSpawnerSubsystem.h b/Engine/Plugins/Runtime/MassGameplay/Source/MassActors/Public/MassActorSpawnerSubsystem.h
--- a/Engine/Plugins/Runtime/MassGameplay/Source/MassActors/Public/MassActorSpawnerSubsystem.h
+++ b/Engine/Plugins/Runtime/MassGameplay/Source/MassActors/Public/MassActorSpawnerSubsystem.h
@@ -271,7 +271,7 @@ protected:
 	MASSACTORS_API virtual ESpawnRequestStatus SpawnOrRetrieveFromPool(const FSpawnArgs& Args);
 
 	UE_DEPRECATED(5.8, "Use the overload that accepts FSpawnArgs. OutSpawnedActor is equivalent to Args.Request.SpawnedActor")
-	MASSACTORS_API virtual ESpawnRequestStatus SpawnOrRetrieveFromPool(FConstStructView SpawnRequestView, TObjectPtr<AActor>& OutSpawnedActor) final
+	virtual ESpawnRequestStatus SpawnOrRetrieveFromPool(FConstStructView SpawnRequestView, TObjectPtr<AActor>& OutSpawnedActor) final
 	{
 		return ESpawnRequestStatus::Failed;
 	}
@@ -280,7 +280,7 @@ protected:
 	MASSACTORS_API virtual ESpawnRequestStatus SpawnActor(const FSpawnArgs& Args, FActorSpawnParameters& InOutSpawnParameters) const;
 
 	UE_DEPRECATED(5.8, "Use the overload that accepts FSpawnArgs. OutSpawnedActor is equivalent to Args.Request.SpawnedActor")
-	MASSACTORS_API virtual ESpawnRequestStatus SpawnActor(FConstStructView SpawnRequestView, TObjectPtr<AActor>& OutSpawnedActor, FActorSpawnParameters& InOutSpawnParameters) const final
+	virtual ESpawnRequestStatus SpawnActor(FConstStructView SpawnRequestView, TObjectPtr<AActor>& OutSpawnedActor, FActorSpawnParameters& InOutSpawnParameters) const final
 	{
 		return ESpawnRequestStatus::Failed;
 	}
diff --git a/Engine/Plugins/Runtime/ModelViewViewModel/Source/ModelViewViewModel/Public/IModelViewViewModelModule.h b/Engine/Plugins/Runtime/ModelViewViewModel/Source/ModelViewViewModel/Public/IModelViewViewModelModule.h
--- a/Engine/Plugins/Runtime/ModelViewViewModel/Source/ModelViewViewModel/Public/IModelViewViewModelModule.h
+++ b/Engine/Plugins/Runtime/ModelViewViewModel/Source/ModelViewViewModel/Public/IModelViewViewModelModule.h
@@ -10,6 +10,6 @@ class FMVVMImplicitConverter;
 class IModelViewViewModelModule : public IModuleInterface
 {
 public:
-	MODELVIEWVIEWMODEL_API virtual void RegisterImplicitConverter(const TSharedRef<FMVVMImplicitConverter>& Converter) = 0;
-	MODELVIEWVIEWMODEL_API virtual void UnregisterImplicitConverter(const TSharedRef<FMVVMImplicitConverter>& Converter) = 0;
+	virtual void RegisterImplicitConverter(const TSharedRef<FMVVMImplicitConverter>& Converter) = 0;
+	virtual void UnregisterImplicitConverter(const TSharedRef<FMVVMImplicitConverter>& Converter) = 0;
 };
diff --git a/Engine/Source/Runtime/AnimGraphRuntime/Public/BoneControllers/AnimNode_AnimDynamics.h b/Engine/Source/Runtime/AnimGraphRuntime/Public/BoneControllers/AnimNode_AnimDynamics.h
--- a/Engine/Source/Runtime/AnimGraphRuntime/Public/BoneControllers/AnimNode_AnimDynamics.h
+++ b/Engine/Source/Runtime/AnimGraphRuntime/Public/BoneControllers/AnimNode_AnimDynamics.h
@@ -562,7 +562,7 @@ struct FAnimNode_AnimDynamics : public FAnimNode_SkeletalControlBase
 	ANIMGRAPHRUNTIME_API virtual void UpdateInternal(const FAnimationUpdateContext& Context) override;
 	ANIMGRAPHRUNTIME_API virtual void EvaluateSkeletalControl_AnyThread(FComponentSpacePoseContext& Output, TArray<FBoneTransform>& OutBoneTransforms) override;
 	ANIMGRAPHRUNTIME_API virtual void GatherDebugData(FNodeDebugData& DebugData) override;
-	ANIMGRAPHRUNTIME_API virtual bool HasPreUpdate() const override { return true; }
+	virtual bool HasPreUpdate() const override { return true; }
 	ANIMGRAPHRUNTIME_API virtual void PreUpdate(const UAnimInstance* InAnimInstance) override;
 	virtual bool NeedsDynamicReset() const override { return true; }
 	virtual void ResetDynamics(ETeleportType InTeleportType) override { RequestInitialise(InTeleportType); }
diff --git a/Engine/Source/Runtime/Engine/Classes/PhysicsEngine/PhysicsConstraintComponent.h b/Engine/Source/Runtime/Engine/Classes/PhysicsEngine/PhysicsConstraintComponent.h
--- a/Engine/Source/Runtime/Engine/Classes/PhysicsEngine/PhysicsConstraintComponent.h
+++ b/Engine/Source/Runtime/Engine/Classes/PhysicsEngine/PhysicsConstraintComponent.h
@@ -103,7 +103,7 @@ public:
 	ENGINE_API void AddReInitDelegatesToComponents();
 	ENGINE_API void InitComponentConstraintWithoutUpdatingFrames();
 	ENGINE_API void CheckAndUpdateConstraintPT();
-	ENGINE_API bool GetIsBrokenLocal() const { return bIsBrokenLocal; };
+	bool GetIsBrokenLocal() const { return bIsBrokenLocal; };
 #endif
 
 	UFUNCTION()
diff --git a/Engine/Source/Runtime/Online/BackgroundHTTPFileHash/Public/BackgroundHttpFileHashHelper.h b/Engine/Source/Runtime/Online/BackgroundHTTPFileHash/Public/BackgroundHttpFileHashHelper.h
--- a/Engine/Source/Runtime/Online/BackgroundHTTPFileHash/Public/BackgroundHttpFileHashHelper.h
+++ b/Engine/Source/Runtime/Online/BackgroundHTTPFileHash/Public/BackgroundHttpFileHashHelper.h
@@ -76,7 +76,7 @@ public:
 	// Removes URL mapping
 	BACKGROUNDHTTPFILEHASH_API void RemoveURLMapping(const FString& URL, const FString& RequestID);
 	UE_DEPRECATED(5.8, "Please use RemoveURLMapping with RequestID argument instead")
-	BACKGROUNDHTTPFILEHASH_API void RemoveURLMapping(const FString& URL) { RemoveURLMapping(URL, {}); }
+	void RemoveURLMapping(const FString& URL) { RemoveURLMapping(URL, {}); }
 
 	// Looks for a URL mapped to the supplied TempFilename. Nullptr returned if one isn't found
 	BACKGROUNDHTTPFILEHASH_API bool FindMappedURLForTempFilename(const FString& TempFilename, FString& OutMappedURL) const;
