[MassGameplay] Added a way to configure MassEntityTraits to be valid for a specific target configuration (Standalone, Client, Server, etc).

#Synced-CL 48899568

diff --git a/Engine/Plugins/Runtime/MassGameplay/Source/MassGameplayEditor/Private/MassTraitRepository.cpp b/Engine/Plugins/Runtime/MassGameplay/Source/MassGameplayEditor/Private/MassTraitRepository.cpp
--- a/Engine/Plugins/Runtime/MassGameplay/Source/MassGameplayEditor/Private/MassTraitRepository.cpp
+++ b/Engine/Plugins/Runtime/MassGameplay/Source/MassGameplayEditor/Private/MassTraitRepository.cpp
@@ -230,6 +230,17 @@ void UMassTraitRepository::OnDebugEvent(const FName EventName, FConstStructView
 				->AddToken(FTextToken::Create(LOCTEXT("TraitIgnoredTrait", "trait was ignored. Check if it's not a duplicate.")));
 		}
 	}
+	else if (EventName == UE::Mass::Debug::TraitOverridden)
+	{
+		if (const FMassGenericDebugEvent* GenericEvent = Payload.GetPtr<const FMassGenericDebugEvent>())
+		{
+			const UMassEntityTraitBase* Trait = Cast<UMassEntityTraitBase>(GenericEvent->Context);
+			FMessageLog MessageLog(UE::Mass::Editor::MessageLogPageName);
+			MessageLog.AddMessage(FTokenizedMessage::Create(OVERRIDABLE_SEVERITY(EMessageSeverity::Info)))
+				->AddToken(FUObjectToken::Create(Trait))
+				->AddToken(FTextToken::Create(LOCTEXT("TraitOverridden", "trait was overridden by another instance. Make sure that's intended.")));
+		}
+	}
 	else if (EventName == DuplicateElementsMessageName)
 	{
 		if (const FMassDuplicateElementsMessage* DuplicateElementsMessage = Payload.GetPtr<const FMassDuplicateElementsMessage>())
diff --git a/Engine/Plugins/Runtime/MassGameplay/Source/MassSpawner/Private/MassEntityConfigAsset.cpp b/Engine/Plugins/Runtime/MassGameplay/Source/MassSpawner/Private/MassEntityConfigAsset.cpp
--- a/Engine/Plugins/Runtime/MassGameplay/Source/MassSpawner/Private/MassEntityConfigAsset.cpp
+++ b/Engine/Plugins/Runtime/MassGameplay/Source/MassSpawner/Private/MassEntityConfigAsset.cpp
@@ -122,15 +122,15 @@ const FMassEntityTemplate* FMassEntityConfig::GetEntityTemplateInternal(const UW
 	return TemplateFound ? &TemplateFound->Get() : nullptr;
 }
 
-void FMassEntityConfig::GetCombinedTraits(TArray<UMassEntityTraitBase*>& OutTraits) const
+void FMassEntityConfig::GetCombinedTraits(TArray<UMassEntityTraitBase*>& OutTraits, FMassEntityTemplateBuildContext* BuildContext) const
 {
 	TArray<const UObject*> Visited;
 	OutTraits.Reset();
 	Visited.Add(ConfigOwner);
-	GetCombinedTraitsInternal(OutTraits, Visited);
+	GetCombinedTraitsInternal(OutTraits, Visited, BuildContext);
 }
 
-void FMassEntityConfig::GetCombinedTraitsInternal(TArray<UMassEntityTraitBase*>& OutTraits, TArray<const UObject*>& Visited) const
+void FMassEntityConfig::GetCombinedTraitsInternal(TArray<UMassEntityTraitBase*>& OutTraits, TArray<const UObject*>& Visited, FMassEntityTemplateBuildContext* BuildContext) const
 {
 	if (Parent)
 	{
@@ -148,7 +148,7 @@ void FMassEntityConfig::GetCombinedTraitsInternal(TArray<UMassEntityTraitBase*>&
 		else
 		{
 			Visited.Add(Parent);
-			Parent->GetConfig().GetCombinedTraitsInternal(OutTraits, Visited);
+			Parent->GetConfig().GetCombinedTraitsInternal(OutTraits, Visited, nullptr);
 		}
 	}
 
@@ -159,9 +159,21 @@ void FMassEntityConfig::GetCombinedTraitsInternal(TArray<UMassEntityTraitBase*>&
 			continue;
 		}
 		// Allow only one feature per type. This is also used to allow child configs override parent features.
-		const int32 Index = OutTraits.IndexOfByPredicate([Trait](const UMassEntityTraitBase* ExistingFeature) -> bool { return Trait->GetClass() == ExistingFeature->GetClass(); });
+		const int32 Index = OutTraits.IndexOfByPredicate([Trait](const UMassEntityTraitBase* ExistingFeature) -> bool
+		{
+			return Trait->GetClass() == ExistingFeature->GetClass()
+				// we allow duplicating trait's that are configured for different targets
+				&& static_cast<bool>(Trait->ValidTargetConfig & ExistingFeature->ValidTargetConfig);
+		});
+
 		if (Index != INDEX_NONE)
 		{
+			UE_LOG(LogMassSpawner, Warning, TEXT("%s: Overriding %s trait with another instance")
+				, *GetNameSafe(ConfigOwner), *Trait->GetName());
+			if (BuildContext)
+			{
+				BuildContext->RegisterOverriddenTrait(Trait);
+			}
 			OutTraits[Index] = Trait;
 		}
 		else
@@ -178,11 +190,10 @@ void FMassEntityConfig::AddTrait(UMassEntityTraitBase& Trait)
 
 bool FMassEntityConfig::ValidateEntityTemplate(const UWorld& World)
 {
-	TArray<UMassEntityTraitBase*> CombinedTraits;
-	GetCombinedTraits(CombinedTraits);
-
 	FMassEntityTemplateData Template;
 	FMassEntityTemplateBuildContext BuildContext(Template);
+	TArray<UMassEntityTraitBase*> CombinedTraits;
+	GetCombinedTraits(CombinedTraits, &BuildContext);
 
 	return BuildContext.BuildFromTraits(CombinedTraits, World);
 }
diff --git a/Engine/Plugins/Runtime/MassGameplay/Source/MassSpawner/Private/MassEntityTemplateRegistry.cpp b/Engine/Plugins/Runtime/MassGameplay/Source/MassSpawner/Private/MassEntityTemplateRegistry.cpp
--- a/Engine/Plugins/Runtime/MassGameplay/Source/MassSpawner/Private/MassEntityTemplateRegistry.cpp
+++ b/Engine/Plugins/Runtime/MassGameplay/Source/MassSpawner/Private/MassEntityTemplateRegistry.cpp
@@ -23,6 +23,7 @@ namespace UE::Mass::Debug
 {
 	const FName TraitFailedValidation(TEXT("TraitFailedValidation"));
 	const FName TraitIgnored(TEXT("TraitIgnored"));
+	const FName TraitOverridden(TEXT("TraitOverridden"));
 
 	static bool bReportDuplicatedFragmentsAsWarnings = false;
 	// anonymous namespace to force CVars's uniqueness - we use the same name in many places, sometimes withing same namespaces
@@ -112,11 +113,14 @@ bool FMassEntityTemplateBuildContext::BuildFromTraits(TConstArrayView<UMassEntit
 	ensureMsgf(bBuildInProgress == false, TEXT("Unexpected occurrence - it suggests FMassEntityTemplateBuildContext::BuildFromTraits "
 		"has been called as a consequence of some UMassEntityTraitBase::BuildTemplate call. Check the callstack."));
 
+	const EProcessorExecutionFlags WorldProcessingFlags = UE::Mass::Utils::DetermineProcessorExecutionFlags(&World);
+
 	bBuildInProgress = true;
 	for (const UMassEntityTraitBase* Trait : Traits)
 	{
 		check(Trait);
-		if (SetTraitBeingProcessed(Trait))
+		if (static_cast<bool>(Trait->ValidTargetConfig & WorldProcessingFlags)
+			&& SetTraitBeingProcessed(Trait))
 		{
 			Trait->BuildTemplate(*this, World);
 		}
@@ -163,6 +167,7 @@ bool FMassEntityTemplateBuildContext::ValidateBuildContext(const UWorld& World)
 
 	int32 ErrorCount = 0;
 	int32 WarningCount = 0;
+	int32 InfosCount = 0;
 
 	// Doing the trait-specific validation first since it can add required elements to the build context
 	for (FTraitData& TraitData : TraitsData)
@@ -252,6 +257,14 @@ bool FMassEntityTemplateBuildContext::ValidateBuildContext(const UWorld& World)
 		);
 		++WarningCount;
 	}
+
+	for (const UMassEntityTraitBase* Trait : OverriddenTraits)
+	{
+		IF_MESSAGES(
+			FMassDebugger::DebugEvent(UE::Mass::Debug::TraitOverridden, FConstStructView::Make(FMassGenericDebugEvent{Trait}));
+		);
+		++InfosCount;
+	}
 	
 	// @todo add dependencies on trait classes? might be hard if traits are unrelated, like requiring UMassLODCollectorTrait 
 	// or UMassDistanceLODCollectorTrait - both supply alternative implementations of a given functionality, but are unrelated.
@@ -259,10 +272,10 @@ bool FMassEntityTemplateBuildContext::ValidateBuildContext(const UWorld& World)
 	// not worth it since we don't even have a use case for it right now.
 
 #if WITH_UNREAL_DEVELOPER_TOOLS && WITH_EDITOR
-	if (GEditor && (ErrorCount || WarningCount))
+	if (GEditor && (ErrorCount || WarningCount || InfosCount))
 	{
 		FMassEditorNotification Notification;
-		Notification.Message = FText::FormatOrdered(LOCTEXT("TraitResult", "Mass Entity Template validation:\n{0} errors and {1} warnings found"), ErrorCount, WarningCount);
+		Notification.Message = FText::FormatOrdered(LOCTEXT("TraitResult", "Mass Entity Template validation:\n{0} errors and {1} warnings found. {2} notes."), ErrorCount, WarningCount, InfosCount);
 		Notification.Severity = ErrorCount ? EMessageSeverity::Error : EMessageSeverity::Warning;
 		Notification.bIncludeSeeOutputLogForDetails = true;
 		Notification.Show();
diff --git a/Engine/Plugins/Runtime/MassGameplay/Source/MassSpawner/Public/MassEntityConfigAsset.h b/Engine/Plugins/Runtime/MassGameplay/Source/MassSpawner/Public/MassEntityConfigAsset.h
--- a/Engine/Plugins/Runtime/MassGameplay/Source/MassSpawner/Public/MassEntityConfigAsset.h
+++ b/Engine/Plugins/Runtime/MassGameplay/Source/MassSpawner/Public/MassEntityConfigAsset.h
@@ -14,6 +14,7 @@
 
 class UMassEntityConfigAsset;
 class UMassEntityTraitBase;
+struct FMassEntityTemplateBuildContext;
 
 /**
  * Describes a Mass agent to spawn. The struct can be embedded to allow last minute changes to the agent (i.e. for debugging).
@@ -77,8 +78,11 @@ struct FMassEntityConfig
 #endif // WITH_EDITOR
 	
 protected:
-	/** Combines traits based on the config hierarchy and returns list of unique traits */
-	MASSSPAWNER_API void GetCombinedTraits(TArray<UMassEntityTraitBase*>& OutTraits) const;
+	/**
+	 * Combines traits based on the config hierarchy and returns list of unique traits
+	 * @param BuildContext if provided, will be used to store additional information for reporting back to the user
+	 */
+	MASSSPAWNER_API void GetCombinedTraits(TArray<UMassEntityTraitBase*>& OutTraits, FMassEntityTemplateBuildContext* BuildContext = nullptr) const;
 
 	/** Looks for a trait of the indicated type, accepting all child classes as well, unless bExactMatch == true */
 	MASSSPAWNER_API UMassEntityTraitBase* FindTraitInternal(TSubclassOf<UMassEntityTraitBase> TraitClass, const bool bExactMatch = false) const;
@@ -99,7 +103,7 @@ private:
 	FGuid ConfigGuid;
 
 private:
-	MASSSPAWNER_API void GetCombinedTraitsInternal(TArray<UMassEntityTraitBase*>& OutTraits, TArray<const UObject*>& Visited) const;
+	MASSSPAWNER_API void GetCombinedTraitsInternal(TArray<UMassEntityTraitBase*>& OutTraits, TArray<const UObject*>& Visited, FMassEntityTemplateBuildContext* BuildContext) const;
 	MASSSPAWNER_API const FMassEntityTemplate* GetEntityTemplateInternal(const UWorld& World, FMassEntityTemplateID& TemplateIDOut) const;
 };
 
diff --git a/Engine/Plugins/Runtime/MassGameplay/Source/MassSpawner/Public/MassEntityTemplateRegistry.h b/Engine/Plugins/Runtime/MassGameplay/Source/MassSpawner/Public/MassEntityTemplateRegistry.h
--- a/Engine/Plugins/Runtime/MassGameplay/Source/MassSpawner/Public/MassEntityTemplateRegistry.h
+++ b/Engine/Plugins/Runtime/MassGameplay/Source/MassSpawner/Public/MassEntityTemplateRegistry.h
@@ -52,6 +52,7 @@ namespace UE::Mass::Debug
 {
 	extern MASSSPAWNER_API const FName TraitFailedValidation;
 	extern MASSSPAWNER_API const FName TraitIgnored;
+	extern MASSSPAWNER_API const FName TraitOverridden;
 }
 #endif // WITH_EDITORONLY_DATA
 
@@ -292,6 +293,10 @@ struct FMassEntityTemplateBuildContext
 	}
 #endif // WITH_EDITORONLY_DATA
 
+	void RegisterOverriddenTrait(TNotNull<UMassEntityTraitBase*> Trait)
+	{
+		OverriddenTraits.Add(Trait);
+	}
 
 protected:
 
@@ -335,6 +340,7 @@ protected:
 	TArray<FTraitData> TraitsData;
 	TSet<const UMassEntityTraitBase*> TraitsProcessed;
 	TSet<const UMassEntityTraitBase*> IgnoredTraits;
+	TSet<const UMassEntityTraitBase*> OverriddenTraits;
 
 	struct FRemovedType
 	{
diff --git a/Engine/Plugins/Runtime/MassGameplay/Source/MassSpawner/Public/MassEntityTraitBase.h b/Engine/Plugins/Runtime/MassGameplay/Source/MassSpawner/Public/MassEntityTraitBase.h
--- a/Engine/Plugins/Runtime/MassGameplay/Source/MassSpawner/Public/MassEntityTraitBase.h
+++ b/Engine/Plugins/Runtime/MassGameplay/Source/MassSpawner/Public/MassEntityTraitBase.h
@@ -5,6 +5,7 @@
 #if UE_ENABLE_INCLUDE_ORDER_DEPRECATED_IN_5_5
 #include "CoreMinimal.h"
 #endif // UE_ENABLE_INCLUDE_ORDER_DEPRECATED_IN_5_6
+#include "MassProcessingTypes.h"
 #include "MassEntityTraitBase.generated.h"
 
 struct FMassEntityTemplateBuildContext;
@@ -84,4 +85,9 @@ protected:
 private:
 	static MASSSPAWNER_API FOnNewTraitType OnNewTraitTypeEvent;
 #endif // WITH_EDITOR
+
+public:
+	/** Configures which configurations are targeted by this trait. Use it to have different configurations for client and server, for example.*/
+	UPROPERTY(EditAnywhere, Category=Mass)
+	EProcessorExecutionFlags ValidTargetConfig = EProcessorExecutionFlags::AllNetModes;
 };
