[MassReplication] Minor touch-ups of replication code - added some checks, ensuring in some cases when MassReplication is miss-configured.

#Synced-CL 48899568

diff --git a/Engine/Plugins/Runtime/MassGameplay/Source/MassActors/Public/MassActorSubsystem.h b/Engine/Plugins/Runtime/MassGameplay/Source/MassActors/Public/MassActorSubsystem.h
--- a/Engine/Plugins/Runtime/MassGameplay/Source/MassActors/Public/MassActorSubsystem.h
+++ b/Engine/Plugins/Runtime/MassGameplay/Source/MassActors/Public/MassActorSubsystem.h
@@ -24,6 +24,12 @@ struct FMassGuidFragment : public FObjectWrapperFragment
 {
 	GENERATED_BODY()
 
+	bool operator==(const FMassGuidFragment& Other) const
+	{
+		return Guid == Other.Guid;
+	}
+
+	UPROPERTY()
 	FGuid Guid;
 };
 
diff --git a/Engine/Plugins/Runtime/MassGameplay/Source/MassActors/Public/MassAgentComponent.h b/Engine/Plugins/Runtime/MassGameplay/Source/MassActors/Public/MassAgentComponent.h
--- a/Engine/Plugins/Runtime/MassGameplay/Source/MassActors/Public/MassAgentComponent.h
+++ b/Engine/Plugins/Runtime/MassGameplay/Source/MassActors/Public/MassAgentComponent.h
@@ -133,7 +133,7 @@ public:
 protected:
 
 	UE_API void SwitchToState(EAgentComponentState NewState);
-	UE_API void SetEntityHandleInternal(const FMassEntityHandle NewHandle);
+	UE_API virtual void SetEntityHandleInternal(const FMassEntityHandle NewHandle);
 	UE_API void ClearEntityHandleInternal();
 	UE_API void DebugCheckStateConsistency();
 
diff --git a/Engine/Plugins/Runtime/MassGameplay/Source/MassReplication/Private/MassReplicationFragments.cpp b/Engine/Plugins/Runtime/MassGameplay/Source/MassReplication/Private/MassReplicationFragments.cpp
--- a/Engine/Plugins/Runtime/MassGameplay/Source/MassReplication/Private/MassReplicationFragments.cpp
+++ b/Engine/Plugins/Runtime/MassGameplay/Source/MassReplication/Private/MassReplicationFragments.cpp
@@ -4,6 +4,7 @@
 #include "MassReplicationSubsystem.h"
 #include "MassReplicationProcessor.h"
 #include "MassExecutionContext.h"
+#include "MassClientBubbleInfoBase.h"
 #include "Engine/World.h"
 
 //----------------------------------------------------------------------//
@@ -81,8 +82,10 @@ FMassReplicationParameters::FMassReplicationParameters()
 //----------------------------------------------------------------------//
 FMassReplicationSharedFragment::FMassReplicationSharedFragment(UMassReplicationSubsystem& ReplicationSubsystem, const FMassReplicationParameters& Params)
 {
-	LODCalculator.Initialize(Params.LODDistance, Params.BufferHysteresisOnDistancePercentage / 100.0f, Params.LODMaxCountPerViewer); // Setting the LODMaxCoundPerViewer as the MAXCount as the processor will be called for each client and will recalculate it per client
+	// Setting the LODMaxCoundPerViewer as the MAXCount as the processor will be called for each client and will recalculate it per client
+	LODCalculator.Initialize(Params.LODDistance, Params.BufferHysteresisOnDistancePercentage / 100.0f, Params.LODMaxCountPerViewer); 
 	BubbleInfoClassHandle = ReplicationSubsystem.GetBubbleInfoClassHandle(Params.BubbleInfoClass);
+	checkf(BubbleInfoClassHandle.IsValid(), TEXT("%s has not been registered as a bubble info class (yet)"), *GetNameSafe(Params.BubbleInfoClass.Get()));
 
 	CachedReplicator = Params.ReplicatorClass.GetDefaultObject();
 	checkf(CachedReplicator, TEXT("Expecting a valid replicator class"))
diff --git a/Engine/Plugins/Runtime/MassGameplay/Source/MassReplication/Private/MassReplicationSubsystem.cpp b/Engine/Plugins/Runtime/MassGameplay/Source/MassReplication/Private/MassReplicationSubsystem.cpp
--- a/Engine/Plugins/Runtime/MassGameplay/Source/MassReplication/Private/MassReplicationSubsystem.cpp
+++ b/Engine/Plugins/Runtime/MassGameplay/Source/MassReplication/Private/MassReplicationSubsystem.cpp
@@ -457,19 +457,20 @@ FMassEntityHandle UMassReplicationSubsystem::ResetEntityIfValid(const FMassNetwo
 
 	FMassReplicationEntityInfo* EntityInfo = FindMassEntityInfoMutable(NetworkID);
 
-	checkf(EntityInfo, TEXT("EntityData must have been added to EntityInfoMap!"));
-	checkf(EntityInfo->ReplicationID >= ReplicationID, TEXT("We must not be removing an item we've never added!"));
-
-	// Only reset the item if its currently Set / Valid and its the most recent ReplicationID. Stale removes after more recent adds are ignored
-	// We do need to check the ReplicationID in this case
-	if (EntityInfo->Entity.IsSet() && (EntityInfo->ReplicationID == ReplicationID))
+	if (ensureMsgf(EntityInfo, TEXT("EntityData must have been added to EntityInfoMap!")))
 	{
-		OnRemovingMassAgent.Broadcast(NetworkID, EntityReset);
+		checkf(EntityInfo->ReplicationID >= ReplicationID, TEXT("We must not be removing an item we've never added!"));
 
-		EntityReset = EntityInfo->Entity;
+		// Only reset the item if its currently Set / Valid and its the most recent ReplicationID. Stale removes after more recent adds are ignored
+		// We do need to check the ReplicationID in this case
+		if (EntityInfo->Entity.IsSet() && (EntityInfo->ReplicationID == ReplicationID))
+		{
+			EntityReset = EntityInfo->Entity;
+			OnRemovingMassAgent.Broadcast(NetworkID, EntityReset);
 
-		//Unset the Entity handle, this indicates that its currently removed from the bubble
-		EntityInfo->Entity = FMassEntityHandle();
+			//Unset the Entity handle, this indicates that its currently removed from the bubble
+			EntityInfo->Entity = FMassEntityHandle();
+		}
 	}
 
 	return EntityReset;
diff --git a/Engine/Plugins/Runtime/MassGameplay/Source/MassReplication/Private/MassReplicationTrait.cpp b/Engine/Plugins/Runtime/MassGameplay/Source/MassReplication/Private/MassReplicationTrait.cpp
--- a/Engine/Plugins/Runtime/MassGameplay/Source/MassReplication/Private/MassReplicationTrait.cpp
+++ b/Engine/Plugins/Runtime/MassGameplay/Source/MassReplication/Private/MassReplicationTrait.cpp
@@ -38,8 +38,13 @@ void UMassReplicationTrait::BuildTemplate(FMassEntityTemplateBuildContext& Build
 	if (LIKELY(!BuildContext.IsInspectingData()))
 	{
 		check(ReplicationSubsystem);
-		FSharedStruct SharedFragment = EntityManager.GetOrCreateSharedFragment<FMassReplicationSharedFragment>(FConstStructView::Make(Params), *ReplicationSubsystem, Params);
-		BuildContext.AddSharedFragment(SharedFragment);
+
+		FMassReplicationSharedFragment ReplicationFragment(*ReplicationSubsystem, Params);
+		if (ensureMsgf(ReplicationFragment.IsValid(), TEXT("Created FMassReplicationSharedFragment instance is not valid")))
+		{
+			FSharedStruct SharedFragment = EntityManager.GetOrCreateSharedFragment(*FMassReplicationSharedFragment::StaticStruct(), reinterpret_cast<uint8*>(&ReplicationFragment));
+			BuildContext.AddSharedFragment(SharedFragment);
+		}
 	}
 	else
 	{
diff --git a/Engine/Plugins/Runtime/MassGameplay/Source/MassReplication/Public/MassReplicationFragments.h b/Engine/Plugins/Runtime/MassGameplay/Source/MassReplication/Public/MassReplicationFragments.h
--- a/Engine/Plugins/Runtime/MassGameplay/Source/MassReplication/Public/MassReplicationFragments.h
+++ b/Engine/Plugins/Runtime/MassGameplay/Source/MassReplication/Public/MassReplicationFragments.h
@@ -149,6 +149,11 @@ public:
 	UPROPERTY(Transient)
 	mutable TObjectPtr<UMassReplicatorBase> CachedReplicator = nullptr;
 
+	bool IsValid() const
+	{
+		return BubbleInfoClassHandle.IsValid() && CachedReplicator;
+	}
+
 	template<typename T>
 	T& GetTypedClientBubbleInfoChecked(FMassClientHandle Handle)
 	{
diff --git a/Engine/Plugins/Runtime/MassGameplay/Source/MassReplication/Public/MassReplicationProcessor.h b/Engine/Plugins/Runtime/MassGameplay/Source/MassReplication/Public/MassReplicationProcessor.h
--- a/Engine/Plugins/Runtime/MassGameplay/Source/MassReplication/Public/MassReplicationProcessor.h
+++ b/Engine/Plugins/Runtime/MassGameplay/Source/MassReplication/Public/MassReplicationProcessor.h
@@ -106,7 +106,8 @@ protected:
 };
 
 template<typename AgentArrayItem, typename CacheViewsCallback, typename AddEntityCallback, typename ModifyEntityCallback, typename RemoveEntityCallback>
-void UMassReplicatorBase::CalculateClientReplication(FMassExecutionContext& Context, FMassReplicationContext& ReplicationContext, CacheViewsCallback&& CacheViews, AddEntityCallback&& AddEntity, ModifyEntityCallback&& ModifyEntity, RemoveEntityCallback&& RemoveEntity)
+void UMassReplicatorBase::CalculateClientReplication(FMassExecutionContext& Context, FMassReplicationContext& ReplicationContext, CacheViewsCallback&& CacheViews
+	, AddEntityCallback&& AddEntity, ModifyEntityCallback&& ModifyEntity, RemoveEntityCallback&& RemoveEntity)
 {
 #if UE_REPLICATION_COMPILE_SERVER_CODE
 	TConstArrayView<FMassNetworkIDFragment> NetworkIDList = Context.GetFragmentView<FMassNetworkIDFragment>();
@@ -157,7 +158,14 @@ void UMassReplicatorBase::CalculateClientReplication(FMassExecutionContext& Cont
 			// as this is a fresh handle, if its valid then we can use the unsafe remove function
 			if (AgentData.Handle.IsValid())
 			{
-				RemoveEntity(Context, AgentData.Handle, ClientHandle);
+				if constexpr (std::is_invocable_v<RemoveEntityCallback, FMassExecutionContext&, FMassEntityHandle, FMassReplicatedAgentHandle, FMassClientHandle>)
+				{
+					RemoveEntity(Context, Context.GetEntity(EntityIt), AgentData.Handle, ClientHandle);
+				}
+				else
+				{
+					RemoveEntity(Context, AgentData.Handle, ClientHandle);
+				}
 				AgentData.Invalidate();
 			}
 		}
@@ -170,7 +178,14 @@ void UMassReplicatorBase::CalculateClientReplication(FMassExecutionContext& Cont
 		FMassReplicatedAgentData& AgentData = It.Value();
 		if (AgentData.Handle.IsValid() && AgentData.bPendingDestruction)
 		{
-			RemoveEntity(Context, AgentData.Handle, ClientHandle);
+			if constexpr (std::is_invocable_v<RemoveEntityCallback, FMassExecutionContext&, FMassEntityHandle, FMassReplicatedAgentHandle, FMassClientHandle>)
+			{
+				RemoveEntity(Context, It.Key(), AgentData.Handle, ClientHandle);
+			}
+			else
+			{
+				RemoveEntity(Context, AgentData.Handle, ClientHandle);
+			}
 			AgentData.Invalidate();
 		}
 	}
