[MassReplication] replaced custom "entity set" creation code in MassReplicaitonProcessor with UE::Mass::Utils::CreateEntityCollections

#[review]-48501807
#Synced-CL 48899568

#Synced-CL 48911231

diff --git a/Engine/Plugins/Runtime/MassGameplay/Source/MassReplication/Private/MassReplicationProcessor.cpp b/Engine/Plugins/Runtime/MassGameplay/Source/MassReplication/Private/MassReplicationProcessor.cpp
--- a/Engine/Plugins/Runtime/MassGameplay/Source/MassReplication/Private/MassReplicationProcessor.cpp
+++ b/Engine/Plugins/Runtime/MassGameplay/Source/MassReplication/Private/MassReplicationProcessor.cpp
@@ -219,183 +219,141 @@ void UMassReplicationProcessor::Execute(FMassEntityManager& EntityManager, FMass
 			ReplicationSubsystem->GetGrid().Query(Bounds, EntitiesInRange);
 		}
 
-		EntityQuery.CacheArchetypes();
-		if (EntityQuery.GetArchetypes().Num() > 0)
-		{
-			// EntitySet stores array of entities per specified archetype, may contain duplicates.
-			struct FEntitySet
-			{
-				void Reset()
-				{
-					Entities.Reset();
-				}
-
-				FMassArchetypeHandle Archetype;
-				TArray<FMassEntityHandle> Entities;
-			};
-			TArray<FEntitySet> EntitySets;
+		EntitiesInRange.Append(ClientReplicationInfo.HandledEntities);
+		TArray<FMassArchetypeEntityCollection> EntitySets;
+		UE::Mass::Utils::CreateEntityCollections(Context.GetEntityManagerChecked(), EntitiesInRange
+				, FMassArchetypeEntityCollection::FoldDuplicates
+				, EntitySets);
 
-			for (const FMassArchetypeHandle& Archetype : EntityQuery.GetArchetypes())
-			{
-				FEntitySet& Set = EntitySets.AddDefaulted_GetRef();
-				Set.Archetype = Archetype;
-			}
-
-			auto BuildEntitySet = [&EntitySets, &EntityManager](const TArray<FMassEntityHandle>& Entities)
-			{
-				FEntitySet* PrevSet = Entities.Num() ? &EntitySets[0] : nullptr;
-				for (const FMassEntityHandle Entity : Entities)
-				{
-					// Add to set of supported archetypes. Dont process if we don't care about the type.
-					const FMassArchetypeHandle Archetype = EntityManager.GetArchetypeForEntity(Entity);
-					FEntitySet* Set = PrevSet && PrevSet->Archetype == Archetype ? PrevSet : EntitySets.FindByPredicate([&Archetype](const FEntitySet& Set) { return Archetype == Set.Archetype; });
-					if (Set != nullptr)
-					{
-						// We don't care about duplicates here, the FMassArchetypeEntityCollection creation below will handle it
-						Set->Entities.Add(Entity);
-						PrevSet = Set;
-					}
-				}
-			};
+		for (FMassArchetypeEntityCollection& Set : EntitySets)
+		{
+			checkf(Set.IsEmpty() == false, TEXT("We don't expect to get empty collections from UE::Mass::Utils::CreateEntityCollections"));
 
-			BuildEntitySet(ClientReplicationInfo.HandledEntities);
-			BuildEntitySet(EntitiesInRange);
+			Context.SetEntityCollection(Set);
 
-			for (FEntitySet& Set : EntitySets)
 			{
-				if (Set.Entities.Num() == 0)
+				QUICK_SCOPE_CYCLE_COUNTER(UMassReplicationProcessor_SyncToMass);
+				SyncClientData.ForEachEntityChunk(Context, [&ClientReplicationInfo](FMassExecutionContext& Context)
 				{
-					continue;
-				}
+					const TArrayView<FMassReplicationLODFragment> ViewerLODList = Context.GetMutableFragmentView<FMassReplicationLODFragment>();
+					TArrayView<FMassReplicatedAgentFragment> ReplicatedAgentList = Context.GetMutableFragmentView<FMassReplicatedAgentFragment>();
 
-				Context.SetEntityCollection(FMassArchetypeEntityCollection(Set.Archetype, Set.Entities, FMassArchetypeEntityCollection::FoldDuplicates));
-
-				{
-					QUICK_SCOPE_CYCLE_COUNTER(UMassReplicationProcessor_SyncToMass);
-					SyncClientData.ForEachEntityChunk(Context, [&ClientReplicationInfo](FMassExecutionContext& Context)
+					for (FMassExecutionContext::FEntityIterator EntityIt = Context.CreateEntityIterator(); EntityIt; ++EntityIt)
 					{
-						const TArrayView<FMassReplicationLODFragment> ViewerLODList = Context.GetMutableFragmentView<FMassReplicationLODFragment>();
-						TArrayView<FMassReplicatedAgentFragment> ReplicatedAgentList = Context.GetMutableFragmentView<FMassReplicatedAgentFragment>();
+						FMassEntityHandle EntityHandle = Context.GetEntity(EntityIt);
+						FMassReplicatedAgentFragment& AgentFragment = ReplicatedAgentList[EntityIt];
+						FMassReplicationLODFragment& LODFragment = ViewerLODList[EntityIt];
 
-						for (FMassExecutionContext::FEntityIterator EntityIt = Context.CreateEntityIterator(); EntityIt; ++EntityIt)
+						if (FMassReplicatedAgentData* AgentData = ClientReplicationInfo.AgentsData.Find(EntityHandle))
 						{
-							FMassEntityHandle EntityHandle = Context.GetEntity(EntityIt);
-							FMassReplicatedAgentFragment& AgentFragment = ReplicatedAgentList[EntityIt];
-							FMassReplicationLODFragment& LODFragment = ViewerLODList[EntityIt];
-
-							if (FMassReplicatedAgentData* AgentData = ClientReplicationInfo.AgentsData.Find(EntityHandle))
-							{
-								LODFragment.LOD = AgentData->LOD;
-								AgentFragment.AgentData = *AgentData;
-							}
-							else
-							{
-								LODFragment.LOD = EMassLOD::Off;
-								AgentFragment.AgentData.Invalidate();
-							}
+							LODFragment.LOD = AgentData->LOD;
+							AgentFragment.AgentData = *AgentData;
 						}
-					});
-				}
+						else
+						{
+							LODFragment.LOD = EMassLOD::Off;
+							AgentFragment.AgentData.Invalidate();
+						}
+					}
+				});
+			}
 
+			{
+				QUICK_SCOPE_CYCLE_COUNTER(UMassReplicationProcessor_LODCollection);
+				CollectViewerInfoQuery.ForEachEntityChunk(Context, [](FMassExecutionContext& Context)
 				{
-					QUICK_SCOPE_CYCLE_COUNTER(UMassReplicationProcessor_LODCollection);
-					CollectViewerInfoQuery.ForEachEntityChunk(Context, [](FMassExecutionContext& Context)
-					{
-						const TConstArrayView<FTransformFragment> LocationList = Context.GetFragmentView<FTransformFragment>();
-						const TArrayView<FMassReplicationViewerInfoFragment> ViewersInfoList = Context.GetMutableFragmentView<FMassReplicationViewerInfoFragment>();
-						FMassReplicationSharedFragment& RepSharedFragment = Context.GetMutableSharedFragment<FMassReplicationSharedFragment>();
-						RepSharedFragment.LODCollector.CollectLODInfo(Context, LocationList, ViewersInfoList, ViewersInfoList);
-					});
-				}
+					const TConstArrayView<FTransformFragment> LocationList = Context.GetFragmentView<FTransformFragment>();
+					const TArrayView<FMassReplicationViewerInfoFragment> ViewersInfoList = Context.GetMutableFragmentView<FMassReplicationViewerInfoFragment>();
+					FMassReplicationSharedFragment& RepSharedFragment = Context.GetMutableSharedFragment<FMassReplicationSharedFragment>();
+					RepSharedFragment.LODCollector.CollectLODInfo(Context, LocationList, ViewersInfoList, ViewersInfoList);
+				});
+			}
 
+			{
+				QUICK_SCOPE_CYCLE_COUNTER(UMassReplicationProcessor_LODCaculation);
+				CalculateLODQuery.ForEachEntityChunk(Context, [](FMassExecutionContext& Context)
 				{
-					QUICK_SCOPE_CYCLE_COUNTER(UMassReplicationProcessor_LODCaculation);
-					CalculateLODQuery.ForEachEntityChunk(Context, [](FMassExecutionContext& Context)
-					{
-						const TConstArrayView<FMassReplicationViewerInfoFragment> ViewersInfoList = Context.GetFragmentView<FMassReplicationViewerInfoFragment>();
-						const TArrayView<FMassReplicationLODFragment> ViewerLODList = Context.GetMutableFragmentView<FMassReplicationLODFragment>();
-						FMassReplicationSharedFragment& RepSharedFragment = Context.GetMutableSharedFragment<FMassReplicationSharedFragment>();
-						RepSharedFragment.LODCalculator.CalculateLOD(Context, ViewersInfoList, ViewerLODList, ViewersInfoList);
-					});
-				}
-				Context.ClearEntityCollection();
+					const TConstArrayView<FMassReplicationViewerInfoFragment> ViewersInfoList = Context.GetFragmentView<FMassReplicationViewerInfoFragment>();
+					const TArrayView<FMassReplicationLODFragment> ViewerLODList = Context.GetMutableFragmentView<FMassReplicationLODFragment>();
+					FMassReplicationSharedFragment& RepSharedFragment = Context.GetMutableSharedFragment<FMassReplicationSharedFragment>();
+					RepSharedFragment.LODCalculator.CalculateLOD(Context, ViewersInfoList, ViewerLODList, ViewersInfoList);
+				});
 			}
+			Context.ClearEntityCollection();
+		}
+
+		{
+			QUICK_SCOPE_CYCLE_COUNTER(UMassReplicationProcessor_LODAdjustDistance);
+			EntityManager.ForEachSharedFragment<FMassReplicationSharedFragment>([](FMassReplicationSharedFragment& RepSharedFragment)
+			{
+				RepSharedFragment.bHasAdjustedDistancesFromCount = RepSharedFragment.LODCalculator.AdjustDistancesFromCount();
+			});
+		}
+
+		for (FMassArchetypeEntityCollection& Set : EntitySets)
+		{
+			Context.SetEntityCollection(Set);
 
 			{
-				QUICK_SCOPE_CYCLE_COUNTER(UMassReplicationProcessor_LODAdjustDistance);
-				EntityManager.ForEachSharedFragment<FMassReplicationSharedFragment>([](FMassReplicationSharedFragment& RepSharedFragment)
+				QUICK_SCOPE_CYCLE_COUNTER(UMassReplicationProcessor_LODAdjustLODFromCount);
+				AdjustLODDistancesQuery.ForEachEntityChunk(Context, [](FMassExecutionContext& Context)
 				{
-					RepSharedFragment.bHasAdjustedDistancesFromCount = RepSharedFragment.LODCalculator.AdjustDistancesFromCount();
+					const TConstArrayView<FMassReplicationViewerInfoFragment> ViewersInfoList = Context.GetFragmentView<FMassReplicationViewerInfoFragment>();
+					const TArrayView<FMassReplicationLODFragment> ViewerLODList = Context.GetMutableFragmentView<FMassReplicationLODFragment>();
+					FMassReplicationSharedFragment& RepSharedFragment = Context.GetMutableSharedFragment<FMassReplicationSharedFragment>();
+					RepSharedFragment.LODCalculator.AdjustLODFromCount(Context, ViewersInfoList, ViewerLODList, ViewersInfoList);
 				});
 			}
 
-			for (FEntitySet& Set : EntitySets)
 			{
-				if (Set.Entities.Num() == 0)
+				QUICK_SCOPE_CYCLE_COUNTER(UMassReplicationProcessor_ProcessClientReplication);
+				FMassReplicationContext ReplicationContext(*World, LODSubsystem, *ReplicationSubsystem);
+				EntityManager.ForEachSharedFragment<FMassReplicationSharedFragment>([&EntityManager, &Context, &ReplicationContext, &ClientHandle](FMassReplicationSharedFragment& RepSharedFragment)
 				{
-					continue;
-				}
-				Context.SetEntityCollection(FMassArchetypeEntityCollection(Set.Archetype, Set.Entities, FMassArchetypeEntityCollection::FoldDuplicates));
+					RepSharedFragment.CurrentClientHandle = ClientHandle;
 
-				{
-					QUICK_SCOPE_CYCLE_COUNTER(UMassReplicationProcessor_LODAdjustLODFromCount);
-					AdjustLODDistancesQuery.ForEachEntityChunk(Context, [](FMassExecutionContext& Context)
+					RepSharedFragment.EntityQuery.ForEachEntityChunk(Context, [&ReplicationContext, &RepSharedFragment](FMassExecutionContext& Context)
 					{
-						const TConstArrayView<FMassReplicationViewerInfoFragment> ViewersInfoList = Context.GetFragmentView<FMassReplicationViewerInfoFragment>();
-						const TArrayView<FMassReplicationLODFragment> ViewerLODList = Context.GetMutableFragmentView<FMassReplicationLODFragment>();
-						FMassReplicationSharedFragment& RepSharedFragment = Context.GetMutableSharedFragment<FMassReplicationSharedFragment>();
-						RepSharedFragment.LODCalculator.AdjustLODFromCount(Context, ViewersInfoList, ViewerLODList, ViewersInfoList);
+						RepSharedFragment.CachedReplicator->ProcessClientReplication(Context, ReplicationContext);
 					});
-				}
+				});
+			}
 
+			{
+				QUICK_SCOPE_CYCLE_COUNTER(UMassReplicationProcessor_SyncFromMass);
+				SyncClientData.ForEachEntityChunk(Context, [&ClientReplicationInfo](FMassExecutionContext& Context)
 				{
-					QUICK_SCOPE_CYCLE_COUNTER(UMassReplicationProcessor_ProcessClientReplication);
-					FMassReplicationContext ReplicationContext(*World, LODSubsystem, *ReplicationSubsystem);
-					EntityManager.ForEachSharedFragment<FMassReplicationSharedFragment>([&EntityManager, &Context, &ReplicationContext, &ClientHandle](FMassReplicationSharedFragment& RepSharedFragment)
-					{
-						RepSharedFragment.CurrentClientHandle = ClientHandle;
-
-						RepSharedFragment.EntityQuery.ForEachEntityChunk(Context, [&ReplicationContext, &RepSharedFragment](FMassExecutionContext& Context)
-						{
-							RepSharedFragment.CachedReplicator->ProcessClientReplication(Context, ReplicationContext);
-						});
-					});
-				}
+					TArrayView<FMassReplicatedAgentFragment> ReplicatedAgentList = Context.GetMutableFragmentView<FMassReplicatedAgentFragment>();
 
-				{
-					QUICK_SCOPE_CYCLE_COUNTER(UMassReplicationProcessor_SyncFromMass);
-					SyncClientData.ForEachEntityChunk(Context, [&ClientReplicationInfo](FMassExecutionContext& Context)
+					for (FMassExecutionContext::FEntityIterator EntityIt = Context.CreateEntityIterator(); EntityIt; ++EntityIt)
 					{
-						TArrayView<FMassReplicatedAgentFragment> ReplicatedAgentList = Context.GetMutableFragmentView<FMassReplicatedAgentFragment>();
-
-						for (FMassExecutionContext::FEntityIterator EntityIt = Context.CreateEntityIterator(); EntityIt; ++EntityIt)
-						{
-							FMassEntityHandle EntityHandle = Context.GetEntity(EntityIt);
-							FMassReplicatedAgentFragment& AgentFragment = ReplicatedAgentList[EntityIt];
-							ClientReplicationInfo.AgentsData.Add(EntityHandle, AgentFragment.AgentData);
-						}
-					});
-				}
+						FMassEntityHandle EntityHandle = Context.GetEntity(EntityIt);
+						FMassReplicatedAgentFragment& AgentFragment = ReplicatedAgentList[EntityIt];
+						ClientReplicationInfo.AgentsData.Add(EntityHandle, AgentFragment.AgentData);
+					}
+				});
+			}
 
 #if WITH_MASSGAMEPLAY_DEBUG
-				// Optional debug display
-				if (UE::Mass::Replication::DebugClientReplicationLOD == ClientHandle.GetIndex())
+			// Optional debug display
+			if (UE::Mass::Replication::DebugClientReplicationLOD == ClientHandle.GetIndex())
+			{
+				EntityManager.ForEachSharedFragment<FMassReplicationSharedFragment>([World, &EntityManager, &Context](FMassReplicationSharedFragment& RepSharedFragment)
 				{
-					EntityManager.ForEachSharedFragment<FMassReplicationSharedFragment>([World, &EntityManager, &Context](FMassReplicationSharedFragment& RepSharedFragment)
+					RepSharedFragment.EntityQuery.ForEachEntityChunk(Context, [World, &RepSharedFragment](FMassExecutionContext& Context)
 					{
-						RepSharedFragment.EntityQuery.ForEachEntityChunk(Context, [World, &RepSharedFragment](FMassExecutionContext& Context)
-						{
-							const TConstArrayView<FTransformFragment> TransformList = Context.GetFragmentView<FTransformFragment>();
-							const TConstArrayView<FMassReplicationLODFragment> ViewerLODList = Context.GetFragmentView<FMassReplicationLODFragment>();
-							RepSharedFragment.LODCalculator.DebugDisplayLOD(Context, ViewerLODList, TransformList, World);
-						});
+						const TConstArrayView<FTransformFragment> TransformList = Context.GetFragmentView<FTransformFragment>();
+						const TConstArrayView<FMassReplicationLODFragment> ViewerLODList = Context.GetFragmentView<FMassReplicationLODFragment>();
+						RepSharedFragment.LODCalculator.DebugDisplayLOD(Context, ViewerLODList, TransformList, World);
 					});
-				}
+				});
+			}
 #endif // WITH_MASSGAMEPLAY_DEBUG
 
-				Context.ClearEntityCollection();
-			}
+			Context.ClearEntityCollection();
 		}
+
+		// note that moving EntitiesInRange will effectively reset it, so not point in calling Reset explicitly. 
 		ClientReplicationInfo.HandledEntities = MoveTemp(EntitiesInRange);
 
 		// Cleanup any AgentData that isn't relevant anymore (that is EMassLOD::OFF or pending destruction)
@@ -412,7 +370,6 @@ void UMassReplicationProcessor::Execute(FMassEntityManager& EntityManager, FMass
 #endif //UE_REPLICATION_COMPILE_SERVER_CODE
 }
 
-
 UMassReplicationEntityDestructionObserver::UMassReplicationEntityDestructionObserver()
 	: EntityQuery(*this)
 {
