Highlight the net mode of the trait in validation message to help user pinpoint the issue.

#Synced-CL 48950132

diff --git a/Engine/Plugins/Runtime/MassGameplay/Source/MassGameplayEditor/Private/MassTraitRepository.cpp b/Engine/Plugins/Runtime/MassGameplay/Source/MassGameplayEditor/Private/MassTraitRepository.cpp
--- a/Engine/Plugins/Runtime/MassGameplay/Source/MassGameplayEditor/Private/MassTraitRepository.cpp
+++ b/Engine/Plugins/Runtime/MassGameplay/Source/MassGameplayEditor/Private/MassTraitRepository.cpp
@@ -136,6 +136,14 @@ void UMassTraitRepository::OnDebugEvent(const FName EventName, FConstStructView
 					? TraitRepo->GetTraitsNameAddingElements(MissingElement)
 					: TConstArrayView<FName>();
 
+				if (MissingTraitMessage->RequestingTrait)
+				{
+					IntroMessage->AddText(
+						FText::FormatOrdered(LOCTEXT("ExecutionFlags", "(ExecutionFlags: {0}) "),
+						UEnum::GetDisplayValueAsText(MissingTraitMessage->RequestingTrait->ValidTargetConfig))
+					);
+				}
+				
 				if (SuggestedTraitNames.Num())
 				{
 					// The FixController will coordinate IFixer instances and the relevant FFixToken to
