[MassGameplay] Modified UMassRepresentationSubsystem::CancelSpawningInternal to also check and report if the input FMassActorSpawnRequestHandle is not loger valid (i.e. outdated)

#rnx
#Synced-CL 48981737

diff --git a/Engine/Plugins/Runtime/MassGameplay/Source/MassRepresentation/Private/MassRepresentationSubsystem.cpp b/Engine/Plugins/Runtime/MassGameplay/Source/MassRepresentation/Private/MassRepresentationSubsystem.cpp
--- a/Engine/Plugins/Runtime/MassGameplay/Source/MassRepresentation/Private/MassRepresentationSubsystem.cpp
+++ b/Engine/Plugins/Runtime/MassGameplay/Source/MassRepresentation/Private/MassRepresentationSubsystem.cpp
@@ -362,12 +362,20 @@ bool UMassRepresentationSubsystem::CancelSpawningInternal(const int16 TemplateAc
 	const TSubclassOf<AActor> TemplateToRelease = TemplateActors[TemplateActorIndex].Actor;
 
 	// Check if there is something to cancel
-	if (!SpawnRequestHandle.IsValid())
+	if (SpawnRequestHandle.IsValid() == false)
 	{
 		return false;
 	}
 
 	check(ActorSpawnerSubsystem);
+	if (ActorSpawnerSubsystem->IsSpawnRequestHandleValid(SpawnRequestHandle) == false)
+	{
+		UE_VLOG_UELOG(this, LogMassRepresentation, Error
+			, TEXT("Attempting to cancel spawning using a NO LONGER VALID spawn request handle [i: %d, s/n: %d]")
+			, SpawnRequestHandle.GetIndex(), SpawnRequestHandle.GetSerialNumber());
+		return false;
+	}
+
 	const FMassActorSpawnRequest& SpawnRequest = ActorSpawnerSubsystem->GetMutableSpawnRequest<FMassActorSpawnRequest>(SpawnRequestHandle);
 	// Check if the spawning request matches the template actor
 	if (SpawnRequest.Template != TemplateToRelease)
