[Mass] Fixed MassAgentComponent not regenerating entity config's GUID when config's parent changes. This is a problem when we change the config's parent in types derived from the original owner class of the component.

#rnx
#Synced-CL 49028421

diff --git a/Engine/Plugins/Runtime/MassGameplay/Source/MassActors/Private/MassAgentComponent.cpp b/Engine/Plugins/Runtime/MassGameplay/Source/MassActors/Private/MassAgentComponent.cpp
--- a/Engine/Plugins/Runtime/MassGameplay/Source/MassActors/Private/MassAgentComponent.cpp
+++ b/Engine/Plugins/Runtime/MassGameplay/Source/MassActors/Private/MassAgentComponent.cpp
@@ -77,7 +77,7 @@ void UMassAgentComponent::PostEditChangeProperty(FPropertyChangedEvent& Property
 
 	Super::PostEditChangeProperty(PropertyChangedEvent);
 
-	if (PropertyChangedEvent.Property && PropertyChangedEvent.Property->GetFName() == NAME_EntityConfig)
+	if (PropertyChangedEvent.MemberProperty && PropertyChangedEvent.MemberProperty->GetFName() == NAME_EntityConfig)
 	{
 		EntityConfig.ForceRegenerateGUID();
 	}
