[Mass] Made MassDebugVisualizer not show up in the Scene Outliner in the editor

#Synced-CL 49028421

diff --git a/Engine/Plugins/Runtime/MassGameplay/Source/MassGameplayDebug/Private/MassDebugVisualizer.cpp b/Engine/Plugins/Runtime/MassGameplay/Source/MassGameplayDebug/Private/MassDebugVisualizer.cpp
--- a/Engine/Plugins/Runtime/MassGameplay/Source/MassGameplayDebug/Private/MassDebugVisualizer.cpp
+++ b/Engine/Plugins/Runtime/MassGameplay/Source/MassGameplayDebug/Private/MassDebugVisualizer.cpp
@@ -13,6 +13,8 @@ AMassDebugVisualizer::AMassDebugVisualizer()
 #if WITH_EDITORONLY_DATA
 	RootComponent = CreateDefaultSubobject<USceneComponent>(TEXT("SceneComp"));
 	DebugVisComponent = CreateDefaultSubobject<UMassDebugVisualizationComponent>(TEXT("VisualizerComponent"));
+	bEditable = false;
+	bListedInSceneOutliner = false;
 #endif
 }
 
