[Mass] Changed how TemplateID for MassAgentComponent's entity config gets calculated - if the agent's config doesn't add any traits we use that config's parent to create the TemplateID, saving memory and time by not creating a bunch of identical templates

#rnx

#Synced-CL 48944341

diff --git a/Engine/Plugins/Runtime/MassGameplay/Source/MassSpawner/Private/MassEntityConfigAsset.cpp b/Engine/Plugins/Runtime/MassGameplay/Source/MassSpawner/Private/MassEntityConfigAsset.cpp
--- a/Engine/Plugins/Runtime/MassGameplay/Source/MassSpawner/Private/MassEntityConfigAsset.cpp
+++ b/Engine/Plugins/Runtime/MassGameplay/Source/MassSpawner/Private/MassEntityConfigAsset.cpp
@@ -112,6 +112,11 @@ const FMassEntityTemplate& FMassEntityConfig::GetEntityTemplateChecked(const UWo
 
 const FMassEntityTemplate* FMassEntityConfig::GetEntityTemplateInternal(const UWorld& World, FMassEntityTemplateID& OutTemplateID) const
 {
+	if (Traits.IsEmpty())
+	{
+		return Parent ? Parent->GetConfig().GetEntityTemplateInternal(World, OutTemplateID) : nullptr;
+	}
+
 	UMassSpawnerSubsystem* SpawnerSystem = UWorld::GetSubsystem<UMassSpawnerSubsystem>(&World);
 	check(SpawnerSystem);
 	const FMassEntityTemplateRegistry& TemplateRegistry = SpawnerSystem->GetTemplateRegistryInstance();
