[MassGameplay] fixed a crash in AMassSpawner::SpawnGeneratedEntities when the spawned entitied do not have FMassCreatedBySpawner fragment

#rnx
#Synced-CL 49100315

diff --git a/Engine/Plugins/Runtime/MassGameplay/Source/MassSpawner/Private/MassSpawner.cpp b/Engine/Plugins/Runtime/MassGameplay/Source/MassSpawner/Private/MassSpawner.cpp
--- a/Engine/Plugins/Runtime/MassGameplay/Source/MassSpawner/Private/MassSpawner.cpp
+++ b/Engine/Plugins/Runtime/MassGameplay/Source/MassSpawner/Private/MassSpawner.cpp
@@ -497,8 +497,10 @@ void AMassSpawner::SpawnGeneratedEntities(TConstArrayView<FMassEntitySpawnDataGe
 #if WITH_EDITOR
 				for (FMassEntityHandle Entity : SpawnedEntities.Entities)
 				{
-					FMassCreatedBySpawner& CreatedBy = EntityManager.GetFragmentDataChecked<FMassCreatedBySpawner>(Entity);
-					CreatedBy.Spawner = this;
+					if (FMassCreatedBySpawner* CreatedBy = EntityManager.GetFragmentDataPtr<FMassCreatedBySpawner>(Entity))
+					{
+						CreatedBy->Spawner = this;
+					}
 				}
 #endif
 			}
