[Mass] made MassAgentComponent's debug-checking for unexpected state configurable via a console variable

#rnx
#Synced-CL 49155255

diff --git a/Engine/Plugins/Runtime/MassGameplay/Source/MassActors/Private/MassAgentComponent.cpp b/Engine/Plugins/Runtime/MassGameplay/Source/MassActors/Private/MassAgentComponent.cpp
--- a/Engine/Plugins/Runtime/MassGameplay/Source/MassActors/Private/MassAgentComponent.cpp
+++ b/Engine/Plugins/Runtime/MassGameplay/Source/MassActors/Private/MassAgentComponent.cpp
@@ -11,16 +11,28 @@
 #include "MassActorSubsystem.h"
 #include "MassSpawner.h"
 #include "Net/UnrealNetwork.h"
-#include "MassReplicationSubsystem.h"
 #include "MassReplicationFragments.h"
 #include "MassMovementFragments.h"
 #include "GameFramework/CharacterMovementComponent.h"
 
 #include UE_INLINE_GENERATED_CPP_BY_NAME(MassAgentComponent)
 
+namespace UE::Mass::Debug
+{
+	bool bBreakOnFailedAgentCheck = true;
+
+	namespace
+	{
+		FAutoConsoleVariableRef CVars[] = 
+		{
+			{TEXT("mass.debug.BreakOnFailedAgentCheck"), bBreakOnFailedAgentCheck, TEXT("Enabled by default. If disabled, won't cause check failures when AgentComponent is in an unexpected state.")}
+		};
+	}
+}
+
 #define MASSAGENT_CHECK( condition, Format, ... ) \
 	UE_CVLOG(!(condition), GetOwner(), LogMass, Error, Format, ##__VA_ARGS__); \
-	checkf( condition, Format, ##__VA_ARGS__ );
+	checkf(condition || !UE::Mass::Debug::bBreakOnFailedAgentCheck, Format, ##__VA_ARGS__);
 
 //----------------------------------------------------------------------//
 // UMassAgentComponent
