[Mass] Made naming spawned actors with entity's GUID, when available, optional
Changed the default behavior to 'false', since it can cause unexpected issues.

#Synced-CL 49155255

diff --git a/Engine/Plugins/Runtime/MassGameplay/Source/MassActors/Private/MassActorSpawnerSubsystem.cpp b/Engine/Plugins/Runtime/MassGameplay/Source/MassActors/Private/MassActorSpawnerSubsystem.cpp
--- a/Engine/Plugins/Runtime/MassGameplay/Source/MassActors/Private/MassActorSpawnerSubsystem.cpp
+++ b/Engine/Plugins/Runtime/MassGameplay/Source/MassActors/Private/MassActorSpawnerSubsystem.cpp
@@ -19,10 +19,20 @@
 
 CSV_DEFINE_CATEGORY(MassActors, true);
 
-namespace UE::MassActors
+namespace UE::Mass::Actors
 {
-	int32 bUseActorPooling = 1;
-	FAutoConsoleVariableRef CVarDebugRepresentationLOD(TEXT("ai.mass.actorpooling"), bUseActorPooling, TEXT("Use Actor Pooling"), ECVF_Scalability);
+	bool bUseActorPooling = true;
+	bool bUseGUIDToNameActors = false;
+	namespace
+	{
+		FAutoConsoleVariableRef CVars[] = 
+		{
+			{TEXT("ai.mass.actorpooling"), bUseActorPooling, TEXT("Use Actor Pooling. DEPRECATED. uSE "), ECVF_Scalability}
+			, {TEXT("mass.spawning.actorpooling"), bUseActorPooling, TEXT("Use Actor Pooling"), ECVF_Scalability}
+			, {TEXT("mass.spawning.UseGUIDToNameActors"), bUseGUIDToNameActors, TEXT("Whether to include spawn request's GUID in the spawned actor's name. Using this can result in some actors being reused if their name matches")}
+		};
+	}
+	
 }
 
 //----------------------------------------------------------------------//
@@ -305,7 +315,7 @@ ESpawnRequestStatus UMassActorSpawnerSubsystem::SpawnActor(const FSpawnArgs& Arg
 {
 	TRACE_CPUPROFILER_EVENT_SCOPE(UMassActorSpawnerSubsystem::SpawnActor);
 
-	if (Args.Request.Guid.IsValid())
+	if (UE::Mass::Actors::bUseGUIDToNameActors && Args.Request.Guid.IsValid())
 	{
 		// offsetting `D` by 1 since `0` has special meaning for FNames
 		InOutSpawnParameters.Name = FName(FString::Printf(TEXT("%s_%ud_%ud_%ud"), *Args.Request.Template->GetName(), Args.Request.Guid.A, Args.Request.Guid.B, Args.Request.Guid.C), Args.Request.Guid.D + 1);
@@ -593,7 +603,7 @@ void UMassActorSpawnerSubsystem::DisableActorPooling()
 
 bool UMassActorSpawnerSubsystem::IsActorPoolingEnabled()
 {
-	return UE::MassActors::bUseActorPooling && bActorPoolingEnabled;
+	return UE::Mass::Actors::bUseActorPooling && bActorPoolingEnabled;
 }
 
 void UMassActorSpawnerSubsystem::ReleaseAllResources()
