[Mass] Fixed MassSimulationSubsystem's simultion not starting properly in client worlds if `mass.SimulationTickingEnabled` is flipped on with a Gamplay Feature Action

#Synced-CL 49155255

diff --git a/Engine/Plugins/Runtime/MassGameplay/Source/MassSimulation/Private/MassSimulationSubsystem.cpp b/Engine/Plugins/Runtime/MassGameplay/Source/MassSimulation/Private/MassSimulationSubsystem.cpp
--- a/Engine/Plugins/Runtime/MassGameplay/Source/MassSimulation/Private/MassSimulationSubsystem.cpp
+++ b/Engine/Plugins/Runtime/MassGameplay/Source/MassSimulation/Private/MassSimulationSubsystem.cpp
@@ -170,7 +170,7 @@ void UMassSimulationSubsystem::HandleSimulationTickingEnabledCVarChange(IConsole
 				{
 					// if enabling call StartSimulation ony if the world has already begun play. Otherwise the 
 					// simulation ticking will be started at the right time automatically
-					if (World->HasBegunPlay())
+					if (MassSimulationSubsystem->HasCalledBeginPlay())
 					{
 						MassSimulationSubsystem->StartSimulation(*World);
 					}
