[MassGameplay] proportion calculation in UMassEntitySpawnDataGeneratorBase::BuildResultsFromEntityTypes fixed

#rnx

#Synced-CL 49198129

diff --git a/Engine/Plugins/Runtime/MassGameplay/Source/MassSpawner/Private/MassEntitySpawnDataGeneratorBase.cpp b/Engine/Plugins/Runtime/MassGameplay/Source/MassSpawner/Private/MassEntitySpawnDataGeneratorBase.cpp
--- a/Engine/Plugins/Runtime/MassGameplay/Source/MassSpawner/Private/MassEntitySpawnDataGeneratorBase.cpp
+++ b/Engine/Plugins/Runtime/MassGameplay/Source/MassSpawner/Private/MassEntitySpawnDataGeneratorBase.cpp
@@ -20,15 +20,25 @@ void UMassEntitySpawnDataGeneratorBase::BuildResultsFromEntityTypes(const int32
 		return;
 	}
 
+	float CumulativeProportion = 0.0f;
+	int TotalSpawned = 0;
+
 	for (int32 i = 0; i < EntityTypes.Num(); i++)
 	{
 		const FMassSpawnedEntityType& EntityType = EntityTypes[i];
-		const int32 EntityCount = static_cast<int32>(static_cast<float>(SpawnCount) * EntityType.Proportion / TotalProportion);
+
+		CumulativeProportion += EntityType.Proportion;
+		const int32 CumulativeEntityCount = FMath::RoundToInt(static_cast<float>(SpawnCount) * CumulativeProportion / TotalProportion);
+
+		const int32 EntityCount = CumulativeEntityCount - TotalSpawned;
+
 		if (EntityCount > 0 && EntityType.GetEntityConfig() != nullptr)
 		{
 			FMassEntitySpawnDataGeneratorResult& Res = OutResults.AddDefaulted_GetRef();
 			Res.NumEntities = EntityCount;
 			Res.EntityConfigIndex = i;
 		}
+
+		TotalSpawned = CumulativeEntityCount;
 	}
 }
