[Mass] reverted the change to FMassEntityConfig::GetEntityTemplateInternal changing the behavior with no traits added. Cause unexpected issues, needs investigation.

#rnx

#Synced-CL 49204654

diff --git a/Engine/Plugins/Runtime/MassGameplay/Source/MassSpawner/Private/MassEntityConfigAsset.cpp b/Engine/Plugins/Runtime/MassGameplay/Source/MassSpawner/Private/MassEntityConfigAsset.cpp
--- a/Engine/Plugins/Runtime/MassGameplay/Source/MassSpawner/Private/MassEntityConfigAsset.cpp
+++ b/Engine/Plugins/Runtime/MassGameplay/Source/MassSpawner/Private/MassEntityConfigAsset.cpp
@@ -112,18 +112,12 @@ const FMassEntityTemplate& FMassEntityConfig::GetEntityTemplateChecked(const UWo
 
 const FMassEntityTemplate* FMassEntityConfig::GetEntityTemplateInternal(const UWorld& World, FMassEntityTemplateID& OutTemplateID) const
 {
-	OutTemplateID = FMassEntityTemplateIDFactory::Make(ConfigGuid);
-
-	if (Traits.IsEmpty())
-	{
-		return Parent ? Parent->GetConfig().GetEntityTemplateInternal(World, OutTemplateID) : nullptr;
-	}
-
 	UMassSpawnerSubsystem* SpawnerSystem = UWorld::GetSubsystem<UMassSpawnerSubsystem>(&World);
 	check(SpawnerSystem);
 	const FMassEntityTemplateRegistry& TemplateRegistry = SpawnerSystem->GetTemplateRegistryInstance();
 
 	// Return existing template if found.
+	OutTemplateID = FMassEntityTemplateIDFactory::Make(ConfigGuid);
 	const TSharedRef<FMassEntityTemplate>* TemplateFound = TemplateRegistry.FindTemplateFromTemplateID(OutTemplateID);
 	return TemplateFound ? &TemplateFound->Get() : nullptr;
 }
