PR #12833: Fix DisconnectActor causing duplicate actors to spawn

#rnx
#Synced-CL 49265344

diff --git a/Engine/Plugins/Runtime/MassGameplay/Source/MassActors/Private/MassActorSubsystem.cpp b/Engine/Plugins/Runtime/MassGameplay/Source/MassActors/Private/MassActorSubsystem.cpp
--- a/Engine/Plugins/Runtime/MassGameplay/Source/MassActors/Private/MassActorSubsystem.cpp
+++ b/Engine/Plugins/Runtime/MassGameplay/Source/MassActors/Private/MassActorSubsystem.cpp
@@ -185,6 +185,8 @@ void FMassActorManager::DisconnectActor(const TObjectKey<const AActor> Actor, co
 		if (FMassActorFragment* Data = EntityManager->GetFragmentDataPtr<FMassActorFragment>(Handle))
 		{
 			Data->ResetAndUpdateHandleMap();
+			// Since we're disconnecting the actor, the entity should no longer be registered in the MassActorSubsystem (or MassActorRepresentationProcessor).
+			EntityManager->RemoveFragmentFromEntity(Handle, FMassActorFragment::StaticStruct());
 		}
 	}
 	else
