[MassGameplay] Fixed in-editor trait validation not working properly after introduction of trati's ValidTargetConfig

#rnx

#Synced-CL 49323322

diff --git a/Engine/Plugins/Runtime/MassGameplay/Source/MassSpawner/Private/MassEntityTemplateRegistry.cpp b/Engine/Plugins/Runtime/MassGameplay/Source/MassSpawner/Private/MassEntityTemplateRegistry.cpp
--- a/Engine/Plugins/Runtime/MassGameplay/Source/MassSpawner/Private/MassEntityTemplateRegistry.cpp
+++ b/Engine/Plugins/Runtime/MassGameplay/Source/MassSpawner/Private/MassEntityTemplateRegistry.cpp
@@ -113,7 +113,9 @@ bool FMassEntityTemplateBuildContext::BuildFromTraits(TConstArrayView<UMassEntit
 	ensureMsgf(bBuildInProgress == false, TEXT("Unexpected occurrence - it suggests FMassEntityTemplateBuildContext::BuildFromTraits "
 		"has been called as a consequence of some UMassEntityTraitBase::BuildTemplate call. Check the callstack."));
 
-	const EProcessorExecutionFlags WorldProcessingFlags = UE::Mass::Utils::DetermineProcessorExecutionFlags(&World);
+	const EProcessorExecutionFlags WorldProcessingFlags = World.IsGameWorld()
+		? UE::Mass::Utils::DetermineProcessorExecutionFlags(&World)
+		: EProcessorExecutionFlags::All;
 
 	bBuildInProgress = true;
 	for (const UMassEntityTraitBase* Trait : Traits)
