[Mass] PR #14089: Mass Visualization: Fixed local transform not applying to the previous transform

#Synced-CL 49323322

diff --git a/Engine/Plugins/Runtime/MassGameplay/Source/MassRepresentation/Private/MassVisualizationComponent.cpp b/Engine/Plugins/Runtime/MassGameplay/Source/MassRepresentation/Private/MassVisualizationComponent.cpp
--- a/Engine/Plugins/Runtime/MassGameplay/Source/MassRepresentation/Private/MassVisualizationComponent.cpp
+++ b/Engine/Plugins/Runtime/MassGameplay/Source/MassRepresentation/Private/MassVisualizationComponent.cpp
@@ -659,10 +659,8 @@ void FMassLODSignificanceRange::AddBatchedTransform(const FMassEntityHandle Enti
 		{
 			SharedData->UpdateInstanceIds.Add(EntityHandle);
 
-			FTransform AdjustedTransform{SharedData->LocalTransform * Transform};
-			
-			SharedData->StaticMeshInstanceTransforms.Add(AdjustedTransform);
-			SharedData->StaticMeshInstancePrevTransforms.Add(PrevTransform);
+			SharedData->StaticMeshInstanceTransforms.Add(SharedData->LocalTransform * Transform);
+			SharedData->StaticMeshInstancePrevTransforms.Add(SharedData->LocalTransform * PrevTransform);
 		}
 	}
 }
