[Mass] Entity Templates unit tests fixed.
Some of the tests stored references to templates while the underlying containers were being reallocated.

#rnx

#Synced-CL 50246518

diff --git a/Engine/Plugins/Runtime/MassGameplay/Source/MassGameplayTestSuite/Private/MassEntityTemplateTest.cpp b/Engine/Plugins/Runtime/MassGameplay/Source/MassGameplayTestSuite/Private/MassEntityTemplateTest.cpp
--- a/Engine/Plugins/Runtime/MassGameplay/Source/MassGameplayTestSuite/Private/MassEntityTemplateTest.cpp
+++ b/Engine/Plugins/Runtime/MassGameplay/Source/MassGameplayTestSuite/Private/MassEntityTemplateTest.cpp
@@ -37,7 +37,7 @@ struct FEntityTemplate_Empty : FEntityTemplateBase
 		const FMassEntityTemplateID InitialTemplateID = FMassEntityTemplateIDFactory::Make(FGuid::NewGuid());
 		FMassEntityTemplateData TemplateData;
 
-		const TSharedRef<FMassEntityTemplate>& FinalizedTemplate = TemplateRegistry.FindOrAddTemplate(InitialTemplateID, MoveTemp(TemplateData));
+		const TSharedRef<FMassEntityTemplate> FinalizedTemplate = TemplateRegistry.FindOrAddTemplate(InitialTemplateID, MoveTemp(TemplateData));
 		
 		const FMassEntityTemplateID TemplateID = FinalizedTemplate->GetTemplateID();
 		AITEST_TRUE("Empty template is expected to be registered as a valid template", TemplateID.IsValid());
@@ -85,7 +85,7 @@ struct FEntityTemplate_Composition : FEntityTemplateBase
 			, ExpectedComposition.IsEquivalent(TemplateData.GetCompositionDescriptor()));
 
 		const FMassEntityTemplateID TemplateID = FMassEntityTemplateIDFactory::Make(FGuid::NewGuid());
-		const TSharedRef<FMassEntityTemplate>& FinalizedTemplate = TemplateRegistry.FindOrAddTemplate(TemplateID, MoveTemp(TemplateData));
+		const TSharedRef<FMassEntityTemplate> FinalizedTemplate = TemplateRegistry.FindOrAddTemplate(TemplateID, MoveTemp(TemplateData));
 
 		const FMassArchetypeHandle& ArchetypeHandle = FinalizedTemplate->GetArchetype();
 		AITEST_TRUE("Empty template is expected to produce a valid (if empty) archetype", ArchetypeHandle.IsValid());
@@ -111,9 +111,9 @@ struct FEntityTemplate_Trivial : FEntityTemplateBase
 		TemplateDataB.AddFragment<FTestFragment_Int>();
 		const FGuid NonEmptyTemplateGuid = FGuid::NewGuid();
 
-		const TSharedRef<FMassEntityTemplate>& FinalizedEmptyTemplate = TemplateRegistry.FindOrAddTemplate(FMassEntityTemplateIDFactory::Make(FGuid::NewGuid()), MoveTemp(EmptyTemplateData));
-		const TSharedRef<FMassEntityTemplate>& FinalizedTemplateA = TemplateRegistry.FindOrAddTemplate(FMassEntityTemplateIDFactory::Make(NonEmptyTemplateGuid), MoveTemp(TemplateDataA));
-		const TSharedRef<FMassEntityTemplate>& FinalizedTemplateB = TemplateRegistry.FindOrAddTemplate(FMassEntityTemplateIDFactory::Make(NonEmptyTemplateGuid), MoveTemp(TemplateDataB));
+		const TSharedRef<FMassEntityTemplate> FinalizedEmptyTemplate = TemplateRegistry.FindOrAddTemplate(FMassEntityTemplateIDFactory::Make(FGuid::NewGuid()), MoveTemp(EmptyTemplateData));
+		const TSharedRef<FMassEntityTemplate> FinalizedTemplateA = TemplateRegistry.FindOrAddTemplate(FMassEntityTemplateIDFactory::Make(NonEmptyTemplateGuid), MoveTemp(TemplateDataA));
+		const TSharedRef<FMassEntityTemplate> FinalizedTemplateB = TemplateRegistry.FindOrAddTemplate(FMassEntityTemplateIDFactory::Make(NonEmptyTemplateGuid), MoveTemp(TemplateDataB));
 
 		AITEST_NOT_EQUAL("Non-empty template data should result in a finalized template different from the empty one", FinalizedEmptyTemplate->GetTemplateID(), FinalizedTemplateA->GetTemplateID());
 		AITEST_EQUAL("Non-empty template data should result in the very same finalized template", FinalizedTemplateA, FinalizedTemplateB);
@@ -139,7 +139,7 @@ struct FEntityTemplate_Modified : FEntityTemplateBase
 		IntFragment.Value = FTestFragment_Int::TestIntValue;
 
 		FMassEntityTemplateData MovedOriginalTemplateData = OriginalTemplateData;
-		const TSharedRef<FMassEntityTemplate>& FinalizedOriginalTemplate = TemplateRegistry.FindOrAddTemplate(FMassEntityTemplateIDFactory::Make(FGuid::NewGuid()), MoveTemp(MovedOriginalTemplateData));
+		const TSharedRef<FMassEntityTemplate> FinalizedOriginalTemplate = TemplateRegistry.FindOrAddTemplate(FMassEntityTemplateIDFactory::Make(FGuid::NewGuid()), MoveTemp(MovedOriginalTemplateData));
 
 		FMassEntityTemplateData NewTemplateData(*FinalizedOriginalTemplate);
 
@@ -149,14 +149,14 @@ struct FEntityTemplate_Modified : FEntityTemplateBase
 		// need to set a new name, otherwise the hash calculated while creating finalized template will end up being the 
 		// same as the original template's. This is the case since we're not using struct's contents while calculating hash (yet).
 		NewTemplateData.SetTemplateName(TEXT("ModifiedTemplate"));
-		FTestFragment_Int* CoppiedFragment = NewTemplateData.GetMutableFragment<FTestFragment_Int>();
-		AITEST_NOT_NULL("The fragment instance is expected to be found", CoppiedFragment);
-		AITEST_EQUAL("The fragment instance is expected to have the same value as the original one", CoppiedFragment->Value, IntFragment.Value);
-		CoppiedFragment->Value = FTestFragment_Int::TestIntValue + 1;
-		AITEST_NOT_EQUAL("Modifying the coppied instance of the fragment doesn't affect the original", CoppiedFragment->Value, IntFragment.Value);
+		FTestFragment_Int* CopiedFragment = NewTemplateData.GetMutableFragment<FTestFragment_Int>();
+		AITEST_NOT_NULL("The fragment instance is expected to be found", CopiedFragment);
+		AITEST_EQUAL("The fragment instance is expected to have the same value as the original one", CopiedFragment->Value, IntFragment.Value);
+		CopiedFragment->Value = FTestFragment_Int::TestIntValue + 1;
+		AITEST_NOT_EQUAL("Modifying the coppied instance of the fragment doesn't affect the original", CopiedFragment->Value, IntFragment.Value);
 
 		const FMassEntityTemplateID NewTemplateID = FMassEntityTemplateIDFactory::MakeFlavor(FinalizedOriginalTemplate->GetTemplateID(), 1);
-		const TSharedRef<FMassEntityTemplate>& FinalizedModifiedTemplate = TemplateRegistry.FindOrAddTemplate(NewTemplateID, MoveTemp(NewTemplateData));
+		const TSharedRef<FMassEntityTemplate> FinalizedModifiedTemplate = TemplateRegistry.FindOrAddTemplate(NewTemplateID, MoveTemp(NewTemplateData));
 
 		AITEST_NOT_EQUAL("The original and modified templates should end up resulting in two different templates", FinalizedModifiedTemplate->GetTemplateID(), FinalizedOriginalTemplate->GetTemplateID());
 		AITEST_EQUAL("The original and modified templates should still point at the same archetype", FinalizedModifiedTemplate->GetArchetype(), FinalizedOriginalTemplate->GetArchetype());
@@ -177,7 +177,7 @@ IMPLEMENT_AI_INSTANT_TEST(FEntityTemplate_Modified, "System.Mass.EntityTemplate.
 //		OriginalTemplateData.AddFragment<FTestFragment_Int>();
 //
 //		FMassEntityTemplateData OriginalTemplateDataCopy = OriginalTemplateData;
-//		const TSharedRef<FMassEntityTemplate>& FinalizedOriginalTemplate = TemplateRegistry.FindOrAddTemplate(MoveTemp(OriginalTemplateData));
+//		const TSharedRef<FMassEntityTemplate> FinalizedOriginalTemplate = TemplateRegistry.FindOrAddTemplate(MoveTemp(OriginalTemplateData));
 //		
 //
 //		AITEST_NOT_EQUAL("Non-empty template data should result in a finalized template different from the empty one", FinalizedEmptyTemplate->GetTemplateID(), FinalizedTemplateA->GetTemplateID());
@@ -206,7 +206,7 @@ struct FEntityTest_BatchSetMissingFragment : FEntityTemplateBase
 		TemplateData.Sort();
 
 		const FMassEntityTemplateID TemplateID = FMassEntityTemplateIDFactory::Make(FGuid::NewGuid());
-		const TSharedRef<FMassEntityTemplate>& EntityTemplate = TemplateRegistry.FindOrAddTemplate(TemplateID, MoveTemp(TemplateData));
+		const TSharedRef<FMassEntityTemplate> EntityTemplate = TemplateRegistry.FindOrAddTemplate(TemplateID, MoveTemp(TemplateData));
 
 		TArray<FMassEntityHandle> CreatedEntities;
 		TSharedRef<FMassEntityManager::FEntityCreationContext> CreationContext
