[Mass] exposing FStaticMeshInstanceVisualizationDesc::IsValid() to API

#Synced-CL 50439895

diff --git a/Engine/Plugins/Runtime/MassGameplay/Source/MassRepresentation/Public/MassRepresentationTypes.h b/Engine/Plugins/Runtime/MassGameplay/Source/MassRepresentation/Public/MassRepresentationTypes.h
--- a/Engine/Plugins/Runtime/MassGameplay/Source/MassRepresentation/Public/MassRepresentationTypes.h
+++ b/Engine/Plugins/Runtime/MassGameplay/Source/MassRepresentation/Public/MassRepresentationTypes.h
@@ -13,6 +13,8 @@
 #include "Math/MathFwd.h"
 #include "MassRepresentationTypes.generated.h"
 
+#define UE_API MASSREPRESENTATION_API
+
 class UMaterialInterface;
 class UStaticMesh;
 struct FMassLODSignificanceRange;
@@ -48,7 +50,7 @@ struct FMassStaticMeshInstanceVisualizationMeshDesc
 {
 	GENERATED_BODY()
 
-	MASSREPRESENTATION_API FMassStaticMeshInstanceVisualizationMeshDesc();
+	UE_API FMassStaticMeshInstanceVisualizationMeshDesc();
 
 	/** A relative Transform from the parent, passed along to FMassISMCSharedData::LocalTransform */
 	UPROPERTY(EditAnywhere, Category = "Mass|Visual")
@@ -164,7 +166,7 @@ struct FStaticMeshInstanceVisualizationDesc : public FTableRowBase
 	}
 
 	/** @return whether any of descriptions in Meshes is valid. This implies that empty Meshes will be treated as not valid.*/
-	bool IsValid() const;
+	UE_API bool IsValid() const;
 };
 
 /** Handle for FStaticMeshInstanceVisualizationDesc's registered with UMassRepresentationSubsystem */
@@ -525,7 +527,7 @@ struct FMassLODSignificanceRange
 	GENERATED_BODY()
 public:
 
-	MASSREPRESENTATION_API void AddBatchedTransform(const FMassEntityHandle EntityHandle, const FTransform& Transform, const FTransform& PrevTransform, TConstArrayView<FISMCSharedDataKey> ExcludeStaticMeshRefs);
+	UE_API void AddBatchedTransform(const FMassEntityHandle EntityHandle, const FTransform& Transform, const FTransform& PrevTransform, TConstArrayView<FISMCSharedDataKey> ExcludeStaticMeshRefs);
 
 	// Adds the specified struct reinterpreted as custom floats to our custom data. Individual members of the specified struct should always fit into a float.
 	// When adding any custom data, the custom data must be added for every instance.
@@ -550,14 +552,14 @@ public:
 		}
 	}
 
-	MASSREPRESENTATION_API void AddBatchedCustomDataFloats(const TArray<float>& CustomFloats, const TArray<FISMCSharedDataKey>& ExcludeStaticMeshRefs);
+	UE_API void AddBatchedCustomDataFloats(const TArray<float>& CustomFloats, const TArray<FISMCSharedDataKey>& ExcludeStaticMeshRefs);
 
 	/** Single-instance version of AddBatchedCustomData when called to add entities (as opposed to modify existing ones).*/
-	MASSREPRESENTATION_API void AddInstance(const FMassEntityHandle EntityHandle, const FTransform& Transform);
+	UE_API void AddInstance(const FMassEntityHandle EntityHandle, const FTransform& Transform);
 
-	MASSREPRESENTATION_API void RemoveInstance(const FMassEntityHandle EntityHandle);
+	UE_API void RemoveInstance(const FMassEntityHandle EntityHandle);
 
-	MASSREPRESENTATION_API void WriteCustomDataFloatsAtStartIndex(int32 StaticMeshIndex, const TArrayView<float>& CustomFloats, const int32 FloatsPerInstance, const int32 StartIndex, const TArray<FISMCSharedDataKey>& ExcludeStaticMeshRefs);
+	UE_API void WriteCustomDataFloatsAtStartIndex(int32 StaticMeshIndex, const TArrayView<float>& CustomFloats, const int32 FloatsPerInstance, const int32 StartIndex, const TArray<FISMCSharedDataKey>& ExcludeStaticMeshRefs);
 
 	/** LOD Significance range */
 	float MinSignificance;
@@ -583,7 +585,7 @@ public:
 	}
 
 	/** Clears out contents so that a given FMassInstancedStaticMeshInfo instance can be reused */
-	MASSREPRESENTATION_API void Reset();
+	UE_API void Reset();
 
 	const FStaticMeshInstanceVisualizationDesc& GetDesc() const
 	{
@@ -704,7 +706,7 @@ public:
 protected:
 
 	/** Destroy the visual instance */
-	MASSREPRESENTATION_API void ClearVisualInstance(UInstancedStaticMeshComponent& ISMComponent);
+	UE_API void ClearVisualInstance(UInstancedStaticMeshComponent& ISMComponent);
 
 	/** Information about this static mesh which will represent all instances */
 	UPROPERTY(VisibleAnywhere, Category = "Mass/Debug")
@@ -778,3 +780,5 @@ typedef FMassInstancedStaticMeshInfoArrayViewAccessDetector FMassInstancedStatic
 typedef TArrayView<FMassInstancedStaticMeshInfo> FMassInstancedStaticMeshInfoArrayView;
 
 #endif // ENABLE_MT_DETECTOR
+
+#undef UE_API
\ No newline at end of file
