Fixed UMassReplicationProcessor accumulating entity handles in ClientReplicationInfo.HandledEntities

#rnx

#Synced-CL 50314128

diff --git a/Engine/Plugins/Runtime/MassGameplay/Source/MassReplication/Private/MassReplicationProcessor.cpp b/Engine/Plugins/Runtime/MassGameplay/Source/MassReplication/Private/MassReplicationProcessor.cpp
--- a/Engine/Plugins/Runtime/MassGameplay/Source/MassReplication/Private/MassReplicationProcessor.cpp
+++ b/Engine/Plugins/Runtime/MassGameplay/Source/MassReplication/Private/MassReplicationProcessor.cpp
@@ -213,6 +213,7 @@ void UMassReplicationProcessor::Execute(FMassEntityManager& EntityManager, FMass
 		// Fetch all entities to process
 		const FVector HalfExtent(MaxLODDistance, MaxLODDistance, 0.0f);
 		TArray<FMassEntityHandle> EntitiesInRange;
+		EntitiesInRange.Reset();
 		for (const FViewerInfo& Viewer : Viewers)
 		{
 			FBox Bounds(Viewer.Location - HalfExtent, Viewer.Location + HalfExtent);
@@ -225,11 +226,17 @@ void UMassReplicationProcessor::Execute(FMassEntityManager& EntityManager, FMass
 				, FMassArchetypeEntityCollection::FoldDuplicates
 				, EntitySets);
 
+		TArray<FMassEntityHandle> EntitiesHandledThisFrame;
+
 		for (FMassArchetypeEntityCollection& Set : EntitySets)
 		{
 			checkf(Set.IsEmpty() == false, TEXT("We don't expect to get empty collections from UE::Mass::Utils::CreateEntityCollections"));
 
 			Context.SetEntityCollection(Set);
+			
+			// Make sure we save a deduplicated list of entities that were handled to avoid
+			// accumulating instances in ClientReplicationInfo.HandledEntities over several frames
+			Set.ExportEntityHandles(EntitiesHandledThisFrame);
 
 			{
 				QUICK_SCOPE_CYCLE_COUNTER(UMassReplicationProcessor_SyncToMass);
@@ -353,8 +360,8 @@ void UMassReplicationProcessor::Execute(FMassEntityManager& EntityManager, FMass
 			Context.ClearEntityCollection();
 		}
 
-		// note that moving EntitiesInRange will effectively reset it, so not point in calling Reset explicitly. 
-		ClientReplicationInfo.HandledEntities = MoveTemp(EntitiesInRange);
+		// note that moving EntitiesHandledThisFrame will effectively reset it, so not point in calling Reset explicitly. 
+		ClientReplicationInfo.HandledEntities = MoveTemp(EntitiesHandledThisFrame);
 
 		// Cleanup any AgentData that isn't relevant anymore (that is EMassLOD::OFF or pending destruction)
 		for (FMassReplicationAgentDataMap::TIterator It = ClientReplicationInfo.AgentsData.CreateIterator(); It; ++It)
