[MassVisualization] avoid creating mass visualization actor/component during world teardown

#rnx

#Synced-CL 50865525

diff --git a/Engine/Plugins/Runtime/MassGameplay/Source/MassGameplayDebug/Private/AssignDebugVisProcessor.cpp b/Engine/Plugins/Runtime/MassGameplay/Source/MassGameplayDebug/Private/AssignDebugVisProcessor.cpp
--- a/Engine/Plugins/Runtime/MassGameplay/Source/MassGameplayDebug/Private/AssignDebugVisProcessor.cpp
+++ b/Engine/Plugins/Runtime/MassGameplay/Source/MassGameplayDebug/Private/AssignDebugVisProcessor.cpp
@@ -40,15 +40,24 @@ void UAssignDebugVisProcessor::Execute(FMassEntityManager& EntityManager, FMassE
 {
 #if WITH_EDITORONLY_DATA
 	QUICK_SCOPE_CYCLE_COUNTER(AssignDebugVisProcessor_Execute);
-	
+
+	// Visualization relies on an Actor/Component combo added to the world
+	// During world tear down those can be removed,
+	// and we don't want to try to access or recreate them.
+	const UWorld* World = Context.GetWorld();
+	if (World == nullptr
+		|| World->bIsTearingDown)
+	{
+		return;
+	}
+
 	// @todo this code bit is temporary, so is the Visualizer->DirtyVisuals at the end of the function. Will be wrapped in
-	// "executable task" once that's implemented. 
-	UMassDebugVisualizationComponent* Visualizer = nullptr;
+	// "executable task" once that's implemented.
 	UMassDebuggerSubsystem& Debugger = Context.GetMutableSubsystemChecked<UMassDebuggerSubsystem>();
-	check(Debugger.GetVisualizationComponent());
-	Visualizer = Debugger.GetVisualizationComponent();
-	// note that this function will create the "visual components" only it they're missing or out of sync. 
-	Debugger.GetVisualizationComponent()->ConditionallyConstructVisualComponent();
+	UMassDebugVisualizationComponent* Visualizer = Debugger.GetVisualizationComponent();
+	check(Visualizer);
+	// note that this function will create the "visual components" only if they're missing or out of sync.
+	Visualizer->ConditionallyConstructVisualComponent();
 
 	EntityQuery.ForEachEntityChunk(Context, [this](FMassExecutionContext& Context)
 	{
diff --git a/Engine/Plugins/Runtime/MassGameplay/Source/MassGameplayDebug/Private/DebugVisLocationProcessor.cpp b/Engine/Plugins/Runtime/MassGameplay/Source/MassGameplayDebug/Private/DebugVisLocationProcessor.cpp
--- a/Engine/Plugins/Runtime/MassGameplay/Source/MassGameplayDebug/Private/DebugVisLocationProcessor.cpp
+++ b/Engine/Plugins/Runtime/MassGameplay/Source/MassGameplayDebug/Private/DebugVisLocationProcessor.cpp
@@ -67,40 +67,48 @@ void UDebugVisLocationProcessor::Execute(FMassEntityManager& EntityManager, FMas
 
 	if (UE::Mass::Debug::bDebugDrawAllEntities == false)
 	{
-		EntityQuery.ForEachEntityChunk(Context, [this](FMassExecutionContext& Context)
-			{
-				UMassDebuggerSubsystem& Debugger = Context.GetMutableSubsystemChecked<UMassDebuggerSubsystem>();
-				UMassDebugVisualizationComponent* Visualizer = Debugger.GetVisualizationComponent();
-				check(Visualizer);
-				TArrayView<UInstancedStaticMeshComponent* const> VisualDataISMCs = Visualizer->GetVisualDataISMCs();
-				if (VisualDataISMCs.Num() > 0)
+		// MassDebugVisualization relies on an Actor/Component combo added to the world
+		// During world tear down those can be removed,
+		// and we don't want to try to access or recreate them.
+		const UWorld* World = Context.GetWorld();
+		if (World != nullptr
+			&& !World->bIsTearingDown)
+		{
+			EntityQuery.ForEachEntityChunk(Context, [this](FMassExecutionContext& Context)
 				{
-					const int32 NumEntities = Context.GetNumEntities();
-					const TConstArrayView<FTransformFragment> LocationList = Context.GetFragmentView<FTransformFragment>();
-					const TConstArrayView<FSimDebugVisFragment> DebugVisList = Context.GetFragmentView<FSimDebugVisFragment>();
-
-					for (int32 i = 0; i < NumEntities; ++i)
+					UMassDebuggerSubsystem& Debugger = Context.GetMutableSubsystemChecked<UMassDebuggerSubsystem>();
+					UMassDebugVisualizationComponent* Visualizer = Debugger.GetVisualizationComponent();
+					check(Visualizer);
+					TArrayView<UInstancedStaticMeshComponent* const> VisualDataISMCs = Visualizer->GetVisualDataISMCs();
+					if (VisualDataISMCs.Num() > 0)
 					{
-						const FSimDebugVisFragment& VisualComp = DebugVisList[i];
+						const int32 NumEntities = Context.GetNumEntities();
+						const TConstArrayView<FTransformFragment> LocationList = Context.GetFragmentView<FTransformFragment>();
+						const TConstArrayView<FSimDebugVisFragment> DebugVisList = Context.GetFragmentView<FSimDebugVisFragment>();
 
-						// @todo: remove this code once the asset is exported with correct alignment SM_Mannequin.uasset
-						FTransform SMTransform = LocationList[i].GetTransform();
-						FQuat FromEngineToSM(FVector::UpVector, -HALF_PI);
-						SMTransform.SetRotation(FromEngineToSM * SMTransform.GetRotation());
+						for (int32 i = 0; i < NumEntities; ++i)
+						{
+							const FSimDebugVisFragment& VisualComp = DebugVisList[i];
+
+							// @todo: remove this code once the asset is exported with correct alignment SM_Mannequin.uasset
+							FTransform SMTransform = LocationList[i].GetTransform();
+							FQuat FromEngineToSM(FVector::UpVector, -HALF_PI);
+							SMTransform.SetRotation(FromEngineToSM * SMTransform.GetRotation());
 
-						VisualDataISMCs[VisualComp.VisualType]->UpdateInstanceTransform(VisualComp.InstanceIndex, SMTransform, true);
+							VisualDataISMCs[VisualComp.VisualType]->UpdateInstanceTransform(VisualComp.InstanceIndex, SMTransform, true);
+						}
 					}
-				}
-				else
-				{
-					UE_LOG(LogMassDebug, Log, TEXT("UDebugVisLocationProcessor: Trying to update InstanceStaticMeshes while none created. Check your debug visualization setup"));
-				}
-			});
+					else
+					{
+						UE_LOG(LogMassDebug, Log, TEXT("UDebugVisLocationProcessor: Trying to update InstanceStaticMeshes while none created. Check your debug visualization setup"));
+					}
+				});
 
-		UMassDebuggerSubsystem* Debugger = UWorld::GetSubsystem<UMassDebuggerSubsystem>(EntityManager.GetWorld());
-		if (ensure(Debugger))
-		{
-			Debugger->GetVisualizationComponent()->DirtyVisuals();
+			UMassDebuggerSubsystem* Debugger = UWorld::GetSubsystem<UMassDebuggerSubsystem>(EntityManager.GetWorld());
+			if (ensure(Debugger))
+			{
+				Debugger->GetVisualizationComponent()->DirtyVisuals();
+			}
 		}
 	}
 #endif // WITH_EDITORONLY_DATA
diff --git a/Engine/Plugins/Runtime/MassGameplay/Source/MassGameplayDebug/Private/MassDebuggerSubsystem.cpp b/Engine/Plugins/Runtime/MassGameplay/Source/MassGameplayDebug/Private/MassDebuggerSubsystem.cpp
--- a/Engine/Plugins/Runtime/MassGameplay/Source/MassGameplayDebug/Private/MassDebuggerSubsystem.cpp
+++ b/Engine/Plugins/Runtime/MassGameplay/Source/MassGameplayDebug/Private/MassDebuggerSubsystem.cpp
@@ -88,7 +88,9 @@ UMassDebugVisualizationComponent* UMassDebuggerSubsystem::GetVisualizationCompon
 		}
 		else
 		{
-			if (UWorld* World = GetWorld())
+			UWorld* World = GetWorld();
+			if (World != nullptr
+				&& !World->bIsTearingDown)
 			{
 				const AMassDebugVisualizer& VisualizerActor = GetOrSpawnDebugVisualizer(*World);
 				VisualizationComponent = &VisualizerActor.GetDebugVisComponent();
