[MassGameplayDebug] made AssignDebugVisProcessor editor only to prevent useless registrations on non-editor builds
Also made sure it is properly setup for registration when WITH_MASSGAMEPLAY_DEBUG i s0

#Synced-CL 50989336

diff --git a/Engine/Plugins/Runtime/MassGameplay/Source/MassGameplayDebug/Private/AssignDebugVisProcessor.cpp b/Engine/Plugins/Runtime/MassGameplay/Source/MassGameplayDebug/Private/AssignDebugVisProcessor.cpp
--- a/Engine/Plugins/Runtime/MassGameplay/Source/MassGameplayDebug/Private/AssignDebugVisProcessor.cpp
+++ b/Engine/Plugins/Runtime/MassGameplay/Source/MassGameplayDebug/Private/AssignDebugVisProcessor.cpp
@@ -1,6 +1,8 @@
 // Copyright Epic Games, Inc. All Rights Reserved.
 
 #include "AssignDebugVisProcessor.h"
+
+#if WITH_EDITORONLY_DATA
 #include UE_INLINE_GENERATED_CPP_BY_NAME(AssignDebugVisProcessor)
 
 #if WITH_MASSGAMEPLAY_DEBUG
@@ -20,10 +22,13 @@ UAssignDebugVisProcessor::UAssignDebugVisProcessor()
 {
 	bAutoRegisterWithProcessingPhases = false;
 
-#if WITH_MASSGAMEPLAY_DEBUG
-	bRequiresGameThreadExecution = true; // due to UMassDebuggerSubsystem
+	// Always set the observed type and operations even if WITH_MASSGAMEPLAY_DEBUG is 0
+	// since the processor may still be registered even if it won't do anything.
 	ObservedType = FSimDebugVisFragment::StaticStruct();
 	ObservedOperations = EMassObservedOperationFlags::Add;
+
+#if WITH_MASSGAMEPLAY_DEBUG
+	bRequiresGameThreadExecution = true; // due to UMassDebuggerSubsystem
 #endif // WITH_MASSGAMEPLAY_DEBUG
 }
 
@@ -38,7 +43,6 @@ void UAssignDebugVisProcessor::ConfigureQueries(const TSharedRef<FMassEntityMana
 
 void UAssignDebugVisProcessor::Execute(FMassEntityManager& EntityManager, FMassExecutionContext& Context)
 {
-#if WITH_EDITORONLY_DATA
 	QUICK_SCOPE_CYCLE_COUNTER(AssignDebugVisProcessor_Execute);
 
 	// Visualization relies on an Actor/Component combo added to the world
@@ -80,6 +84,7 @@ void UAssignDebugVisProcessor::Execute(FMassEntityManager& EntityManager, FMassE
 	{
 		Visualizer->DirtyVisuals();
 	}
-#endif // WITH_EDITORONLY_DATA
 }
 #endif // WITH_MASSGAMEPLAY_DEBUG
+
+#endif // WITH_EDITORONLY_DATA
\ No newline at end of file
diff --git a/Engine/Plugins/Runtime/MassGameplay/Source/MassGameplayDebug/Public/AssignDebugVisProcessor.h b/Engine/Plugins/Runtime/MassGameplay/Source/MassGameplayDebug/Public/AssignDebugVisProcessor.h
--- a/Engine/Plugins/Runtime/MassGameplay/Source/MassGameplayDebug/Public/AssignDebugVisProcessor.h
+++ b/Engine/Plugins/Runtime/MassGameplay/Source/MassGameplayDebug/Public/AssignDebugVisProcessor.h
@@ -2,6 +2,9 @@
 
 #pragma once
 
+#include "Misc/CoreMiscDefines.h"
+
+#if WITH_EDITORONLY_DATA
 #include "MassCommonTypes.h"
 #include "MassEntityQuery.h"
 #include "MassObserverProcessor.h"
@@ -30,3 +33,4 @@ protected:
 };
 
 #undef UE_API
+#endif // #if WITH_EDITORONLY_DATA
