[Mass] fixed MassReplication android non-unity compilation by extacting IClientBubbleHandlerInterface into a separate file to include in MassReplicationTemplates.h

#rnx

#Synced-CL 51090036

diff --git a/Engine/Plugins/Runtime/MassGameplay/Source/MassReplication/Public/MassClientBubbleHandler.h b/Engine/Plugins/Runtime/MassGameplay/Source/MassReplication/Public/MassClientBubbleHandler.h
--- a/Engine/Plugins/Runtime/MassGameplay/Source/MassReplication/Public/MassClientBubbleHandler.h
+++ b/Engine/Plugins/Runtime/MassGameplay/Source/MassReplication/Public/MassClientBubbleHandler.h
@@ -14,6 +14,7 @@
 #include "MassEntityTemplate.h"
 #include "MassEntityView.h"
 #include "MassEntityQuery.h"
+#include "MassClientBubbleHandlerInterface.h"
 #if UE_ENABLE_INCLUDE_ORDER_DEPRECATED_IN_5_6
 #include "MassEntityManager.h"
 #include "MassSpawnerTypes.h"
@@ -98,36 +99,6 @@ struct FMassAgentRemoveData
 	double TimeLastRemoved = 0.;
 };
 
-/**
- * Interface for the bubble handler classes. All the outside interaction with the FastArray logic should be done via the Handler interface
- * or derived classes where possible.
- * These virtual functions are either only called once each per frame on the client for a few struct instances
- * or called at startup / shutdown.
- */
-class IClientBubbleHandlerInterface
-{
-public:
-	virtual ~IClientBubbleHandlerInterface() = default;
-
-	virtual void InitializeForWorld(UWorld& InWorld) = 0;
-
-#if UE_REPLICATION_COMPILE_CLIENT_CODE
-	/** These functions are processed internally by TClientBubbleHandlerBase */
-	virtual void PreReplicatedRemove(const TArrayView<int32> RemovedIndices, int32 FinalSize) = 0;
-	virtual void PostReplicatedAdd(const TArrayView<int32> AddedIndices, int32 FinalSize) = 0;
-	virtual void PostReplicatedChange(const TArrayView<int32> ChangedIndices, int32 FinalSize) = 0;
-#endif // UE_REPLICATION_COMPILE_CLIENT_CODE
-
-	virtual void Reset() = 0;
-	virtual void UpdateAgentsToRemove() = 0;
-
-	virtual void Tick(float DeltaTime) = 0;
-	virtual void SetClientHandle(FMassClientHandle InClientHandle) = 0;
-
-	virtual void DebugValidateBubbleOnServer() = 0;
-	virtual void DebugValidateBubbleOnClient() = 0;
-};
-
 
 /** 
  * Template client bubble functionality. Replication logic for specific agent types is provided by deriving from this class.
diff --git a/Engine/Plugins/Runtime/MassGameplay/Source/MassReplication/Public/MassClientBubbleHandlerInterface.h b/Engine/Plugins/Runtime/MassGameplay/Source/MassReplication/Public/MassClientBubbleHandlerInterface.h
new file mode 100644
--- /dev/null
+++ b/Engine/Plugins/Runtime/MassGameplay/Source/MassReplication/Public/MassClientBubbleHandlerInterface.h
@@ -0,0 +1,38 @@
+// Copyright Epic Games, Inc. All Rights Reserved.
+
+#pragma once
+
+#include "MassReplicationTypes.h"
+#include "Containers/ArrayView.h"
+
+class UWorld;
+
+/**
+ * Interface for the bubble handler classes. All the outside interaction with the FastArray logic should be done via the Handler interface
+ * or derived classes where possible.
+ * These virtual functions are either only called once each per frame on the client for a few struct instances
+ * or called at startup / shutdown.
+ */
+class IClientBubbleHandlerInterface
+{
+public:
+	virtual ~IClientBubbleHandlerInterface() = default;
+
+	virtual void InitializeForWorld(UWorld& InWorld) = 0;
+
+#if UE_REPLICATION_COMPILE_CLIENT_CODE
+	/** These functions are processed internally by TClientBubbleHandlerBase */
+	virtual void PreReplicatedRemove(const TArrayView<int32> RemovedIndices, int32 FinalSize) = 0;
+	virtual void PostReplicatedAdd(const TArrayView<int32> AddedIndices, int32 FinalSize) = 0;
+	virtual void PostReplicatedChange(const TArrayView<int32> ChangedIndices, int32 FinalSize) = 0;
+#endif // UE_REPLICATION_COMPILE_CLIENT_CODE
+
+	virtual void Reset() = 0;
+	virtual void UpdateAgentsToRemove() = 0;
+
+	virtual void Tick(float DeltaTime) = 0;
+	virtual void SetClientHandle(FMassClientHandle InClientHandle) = 0;
+
+	virtual void DebugValidateBubbleOnServer() = 0;
+	virtual void DebugValidateBubbleOnClient() = 0;
+};
diff --git a/Engine/Plugins/Runtime/MassGameplay/Source/MassReplication/Public/MassReplicationTemplates.h b/Engine/Plugins/Runtime/MassGameplay/Source/MassReplication/Public/MassReplicationTemplates.h
--- a/Engine/Plugins/Runtime/MassGameplay/Source/MassReplication/Public/MassReplicationTemplates.h
+++ b/Engine/Plugins/Runtime/MassGameplay/Source/MassReplication/Public/MassReplicationTemplates.h
@@ -5,6 +5,7 @@
 #include "MassEntityView.h"
 #include "MassExecutionContext.h"
 #include "MassReplicationTypes.h"
+#include "MassClientBubbleHandlerInterface.h"
 #include <tuple>
 
 /**
