[Mass] Fixed TClientBubbleHandlerBase<AgentArrayItem>::PostReplicatedAddEntitiesHelper crashing if initializers or observers caused just-spawned entities to change archetypes.

#Synced-CL 51090036

diff --git a/Engine/Plugins/Runtime/MassGameplay/Source/MassReplication/Public/MassClientBubbleHandler.h b/Engine/Plugins/Runtime/MassGameplay/Source/MassReplication/Public/MassClientBubbleHandler.h
--- a/Engine/Plugins/Runtime/MassGameplay/Source/MassReplication/Public/MassClientBubbleHandler.h
+++ b/Engine/Plugins/Runtime/MassGameplay/Source/MassReplication/Public/MassClientBubbleHandler.h
@@ -587,11 +587,6 @@ void TClientBubbleHandlerBase<AgentArrayItem>::PostReplicatedAddEntitiesHelper(c
 
 		SpawnerSubsystem->SpawnEntities(TemplateID, AgentsSpawn.Num(), FStructView(), TSubclassOf<UMassProcessor>(), Entities);
 
-		const FMassEntityTemplate* MassEntityTemplate = SpawnerSubsystem->GetMassEntityTemplate(TemplateID);
-		check(MassEntityTemplate);
-		const FMassArchetypeHandle& ArchetypeHandle = MassEntityTemplate->GetArchetype();
-
-
 		FMassExecutionContext ExecContext(EntityManager);
 		FMassEntityQuery Query(EntityManager.AsShared());
 
@@ -599,13 +594,22 @@ void TClientBubbleHandlerBase<AgentArrayItem>::PostReplicatedAddEntitiesHelper(c
 
 		Query.AddRequirement<FMassNetworkIDFragment>(EMassFragmentAccess::ReadWrite);
 
+		// if there were any initializer or observer processors it's possible that the spawned entities
+		// are no longer in the archetype indicated by the template. We need to create the entity collections
+		// the slower way (i.e. assuming we don't know the current archetype)
+		// Note that we're not using CreationContext->GetEntityCollections (with CreationContext being
+		// the return value of SpawnerSubsystem->SpawnEntities) because it might contain unrelated entities.
+		TArray<FMassArchetypeEntityCollection> EntityCollections;
+		UE::Mass::Utils::CreateEntityCollections(EntityManager, Entities, FMassArchetypeEntityCollection::NoDuplicates, EntityCollections);
+
+		// Note that the way this functionality is implemented here it's not possible to parallelize the lambda's
+		// execution, due to using AgentsSpawnIdx. If attempt is made at parallelizing the work this detail needs to be changed.
 		int32 AgentsSpawnIdx = 0;
 
-		Query.ForEachEntityChunk(FMassArchetypeEntityCollection(ArchetypeHandle, Entities, FMassArchetypeEntityCollection::NoDuplicates)
-								, ExecContext, [&AgentsSpawn, &AgentsSpawnIdx, this, ReplicationSubsystem, &ExecContext, &EventHandlers
-								, &EntityManager](FMassExecutionContext& Context)
+		Query.ForEachEntityChunkInCollections(EntityCollections, ExecContext
+			, [&AgentsSpawn, &AgentsSpawnIdx, this, ReplicationSubsystem, &ExecContext, &EventHandlers](FMassExecutionContext& Context)
 			{
-				EventHandlers.CacheFragmentViewsForSpawnQuery(ExecContext);
+				EventHandlers.CacheFragmentViewsForSpawnQuery(Context);
 
 				const TArrayView<FMassNetworkIDFragment> NetworkIDList = Context.GetMutableFragmentView<FMassNetworkIDFragment>();
 
@@ -619,7 +623,7 @@ void TClientBubbleHandlerBase<AgentArrayItem>::PostReplicatedAddEntitiesHelper(c
 					NetIDFragment.NetID = AgentSpawn.GetNetID();
 					ReplicationSubsystem->SetEntity(NetIDFragment.NetID, Entity);
 
-					FMassEntityView EntityView(EntityManager, Entity);
+					FMassEntityView EntityView(Context.GetEntityManagerChecked(), Entity);
 					EventHandlers.SetSpawnedEntityData(EntityView, AgentSpawn, EntityIt);
 
 					++AgentsSpawnIdx;
