[Mass] Replication: added spawn query caching to avoid recreating it every time. The recreation was preventing the efficient use of cached archetypes, which was affecting performance.

#Synced-CL 51157439

diff --git a/Engine/Plugins/Runtime/MassGameplay/Source/MassReplication/Public/MassClientBubbleHandler.h b/Engine/Plugins/Runtime/MassGameplay/Source/MassReplication/Public/MassClientBubbleHandler.h
--- a/Engine/Plugins/Runtime/MassGameplay/Source/MassReplication/Public/MassClientBubbleHandler.h
+++ b/Engine/Plugins/Runtime/MassGameplay/Source/MassReplication/Public/MassClientBubbleHandler.h
@@ -166,10 +166,12 @@ protected:
 	 *		CacheFragmentViewsForSpawnQuery;
 	 *		SetSpawnedEntityData;
 	 *		SetModifiedEntityData;
+	 *		GetSpawnQuery -> FMassEntityQuery;
 	 * }
 	 * Note:
 	 *	- PostReplicatedAddEntitiesHelper doesn't call SetModifiedEntityData
 	 *	- PostReplicatedChangeHelper only calls SetModifiedEntityData
+	 *	- it's recommended to supply GetSpawnQuery by deriving from UE::Mass::Replication::TCachedSpawnQuery
 	 */
 
 	/** Called from TClientBubbleHandlerBase derived classes in PostReplicatedAd() */
@@ -534,8 +536,9 @@ template<typename AgentArrayItem>
 void TClientBubbleHandlerBase<AgentArrayItem>::PostReplicatedAddHelper(const TArrayView<int32> AddedIndices, FAddRequirementsForSpawnQueryFunction AddRequirementsForSpawnQuery
 	, FCacheFragmentViewsForSpawnQueryFunction CacheFragmentViewsForSpawnQuery, FSetSpawnedEntityDataFunction SetSpawnedEntityData, FSetModifiedEntityDataFunction SetModifiedEntityData)
 {
-	struct FEventHandlers
+	struct FEventHandlers : UE::Mass::Replication::TCachedSpawnQuery<FEventHandlers>
 	{
+		using UE::Mass::Replication::TCachedSpawnQuery<FEventHandlers>::GetSpawnQuery;
 		FAddRequirementsForSpawnQueryFunction AddRequirementsForSpawnQuery;
 		FCacheFragmentViewsForSpawnQueryFunction CacheFragmentViewsForSpawnQuery;
 		FSetSpawnedEntityDataFunction SetSpawnedEntityData;
@@ -577,6 +580,9 @@ void TClientBubbleHandlerBase<AgentArrayItem>::PostReplicatedAddEntitiesHelper(c
 		AgentsArray.Add(&Agent);
 	}
 
+	FMassExecutionContext ExecContext(EntityManager);
+	FMassEntityQuery& Query = EventHandlers.GetSpawnQuery(EntityManager.AsShared());
+
 	// Batch spawn per FMassEntityTemplateID
 	for (const TPair<FMassEntityTemplateID, TArray<typename AgentArrayItem::FReplicatedAgentType*>>& Item : AgentsSpawnMap)
 	{
@@ -587,13 +593,6 @@ void TClientBubbleHandlerBase<AgentArrayItem>::PostReplicatedAddEntitiesHelper(c
 
 		SpawnerSubsystem->SpawnEntities(TemplateID, AgentsSpawn.Num(), FStructView(), TSubclassOf<UMassProcessor>(), Entities);
 
-		FMassExecutionContext ExecContext(EntityManager);
-		FMassEntityQuery Query(EntityManager.AsShared());
-
-		EventHandlers.AddRequirementsForSpawnQuery(Query);
-
-		Query.AddRequirement<FMassNetworkIDFragment>(EMassFragmentAccess::ReadWrite);
-
 		// if there were any initializer or observer processors it's possible that the spawned entities
 		// are no longer in the archetype indicated by the template. We need to create the entity collections
 		// the slower way (i.e. assuming we don't know the current archetype)
@@ -636,8 +635,9 @@ template<typename AgentArrayItem>
 void TClientBubbleHandlerBase<AgentArrayItem>::PostReplicatedAddEntitiesHelper(const TArrayView<int32> AddedIndices, FAddRequirementsForSpawnQueryFunction AddRequirementsForSpawnQuery
 	, FCacheFragmentViewsForSpawnQueryFunction CacheFragmentViewsForSpawnQuery, FSetSpawnedEntityDataFunction SetSpawnedEntityData)
 {
-	struct FEventHandlers
+	struct FEventHandlers : UE::Mass::Replication::TCachedSpawnQuery<FEventHandlers>
 	{
+		using UE::Mass::Replication::TCachedSpawnQuery<FEventHandlers>::GetSpawnQuery;
 		FAddRequirementsForSpawnQueryFunction AddRequirementsForSpawnQuery;
 		FCacheFragmentViewsForSpawnQueryFunction CacheFragmentViewsForSpawnQuery;
 		FSetSpawnedEntityDataFunction SetSpawnedEntityData;
diff --git a/Engine/Plugins/Runtime/MassGameplay/Source/MassReplication/Public/MassReplicationTemplates.h b/Engine/Plugins/Runtime/MassGameplay/Source/MassReplication/Public/MassReplicationTemplates.h
--- a/Engine/Plugins/Runtime/MassGameplay/Source/MassReplication/Public/MassReplicationTemplates.h
+++ b/Engine/Plugins/Runtime/MassGameplay/Source/MassReplication/Public/MassReplicationTemplates.h
@@ -6,6 +6,8 @@
 #include "MassExecutionContext.h"
 #include "MassReplicationTypes.h"
 #include "MassClientBubbleHandlerInterface.h"
+#include "MassEntityQuery.h"
+#include "MassReplicationFragments.h"
 #include <tuple>
 
 /**
@@ -411,15 +413,44 @@ namespace UE::Mass::Replication
 		std::tuple<TFragmentHandlers...> FragmentHandlers;
 	};
 
+	/**
+	 * Helper struct to be used as CRTP to supply types with SpawnQuery creation and caching functionality
+	 * @param THandlers anything that implements AddRequirementsForSpawnQuery
+	 */
+	template<typename THandlers>
+	struct TCachedSpawnQuery
+	{
+#if UE_REPLICATION_COMPILE_CLIENT_CODE
+		/**
+		 * Returns the cached SpawnQuery. If it's not initialized yet EntityManager will be used for the initialization,
+		 * along with AddRequirementsForSpawnQuery member function
+		 */
+		FMassEntityQuery& GetSpawnQuery(const TSharedRef<FMassEntityManager>& EntityManager)
+		{
+			if (!SpawnQuery.IsInitialized())
+			{
+				SpawnQuery.Initialize(EntityManager);
+				static_cast<THandlers*>(this)->AddRequirementsForSpawnQuery(SpawnQuery);
+				SpawnQuery.AddRequirement<FMassNetworkIDFragment>(EMassFragmentAccess::ReadWrite);
+			}
+			return SpawnQuery;
+		}
+	private:
+		/** Entity query used to configure entities added to the client */
+		FMassEntityQuery SpawnQuery;
+#endif // UE_REPLICATION_COMPILE_CLIENT_CODE
+	};
+
 	/**
 	 * Client-side handler pack for multiple replicated data types.
 	 * Orchestrates multiple TClientBubbleFragmentHandler instances.
 	 */
 	template<typename... TFragmentHandlers>
-	struct TClientBubbleDataHandlerPack
+	struct TClientBubbleDataHandlerPack : TCachedSpawnQuery<TClientBubbleDataHandlerPack<TFragmentHandlers...>>
 	{
 		template<typename TDataHandler>
 		using TReplicatedData = typename TRemoveReference<TDataHandler>::Type::FReplicatedDataType;
+		using TCachedSpawnQuery<TClientBubbleDataHandlerPack>::GetSpawnQuery;
 
 		static void AddRequirementsForSpawnQuery(FMassEntityQuery& InQuery)
 		{
