[Mass] Fixed MassReplicationProcessor to properly express its dependence on UMassReplicationSubsystem to avoid running in parallel with other processors using the system in RW mode.

#rnx

#Synced-CL 51157439

diff --git a/Engine/Plugins/Runtime/MassGameplay/Source/MassReplication/Private/MassReplicationProcessor.cpp b/Engine/Plugins/Runtime/MassGameplay/Source/MassReplication/Private/MassReplicationProcessor.cpp
--- a/Engine/Plugins/Runtime/MassGameplay/Source/MassReplication/Private/MassReplicationProcessor.cpp
+++ b/Engine/Plugins/Runtime/MassGameplay/Source/MassReplication/Private/MassReplicationProcessor.cpp
@@ -68,6 +68,7 @@ void UMassReplicationProcessor::ConfigureQueries(const TSharedRef<FMassEntityMan
 	EntityQuery.AddSharedRequirement<FMassReplicationSharedFragment>(EMassFragmentAccess::ReadWrite);
 
 	ProcessorRequirements.AddSubsystemRequirement<UMassLODSubsystem>(EMassFragmentAccess::ReadOnly);
+	ProcessorRequirements.AddSubsystemRequirement<UMassReplicationSubsystem>(EMassFragmentAccess::ReadWrite);
 }
 
 void UMassReplicationProcessor::InitializeInternal(UObject& Owner, const TSharedRef<FMassEntityManager>& EntityManager)
