Fix assert when UMassActorEditorSubsystem attempts to access UMassEntityEditorSubsystem
- Despite being an explicit dependency declared in UBT, the module containing UMassEntityEditorSubsystem is not guaranteed to be loaded.
- In this particular case, when `DisableEnginePluginsByDefault` is set no other module ends up explicitly loading the module and the subsystem is never initialized.

#Synced-CL 51481056

diff --git a/Engine/Plugins/Runtime/MassGameplay/Source/MassGameplayEditor/Private/MassGameplayEditorModule.cpp b/Engine/Plugins/Runtime/MassGameplay/Source/MassGameplayEditor/Private/MassGameplayEditorModule.cpp
--- a/Engine/Plugins/Runtime/MassGameplay/Source/MassGameplayEditor/Private/MassGameplayEditorModule.cpp
+++ b/Engine/Plugins/Runtime/MassGameplay/Source/MassGameplayEditor/Private/MassGameplayEditorModule.cpp
@@ -17,6 +17,10 @@ void FMassGameplayEditorModule::StartupModule()
 	RegisterSectionMappings();
 	
 	PropertyModule.NotifyCustomizationModuleChanged();
+
+	// Ensure dependencies are loaded. We rely on subsystems within them.
+	FModuleManager& ModuleManager = FModuleManager::Get();
+	ModuleManager.LoadModuleChecked("MassEntityEditor");
 }
 
 void FMassGameplayEditorModule::ShutdownModule()
