[Mass] Document re-registration edge case in MassAgentSubsystem.cpp per [REDACTED_HANDLE]

Replaced the [REDACTED_HANDLE] comment with a detailed explanation of why RegisterAgentComponent bails out when called on a non-puppet actor that still has a valid TemplateID from a previous registration cycle. This edge case occurs when an actor with a UMassAgentComponent gets unregistered and re-registered.

#rnx
#Synced-CL 51621734

diff --git a/Engine/Plugins/Runtime/MassGameplay/Source/MassActors/Private/MassAgentSubsystem.cpp b/Engine/Plugins/Runtime/MassGameplay/Source/MassActors/Private/MassAgentSubsystem.cpp
--- a/Engine/Plugins/Runtime/MassGameplay/Source/MassActors/Private/MassAgentSubsystem.cpp
+++ b/Engine/Plugins/Runtime/MassGameplay/Source/MassActors/Private/MassAgentSubsystem.cpp
@@ -84,7 +84,13 @@ FMassEntityTemplateID UMassAgentSubsystem::RegisterAgentComponent(UMassAgentComp
 		return AgentComp.GetTemplateID();
 	}
 
-	// @todo note that this will happen when a regular actor with AgentComp gets unregister and then registered again.
+	// This path is hit when a regular (non-puppet) actor that owns a UMassAgentComponent gets
+	// unregistered (e.g. via EndPlay or RemoveFromWorld) and then registered again (e.g. re-added
+	// to the world or re-activated). In that scenario the component still carries a valid
+	// TemplateID from its previous registration. We intentionally bail out with a warning rather
+	// than re-register, because the old Mass entity handle may still be live and double-registration
+	// would create orphaned entity data. Callers that need a fresh registration should first ensure
+	// the component has been fully unregistered (TemplateID invalidated) before calling this again.
 	if (AgentComp.GetTemplateID().IsValid())
 	{
 		UE_VLOG_UELOG(this, LogMassActor, Warning, TEXT("UMassAgentSubsystem::RegisterAgentComponent called while the given agent component has already been registered (Owner: %s, entity handle %s, template ID %s)")
