[Mass] Use FObjectKey for PooledActors map key in MassActorSpawnerSubsystem

#rnx
#Synced-CL 51672599

diff --git a/Engine/Plugins/Runtime/MassGameplay/Source/MassActors/Public/MassActorSpawnerSubsystem.h b/Engine/Plugins/Runtime/MassGameplay/Source/MassActors/Public/MassActorSpawnerSubsystem.h
--- a/Engine/Plugins/Runtime/MassGameplay/Source/MassActors/Public/MassActorSpawnerSubsystem.h
+++ b/Engine/Plugins/Runtime/MassGameplay/Source/MassActors/Public/MassActorSpawnerSubsystem.h
@@ -10,6 +10,7 @@
 #endif // UE_ENABLE_INCLUDE_ORDER_DEPRECATED_IN_5_6
 #include "Templates/SubclassOf.h"
 #include "GameFramework/Actor.h"
+#include "UObject/ObjectKey.h"
 #include "MassSubsystemBase.h"
 #include "MassEntityHandle.h"
 #include "StructUtils/StructView.h"
@@ -317,8 +318,7 @@ protected:
 
 	bool bActorPoolingEnabled = true;
 
-	// @todo should use FObjetKey for the map key.
-	TMap<TSubclassOf<AActor>, TArray<TObjectPtr<AActor>>> PooledActors;
+	TMap<FObjectKey, TArray<TObjectPtr<AActor>>> PooledActors;
 
 	FMassEntityHandleManager_ActorSpawnRequest SpawnRequestHandleManager;
 
