[Mass][ISKM] Add Basic Instanced Skinned Mesh Support To MassRepresentation
- Added New Enum for Skinned Representation
- Basic Mesh Selection
- New Processor For ISKM Updates
- New Animation Processor and Manger Class for Inherited Functionality

#rnx
#Synced-CL 51901791

diff --git a/Engine/Plugins/AI/MassCrowd/Source/MassCrowd/Private/MassCrowdVisualizationLODProcessor.cpp b/Engine/Plugins/AI/MassCrowd/Source/MassCrowd/Private/MassCrowdVisualizationLODProcessor.cpp
--- a/Engine/Plugins/AI/MassCrowd/Source/MassCrowd/Private/MassCrowdVisualizationLODProcessor.cpp
+++ b/Engine/Plugins/AI/MassCrowd/Source/MassCrowd/Private/MassCrowdVisualizationLODProcessor.cpp
@@ -8,6 +8,7 @@
 #include "MassDebugger.h"
 #include "MassEntityManager.h"
 #include "MassRepresentationFragments.h"
+#include "MassRepresentationTypes.h"
 
 #include UE_INLINE_GENERATED_CPP_BY_NAME(MassCrowdVisualizationLODProcessor)
 
@@ -94,7 +95,7 @@ void UMassCrowdVisualizationLODProcessor::Execute(FMassEntityManager& EntityMana
 				{
 					const FMassRepresentationFragment& RepresentationFragment = RepresentationFragmentList[EntityIt];
 					const FMassViewerInfoFragment& LODInfoFragment = LODInfoFragmentList[EntityIt];
-					if (RepresentationFragment.CurrentRepresentation == EMassRepresentationType::StaticMeshInstance && SpecifiedRangeSquaredCentimeters > LODInfoFragment.ClosestViewerDistanceSq)
+					if (UE::Mass::Representation::IsValidMeshRepresentation(RepresentationFragment.CurrentRepresentation) && SpecifiedRangeSquaredCentimeters > LODInfoFragment.ClosestViewerDistanceSq)
 					{
 						const FTransformFragment& EntityLocation = LocationList[EntityIt];
 						DrawDebugSolidBox(DebugDrawingWorld, EntityLocation.GetTransform().GetLocation() + FVector(0.0f, 0.0f, 150.0f), FVector(50.0f), FColor::Red);
diff --git a/Engine/Plugins/Runtime/InstancedActors/Source/InstancedActors/Private/InstancedActorsDebugProcessor.cpp b/Engine/Plugins/Runtime/InstancedActors/Source/InstancedActors/Private/InstancedActorsDebugProcessor.cpp
--- a/Engine/Plugins/Runtime/InstancedActors/Source/InstancedActors/Private/InstancedActorsDebugProcessor.cpp
+++ b/Engine/Plugins/Runtime/InstancedActors/Source/InstancedActors/Private/InstancedActorsDebugProcessor.cpp
@@ -51,14 +51,14 @@ namespace UE::InstancedActors::Debug
 			{TEXT("IA.debug.PreviousRepresentation"), bDebugDrawPrevRepresentation
 				, TEXT("When enabled (the default value) will debug draw data related to previous "
 				"representation of instanced actors being drawn. Green indicates HighResSpawnedActor, DarkOliveGreen "
-				"indicates LowResSpawnedActor, Magenta indicates StaticMeshInstance and Red means Off."
+				"indicates LowResSpawnedActor, Magenta indicates StaticMeshInstance, Purple indicates SkinnedMeshInstance and Red means Off."
 				"Note that the data will be drawn only if Previous Representation differs from the Current Representation"), ECVF_Cheat},
 			{TEXT("IA.debug.CurrentRepresentationAll"), bDebugDrawAllEntities
 				, TEXT("When enabled will debug draw data related to current representation of all instanced "
 				"actors.\nFor instanced at Detailed batch LOD level Green indicates HighResSpawnedActor, DarkOliveGreen "
-				"indicates LowResSpawnedActor, Magenta indicates StaticMeshInstance and Red means Off."
+				"indicates LowResSpawnedActor, Magenta indicates StaticMeshInstance, Purple indicates SkinnedMeshInstance and Red means Off."
 				"\nFor instanced not at Detailed level: Yellow indicates HighResSpawnedActor, Orange indicates "
-				"LowResSpawnedActor, NeonPink indicates StaticMeshInstance and Red means Off."), ECVF_Cheat},
+				"LowResSpawnedActor, NeonPink indicates StaticMeshInstance, CadetBlue indicates SkinnedMeshInstance, and Red means Off."), ECVF_Cheat},
 	};
 	}
 
@@ -74,6 +74,9 @@ namespace UE::InstancedActors::Debug
 		case EMassRepresentationType::LowResSpawnedActor:
 			DrawColor = FColorList::DarkOliveGreen;
 			break;
+		case EMassRepresentationType::SkinnedMeshInstance:
+			DrawColor = FColor::Purple;
+			break;
 		case EMassRepresentationType::StaticMeshInstance:
 			DrawColor = FColor::Magenta;
 			break;
@@ -93,6 +96,9 @@ namespace UE::InstancedActors::Debug
 		case EMassRepresentationType::LowResSpawnedActor:
 			DrawColor = FColorList::Orange;
 			break;
+		case EMassRepresentationType::SkinnedMeshInstance:
+			DrawColor = FColorList::CadetBlue;
+			break;
 		case EMassRepresentationType::StaticMeshInstance:
 			DrawColor = FColorList::NeonPink;
 			break;
diff --git a/Engine/Plugins/Runtime/InstancedActors/Source/InstancedActors/Private/InstancedActorsRepresentationActorManagement.cpp b/Engine/Plugins/Runtime/InstancedActors/Source/InstancedActors/Private/InstancedActorsRepresentationActorManagement.cpp
--- a/Engine/Plugins/Runtime/InstancedActors/Source/InstancedActors/Private/InstancedActorsRepresentationActorManagement.cpp
+++ b/Engine/Plugins/Runtime/InstancedActors/Source/InstancedActors/Private/InstancedActorsRepresentationActorManagement.cpp
@@ -191,6 +191,7 @@ AActor* UInstancedActorsRepresentationActorManagement::FindOrInstantlySpawnActor
 		RepresentationLOD.LOD = EMassLOD::High;
 
 		// ALSO: need to hide the ISM instance
+		// @todo Duplicated Code could be solved with a "trait" pattern - collect all relevant classes in one declaration per mesh-type and use that. 
 		if (Representation.CurrentRepresentation == EMassRepresentationType::StaticMeshInstance)
 		{
 			FMassInstancedStaticMeshInfoArrayView ISMInfosView = RepresentationSubsystem.GetMutableInstancedStaticMeshInfos();
@@ -204,6 +205,19 @@ AActor* UInstancedActorsRepresentationActorManagement::FindOrInstantlySpawnActor
 				ISMInfo.RemoveInstance(EntityView.GetEntity(), RepresentationLOD.LODSignificance);
 			}
 		}
+		else if (Representation.CurrentRepresentation == EMassRepresentationType::SkinnedMeshInstance)
+		{
+			FMassInstancedSkinnedMeshInfoArrayView SkinnedMeshInfosView = RepresentationSubsystem.GetMutableInstancedSkinnedMeshInfos();
+
+			const int32 SkinnedMeshInfoIndex = Representation.SkinnedMeshDescHandle.ToIndex();
+			if (ensureMsgf(SkinnedMeshInfosView.IsValidIndex(SkinnedMeshInfoIndex), TEXT("Invalid handle index %u for SkinnedMeshInfosView"), SkinnedMeshInfoIndex))
+			{
+				FMassInstancedSkinnedMeshInfo& SkinnedMeshInfo = SkinnedMeshInfosView[SkinnedMeshInfoIndex];
+
+				SkinnedMeshInfo.RemoveInstance(EntityView.GetEntity(), Representation.PrevLODSignificance);
+				SkinnedMeshInfo.RemoveInstance(EntityView.GetEntity(), RepresentationLOD.LODSignificance);
+			}
+		}
 
 		Representation.PrevRepresentation = Representation.CurrentRepresentation;
 		Representation.CurrentRepresentation = EMassRepresentationType::HighResSpawnedActor;
diff --git a/Engine/Plugins/Runtime/InstancedActors/Source/InstancedActors/Private/InstancedActorsVisualizationSwitcherProcessor.cpp b/Engine/Plugins/Runtime/InstancedActors/Source/InstancedActors/Private/InstancedActorsVisualizationSwitcherProcessor.cpp
--- a/Engine/Plugins/Runtime/InstancedActors/Source/InstancedActors/Private/InstancedActorsVisualizationSwitcherProcessor.cpp
+++ b/Engine/Plugins/Runtime/InstancedActors/Source/InstancedActors/Private/InstancedActorsVisualizationSwitcherProcessor.cpp
@@ -41,6 +41,7 @@ void UInstancedActorsVisualizationSwitcherProcessor::Execute(FMassEntityManager&
 		check(RepresentationSubsystem);
 		checkSlow(RepresentationSubsystem->IsA<UInstancedActorsRepresentationSubsystem>());
 		FMassInstancedStaticMeshInfoArrayView ISMInfosView = RepresentationSubsystem->GetMutableInstancedStaticMeshInfos();
+		FMassInstancedSkinnedMeshInfoArrayView MeshInfosView = RepresentationSubsystem->GetMutableInstancedSkinnedMeshInfos();
 
 		TConstArrayView<FInstancedActorsMeshSwitchFragment> MeshSwitchFragments = Context.GetFragmentView<FInstancedActorsMeshSwitchFragment>();
 		TArrayView<FMassRepresentationFragment> RepresentationFragments = Context.GetMutableFragmentView<FMassRepresentationFragment>();
@@ -52,6 +53,7 @@ void UInstancedActorsVisualizationSwitcherProcessor::Execute(FMassEntityManager&
 			FMassRepresentationFragment& RepresentationFragment = RepresentationFragments[EntityIt];
 
 			SwitchEntityMeshDesc(ISMInfosView, RepresentationFragment, EntityHandle, MeshSwitchFragment.NewStaticMeshDescHandle);
+			SwitchEntitySkinnedMeshDesc(MeshInfosView, RepresentationFragment, EntityHandle, MeshSwitchFragment.NewSkinnedMeshDescHandle);
 
 			Context.Defer().RemoveFragment<FInstancedActorsMeshSwitchFragment>(EntityHandle);
 		}
@@ -85,7 +87,36 @@ void UInstancedActorsVisualizationSwitcherProcessor::SwitchEntityMeshDesc(FMassI
 			// UMassStationaryISMSwitcherProcessor see that a new instance needs to be made
 			RepresentationFragment.PrevRepresentation = EMassRepresentationType::None;
 		}
-
+		
 		RepresentationFragment.StaticMeshDescHandle = NewStaticMeshDescHandle;
 	}
 }
+
+void UInstancedActorsVisualizationSwitcherProcessor::SwitchEntitySkinnedMeshDesc(FMassInstancedSkinnedMeshInfoArrayView& MeshInfosView, FMassRepresentationFragment& RepresentationFragment, FMassEntityHandle EntityHandle, FSkinnedMeshInstanceVisualizationDescHandle NewMeshDescHandle)
+{
+	if (NewMeshDescHandle != RepresentationFragment.SkinnedMeshDescHandle)
+	{
+		// Remove current FSkinnedMeshInstanceVisualizationDescHandle Mesh Component instance before we switch to NewMeshDescHandle
+		// and 'forget' about it.
+		 if (RepresentationFragment.PrevRepresentation == EMassRepresentationType::SkinnedMeshInstance)
+		{
+			if (ensureMsgf(RepresentationFragment.SkinnedMeshDescHandle.IsValid(), TEXT("Switching visualization while MeshDescHandle is invalid"))
+				&& ensureMsgf(MeshInfosView.IsValidIndex(RepresentationFragment.SkinnedMeshDescHandle.ToIndex()), TEXT("Switching visualization while MeshDescHandle is not a valid MeshInfo index, %d out of %d")
+					, RepresentationFragment.SkinnedMeshDescHandle.ToIndex(), MeshInfosView.Num()))
+			{
+				const int32 MeshInfoIndex = RepresentationFragment.SkinnedMeshDescHandle.ToIndex();
+				FMassInstancedSkinnedMeshInfo& MeshInfo = MeshInfosView[MeshInfoIndex];
+
+				// Note that we're using the PrevLODSignificance here, and the reason for it is that the Prev value matches the 
+				// PrevRepresentation - thus we need to remove from the "previously" used LODSignificance range.
+				MeshInfo.RemoveInstance(EntityHandle, RepresentationFragment.PrevLODSignificance);
+			}
+
+			// Set PrevRepresentation to None to match the new removed instance state and let 
+			// UMassStationaryISMSwitcherProcessor see that a new instance needs to be made
+			RepresentationFragment.PrevRepresentation = EMassRepresentationType::None;
+		}
+
+		RepresentationFragment.SkinnedMeshDescHandle = NewMeshDescHandle;
+	}
+}
diff --git a/Engine/Plugins/Runtime/InstancedActors/Source/InstancedActors/Public/InstancedActorsTypes.h b/Engine/Plugins/Runtime/InstancedActors/Source/InstancedActors/Public/InstancedActorsTypes.h
--- a/Engine/Plugins/Runtime/InstancedActors/Source/InstancedActors/Public/InstancedActorsTypes.h
+++ b/Engine/Plugins/Runtime/InstancedActors/Source/InstancedActors/Public/InstancedActorsTypes.h
@@ -7,6 +7,7 @@
 #include "ISMPartition/ISMComponentDescriptor.h"
 #include "MassEntityTypes.h"
 #include "MassRepresentationTypes.h"
+#include "MassSkinnedMeshRepresentationTypes.h"
 #include "InstancedActorsIndex.h"
 #include "InstancedActorsTypes.generated.h"
 
@@ -259,6 +260,11 @@ struct FInstancedActorsMeshSwitchFragment : public FMassFragment
 	// @see UInstancedActorsVisualizationSwitcherProcessor
 	UPROPERTY()
 	FStaticMeshInstanceVisualizationDescHandle NewStaticMeshDescHandle;
+	
+	// The pending Mass skinned mesh representation index we want to switch to.
+	// @see UInstancedActorsVisualizationSwitcherProcessor
+	UPROPERTY()
+	FSkinnedMeshInstanceVisualizationDescHandle NewSkinnedMeshDescHandle;
 };
 
 
diff --git a/Engine/Plugins/Runtime/InstancedActors/Source/InstancedActors/Public/InstancedActorsVisualizationSwitcherProcessor.h b/Engine/Plugins/Runtime/InstancedActors/Source/InstancedActors/Public/InstancedActorsVisualizationSwitcherProcessor.h
--- a/Engine/Plugins/Runtime/InstancedActors/Source/InstancedActors/Public/InstancedActorsVisualizationSwitcherProcessor.h
+++ b/Engine/Plugins/Runtime/InstancedActors/Source/InstancedActors/Public/InstancedActorsVisualizationSwitcherProcessor.h
@@ -4,6 +4,7 @@
 
 #include "MassProcessor.h"
 #include "MassRepresentationTypes.h"
+#include "MassSkinnedMeshRepresentationTypes.h"
 #include "MassEntityQuery.h"
 #include "InstancedActorsVisualizationSwitcherProcessor.generated.h"
 
@@ -24,9 +25,13 @@ public:
 
 	// Switch EntityHandle to NewStaticMeshDescHandle by removing any current / previous ISMC instances for the current StaticMeshDescIHandle
 	// and setting RepresentationFragment.PrevRepresentation = EMassRepresentationType::None to let the subsequent
-	// UMassStationaryISMSwitcherProcessor see that a new isntance needs to be added for the now set NewStaticMeshDescHandle
+	// UMassStationaryISMSwitcherProcessor see that a new instance needs to be added for the now set NewStaticMeshDescHandle
 	INSTANCEDACTORS_API static void SwitchEntityMeshDesc(FMassInstancedStaticMeshInfoArrayView& ISMInfosView, FMassRepresentationFragment& RepresentationFragment, FMassEntityHandle EntityHandle, FStaticMeshInstanceVisualizationDescHandle NewStaticMeshDescHandle);
 
+	// Switch EntityHandle to NewMeshDescHandle by removing any current / previous ISMC instances for the current MeshDescIHandle
+	// and setting RepresentationFragment.PrevRepresentation = EMassRepresentationType::None to let the subsequent
+	// UMassStationaryISMSwitcherProcessor see that a new instance needs to be added for the now set NewMeshDescHandle
+	INSTANCEDACTORS_API static void SwitchEntitySkinnedMeshDesc(FMassInstancedSkinnedMeshInfoArrayView& MeshInfosView, FMassRepresentationFragment& RepresentationFragment, FMassEntityHandle EntityHandle, FSkinnedMeshInstanceVisualizationDescHandle NewMeshDescHandle);
 protected:
 	virtual void ConfigureQueries(const TSharedRef<FMassEntityManager>& EntityManager) override;
 	virtual void Execute(FMassEntityManager& EntityManager, FMassExecutionContext& Context) override;
diff --git a/Engine/Plugins/Runtime/MassGameplay/Source/MassRepresentation/Private/MassRepresentationAnimationProcessor.cpp b/Engine/Plugins/Runtime/MassGameplay/Source/MassRepresentation/Private/MassRepresentationAnimationProcessor.cpp
new file mode 100644
--- /dev/null
+++ b/Engine/Plugins/Runtime/MassGameplay/Source/MassRepresentation/Private/MassRepresentationAnimationProcessor.cpp
@@ -0,0 +1,70 @@
+// Copyright Epic Games, Inc. All Rights Reserved.
+
+#include "MassRepresentationAnimationProcessor.h"
+
+#include "Engine/World.h"
+#include "MassCommonFragments.h"
+#include "MassEntityManager.h"
+#include "MassExecutionContext.h"
+#include "MassLODFragments.h"
+#include "MassRepresentationFragments.h"
+#include "MassRepresentationSubsystem.h"
+#include "Animation/AnimSequence.h"
+#include "Animation/AnimSequenceTransformProviderData.h"
+
+#include UE_INLINE_GENERATED_CPP_BY_NAME(MassRepresentationAnimationProcessor)
+
+UMassConsumeInstancedSkinnedMeshAnimationProcessor::UMassConsumeInstancedSkinnedMeshAnimationProcessor()
+	: EntityQuery(*this)
+{
+	ExecutionFlags = (int32)(EProcessorExecutionFlags::Client | EProcessorExecutionFlags::Standalone);
+
+	ExecutionOrder.ExecuteInGroup = UE::Mass::ProcessorGroupNames::Representation;
+	bRequiresGameThreadExecution = true;
+}
+
+void UMassConsumeInstancedSkinnedMeshAnimationProcessor::ConfigureQueries(const TSharedRef<FMassEntityManager>& EntityManager)
+{
+	EntityQuery.AddRequirement<FMassRepresentationFragment>(EMassFragmentAccess::ReadOnly);
+	EntityQuery.AddRequirement<FMassRepresentationAnimationFragment>(EMassFragmentAccess::ReadOnly);
+	EntityQuery.AddRequirement<FMassRepresentationLODFragment>(EMassFragmentAccess::ReadOnly);
+
+	EntityQuery.AddSharedRequirement<FMassRepresentationSubsystemSharedFragment>(EMassFragmentAccess::ReadWrite);
+	EntityQuery.AddSubsystemRequirement<UMassRepresentationSubsystem>(EMassFragmentAccess::ReadWrite);
+}
+
+void UMassConsumeInstancedSkinnedMeshAnimationProcessor::Execute(FMassEntityManager& EntityManager, FMassExecutionContext& Context)
+{
+	EntityQuery.ForEachEntityChunk(Context, [](FMassExecutionContext& Context)
+	{
+		UMassRepresentationSubsystem* RepresentationSubsystem = Context.GetSharedFragment<FMassRepresentationSubsystemSharedFragment>().RepresentationSubsystem;
+		check(RepresentationSubsystem);
+
+		FMassInstancedSkinnedMeshInfoArrayView MeshInfos = RepresentationSubsystem->GetMutableInstancedSkinnedMeshInfos();
+
+		const TConstArrayView<FMassRepresentationFragment> RepresentationList = Context.GetFragmentView<FMassRepresentationFragment>();
+		const TConstArrayView<FMassRepresentationLODFragment> RepresentationLODList = Context.GetFragmentView<FMassRepresentationLODFragment>();
+		const TConstArrayView<FMassRepresentationAnimationFragment> AnimationList = Context.GetFragmentView<FMassRepresentationAnimationFragment>();
+
+		for (FMassExecutionContext::FEntityIterator EntityIt = Context.CreateEntityIterator(); EntityIt; ++EntityIt)
+		{
+			const FMassRepresentationLODFragment& RepresentationLOD = RepresentationLODList[EntityIt];
+			const FMassRepresentationAnimationFragment& Animation = AnimationList[EntityIt];
+			const FMassRepresentationFragment& Representation = RepresentationList[EntityIt];
+
+			if (Representation.CurrentRepresentation == EMassRepresentationType::SkinnedMeshInstance)
+			{
+				const int32 MeshInfoIndex = Representation.SkinnedMeshDescHandle.ToIndex();
+				if (ensureMsgf(MeshInfos.IsValidIndex(MeshInfoIndex), TEXT("Invalid handle index %u for MeshInfos View"), MeshInfoIndex))
+				{
+					UpdateMeshAnimation(Context.GetEntity(EntityIt), MeshInfos[MeshInfoIndex], Animation.AnimData, RepresentationLOD.LODSignificance, Representation.PrevLODSignificance);
+				}
+			}
+		}
+	});
+}
+
+void UMassConsumeInstancedSkinnedMeshAnimationProcessor::UpdateMeshAnimation(FMassEntityHandle EntityHandle, FMassInstancedSkinnedMeshInfo& MeshInfo, FAnimSequenceTrackAutoPlayData AnimationData, const float LODSignificance, const float PrevLODSignificance/* = -1.0f*/)
+{
+	MeshInfo.AddBatchedAnimationData(EntityHandle, AnimationData, LODSignificance, PrevLODSignificance);
+}
\ No newline at end of file
diff --git a/Engine/Plugins/Runtime/MassGameplay/Source/MassRepresentation/Private/MassRepresentationProcessor.cpp b/Engine/Plugins/Runtime/MassGameplay/Source/MassRepresentation/Private/MassRepresentationProcessor.cpp
--- a/Engine/Plugins/Runtime/MassGameplay/Source/MassRepresentation/Private/MassRepresentationProcessor.cpp
+++ b/Engine/Plugins/Runtime/MassGameplay/Source/MassRepresentation/Private/MassRepresentationProcessor.cpp
@@ -77,8 +77,8 @@ void UMassRepresentationProcessor::UpdateRepresentation(FMassExecutionContext& C
 	// 2. Change EMassRepresentationType value based on some configs (not all flows will care about this);
 	// 3. Switch the in-game instance representation depending on the EMassRepresentationType;
 	// 		a. If EMassRepresentationType == HighResSpawnedActor or LowResSpawnedActor, sends an Actor SpawnRequest to MassRepresentationActorManagement;
-	// 		b. If EMassRepresentationType == StaticMeshInstance, sends an Actor disable request to MassRepresentationActorManagement;
-	//		c. If EMassRepresentationType == StaticMeshInstance, sends an Actor disable request to MassRepresentationActorManagement.
+	// 		b. If EMassRepresentationType == StaticMeshInstance or SkinnedMeshInstance, sends an Actor disable request to MassRepresentationActorManagement;
+	//		c. If EMassRepresentationType == StaticMeshInstance or SkinnedMeshInstance, sends an Actor disable request to MassRepresentationActorManagement.
 	// 		NOTE: I guess this system assumes all instances are already represented by ISMs, which is why we're only dealing with Actor spawn/deactivation
 	for (FMassExecutionContext::FEntityIterator EntityIt = Context.CreateEntityIterator(); EntityIt; ++EntityIt)
 	{
@@ -227,7 +227,7 @@ void UMassRepresentationProcessor::UpdateRepresentation(FMassExecutionContext& C
 					{
 						// Make sure our (re)activated actor is at the simulated position
 						// Needs to be done before enabling the actor so the animation initialization can use the new values
-						if (Representation.CurrentRepresentation == EMassRepresentationType::StaticMeshInstance)
+						if (UE::Mass::Representation::IsValidMeshRepresentation(Representation.CurrentRepresentation))
 						{
 							RepresentationActorManagement->TeleportActor(Representation.PrevTransform, *NewActor, Context.Defer());
 						}
@@ -242,14 +242,21 @@ void UMassRepresentationProcessor::UpdateRepresentation(FMassExecutionContext& C
 					break;
 				}
 				case EMassRepresentationType::StaticMeshInstance:
-					if (!bDoKeepActorExtraFrame || 
-					   (Representation.PrevRepresentation != EMassRepresentationType::HighResSpawnedActor && Representation.PrevRepresentation != EMassRepresentationType::LowResSpawnedActor))
+					if (!bDoKeepActorExtraFrame || !UE::Mass::Representation::IsValidActorRepresentation(Representation.PrevRepresentation))
 					{
 						DisableActorForISM(Actor);
 					}
  
 					Representation.CurrentRepresentation = EMassRepresentationType::StaticMeshInstance;
 					break;
+				case EMassRepresentationType::SkinnedMeshInstance:
+					if (!bDoKeepActorExtraFrame || !UE::Mass::Representation::IsValidActorRepresentation(Representation.PrevRepresentation))
+					{
+						DisableActorForISM(Actor);
+					}
+
+					Representation.CurrentRepresentation = EMassRepresentationType::SkinnedMeshInstance;
+					break;
 				case EMassRepresentationType::None:
 					if (!Actor || ActorInfo.IsOwnedByMass())
 					{
@@ -269,7 +276,7 @@ void UMassRepresentationProcessor::UpdateRepresentation(FMassExecutionContext& C
 			}
 		}
 		else if (bDoKeepActorExtraFrame
-				&& Representation.PrevRepresentation == EMassRepresentationType::StaticMeshInstance 
+				&& UE::Mass::Representation::IsValidMeshRepresentation(Representation.PrevRepresentation)
 				&& (PrevRepresentationCopy == EMassRepresentationType::HighResSpawnedActor 
 					|| PrevRepresentationCopy == EMassRepresentationType::LowResSpawnedActor))
 		{
diff --git a/Engine/Plugins/Runtime/MassGameplay/Source/MassRepresentation/Private/MassRepresentationSubsystem.cpp b/Engine/Plugins/Runtime/MassGameplay/Source/MassRepresentation/Private/MassRepresentationSubsystem.cpp
--- a/Engine/Plugins/Runtime/MassGameplay/Source/MassRepresentation/Private/MassRepresentationSubsystem.cpp
+++ b/Engine/Plugins/Runtime/MassGameplay/Source/MassRepresentation/Private/MassRepresentationSubsystem.cpp
@@ -94,6 +94,21 @@ FMassInstancedStaticMeshInfoArrayView UMassRepresentationSubsystem::GetMutableIn
 	return MoveTemp(EmptyView);
 }
 
+FMassInstancedSkinnedMeshInfoArrayView UMassRepresentationSubsystem::GetMutableInstancedSkinnedMeshInfos()
+{
+	if (LIKELY(VisualizationComponent))
+	{
+		FMassInstancedSkinnedMeshInfoArrayView View = VisualizationComponent->GetMutableInstanceVisualInfos();
+		return MoveTemp(View);
+	}
+
+	// This can happen if the visualizer actor is destroyed prematurely.
+	// It can also happen during world teardown if this is called after actors are destroyed.
+	// This is not hit during normal execution, but we still need a default return value.
+	FMassInstancedSkinnedMeshInfoArrayView EmptyView = MAKE_MASS_INSTANCED_SKINNED_MESH_INFO_ARRAY_VIEW(MakeArrayView<FMassInstancedSkinnedMeshInfo>(nullptr, 0), InstancedSkinnedMeshInfosDetector);
+	return MoveTemp(EmptyView);
+}
+
 void UMassRepresentationSubsystem::DirtyStaticMeshInstances()
 {
 	if (LIKELY(VisualizationComponent))
@@ -573,3 +588,13 @@ void UMassRepresentationSubsystem::AddReferencedObjects(UObject* InThis, FRefere
 		Collector.AddStableReference(&TemplateActorData.Actor.GetGCPtr());
 	}
 }
+
+FSkinnedMeshInstanceVisualizationDescHandle UMassRepresentationSubsystem::FindOrAddSkinnedMeshVisualizationDesc(const FSkinnedMeshInstanceVisualizationDesc& Desc)
+{
+	if (LIKELY(VisualizationComponent))
+	{
+		return VisualizationComponent->FindOrAddSkinnedMeshInstanceVisualDesc(Desc);
+	}
+
+	return FSkinnedMeshInstanceVisualizationDescHandle();
+}
diff --git a/Engine/Plugins/Runtime/MassGameplay/Source/MassRepresentation/Private/MassSkinnedMeshRepresentationTypes.cpp b/Engine/Plugins/Runtime/MassGameplay/Source/MassRepresentation/Private/MassSkinnedMeshRepresentationTypes.cpp
new file mode 100644
--- /dev/null
+++ b/Engine/Plugins/Runtime/MassGameplay/Source/MassRepresentation/Private/MassSkinnedMeshRepresentationTypes.cpp
@@ -0,0 +1,34 @@
+// Copyright Epic Games, Inc. All Rights Reserved.
+
+#include "MassSkinnedMeshRepresentationTypes.h"
+#include "Components/InstancedSkinnedMeshComponent.h"
+
+//-----------------------------------------------------------------------------
+// FMassSkinnedMeshInstanceVisualizationMeshDesc
+//-----------------------------------------------------------------------------
+FMassSkinnedMeshInstanceVisualizationMeshDesc::FMassSkinnedMeshInstanceVisualizationMeshDesc()
+{
+	InstancedSkinnedMeshComponentClass = UInstancedSkinnedMeshComponent::StaticClass();
+}
+
+//-----------------------------------------------------------------------------
+// FSkinnedMeshInstanceVisualizationDesc
+//-----------------------------------------------------------------------------
+bool FSkinnedMeshInstanceVisualizationDesc::IsValid() const
+{
+	for (const FMassSkinnedMeshInstanceVisualizationMeshDesc& MeshDesc : Meshes)
+	{
+		if (MeshDesc.Asset && MeshDesc.InstancedSkinnedMeshComponentClass)
+		{
+			return true;
+		}
+	}
+	return false;
+}
+
+void FMassInstancedSkinnedMeshInfo::Reset()
+{
+	Desc.Reset();
+	InstancedSkinnedMeshComponents.Reset();
+	LODSignificanceRanges.Reset();
+}
diff --git a/Engine/Plugins/Runtime/MassGameplay/Source/MassRepresentation/Private/MassStationaryISMRepresentationFragmentDestructor.cpp b/Engine/Plugins/Runtime/MassGameplay/Source/MassRepresentation/Private/MassStationaryISMRepresentationFragmentDestructor.cpp
--- a/Engine/Plugins/Runtime/MassGameplay/Source/MassRepresentation/Private/MassStationaryISMRepresentationFragmentDestructor.cpp
+++ b/Engine/Plugins/Runtime/MassGameplay/Source/MassRepresentation/Private/MassStationaryISMRepresentationFragmentDestructor.cpp
@@ -2,6 +2,7 @@
 
 #include "MassStationaryISMRepresentationFragmentDestructor.h"
 #include "MassRepresentationSubsystem.h"
+#include "MassSkinnedMeshRepresentationTypes.h"
 #include "MassCommonFragments.h"
 #include "MassRepresentationProcessor.h"
 
@@ -34,6 +35,7 @@ void UMassStationaryISMRepresentationFragmentDestructor::Execute(FMassEntityMana
 		UMassRepresentationSubsystem* RepresentationSubsystem = Context.GetMutableSharedFragment<FMassRepresentationSubsystemSharedFragment>().RepresentationSubsystem;
 		check(RepresentationSubsystem);
 		FMassInstancedStaticMeshInfoArrayView ISMInfosView = RepresentationSubsystem->GetMutableInstancedStaticMeshInfos();
+		FMassInstancedSkinnedMeshInfoArrayView MeshInfosView = RepresentationSubsystem->GetMutableInstancedSkinnedMeshInfos();
 
 		const TArrayView<FMassRepresentationFragment> RepresentationList = Context.GetMutableFragmentView<FMassRepresentationFragment>();
 
@@ -57,6 +59,23 @@ void UMassStationaryISMRepresentationFragmentDestructor::Execute(FMassEntityMana
 				}
 				Representation.CurrentRepresentation = EMassRepresentationType::None;
 			}
+			else if (Representation.CurrentRepresentation == EMassRepresentationType::SkinnedMeshInstance)
+			{
+				const int32 MeshInfoIndex = Representation.SkinnedMeshDescHandle.ToIndex();
+				if (ensureMsgf(MeshInfosView.IsValidIndex(MeshInfoIndex), TEXT("Invalid handle index %u for MeshInfosView"), MeshInfoIndex))
+				{
+					FMassInstancedSkinnedMeshInfo& MeshInfo = MeshInfosView[MeshInfoIndex];
+					if (FMassLODInstancedSkinnedMeshSignificanceRange* OldRange = MeshInfo.GetLODSignificanceRange(Representation.PrevLODSignificance))
+					{
+						const FMassEntityHandle EntityHandle = Context.GetEntity(EntityIt);
+						if (OldRange)
+						{
+							OldRange->RemoveInstance(EntityHandle);
+						}
+					}
+				}
+				Representation.CurrentRepresentation = EMassRepresentationType::None;
+			}
 		}
 	});
 }
diff --git a/Engine/Plugins/Runtime/MassGameplay/Source/MassRepresentation/Private/MassStationaryISMSwitcherProcessor.cpp b/Engine/Plugins/Runtime/MassGameplay/Source/MassRepresentation/Private/MassStationaryISMSwitcherProcessor.cpp
--- a/Engine/Plugins/Runtime/MassGameplay/Source/MassRepresentation/Private/MassStationaryISMSwitcherProcessor.cpp
+++ b/Engine/Plugins/Runtime/MassGameplay/Source/MassRepresentation/Private/MassStationaryISMSwitcherProcessor.cpp
@@ -4,6 +4,7 @@
 #include "MassRepresentationSubsystem.h"
 #include "MassCommonFragments.h"
 #include "MassRepresentationProcessor.h"
+#include "MassRepresentationTypes.h"
 
 #include "MassSignalSubsystem.h"
 
@@ -44,6 +45,7 @@ void UMassStationaryISMSwitcherProcessor::ProcessContext(FMassExecutionContext&
 	UMassRepresentationSubsystem* RepresentationSubsystem = Context.GetMutableSharedFragment<FMassRepresentationSubsystemSharedFragment>().RepresentationSubsystem;
 	check(RepresentationSubsystem);
 	FMassInstancedStaticMeshInfoArrayView ISMInfosView = RepresentationSubsystem->GetMutableInstancedStaticMeshInfos();
+	FMassInstancedSkinnedMeshInfoArrayView MeshInfosView = RepresentationSubsystem->GetMutableInstancedSkinnedMeshInfos();
 
 	const TConstArrayView<FMassRepresentationLODFragment> RepresentationLODList = Context.GetFragmentView<FMassRepresentationLODFragment>();
 	const TConstArrayView<FTransformFragment> TransformList = Context.GetFragmentView<FTransformFragment>();
@@ -59,24 +61,38 @@ void UMassStationaryISMSwitcherProcessor::ProcessContext(FMassExecutionContext&
 		const FTransformFragment& TransformFragment = TransformList[EntityIt];
 		FMassRepresentationFragment& Representation = RepresentationList[EntityIt];
 
-		if ((Representation.PrevRepresentation != EMassRepresentationType::StaticMeshInstance
-			&& Representation.CurrentRepresentation != EMassRepresentationType::StaticMeshInstance)
-			|| Representation.bIsPendingDestruction)
+		if (Representation.bIsPendingDestruction)
 		{
 			// nothing to do here
 			continue;
 		}
 
-		if (!ensureMsgf(Representation.StaticMeshDescHandle.IsValid() && ISMInfosView.IsValidIndex(Representation.StaticMeshDescHandle.ToIndex())
+		if (!UE::Mass::Representation::IsValidMeshRepresentation(Representation.PrevRepresentation) && !UE::Mass::Representation::IsValidMeshRepresentation(Representation.CurrentRepresentation))
+{
+			// nothing to do here
+			continue;
+		}
+
+		bool bIsStaticMeshInstance = (Representation.PrevRepresentation == EMassRepresentationType::StaticMeshInstance || Representation.CurrentRepresentation == EMassRepresentationType::StaticMeshInstance);
+		bool bIsSkinnedMeshInstance = (Representation.PrevRepresentation == EMassRepresentationType::SkinnedMeshInstance || Representation.CurrentRepresentation == EMassRepresentationType::SkinnedMeshInstance);
+
+		if (bIsStaticMeshInstance && !ensureMsgf(Representation.StaticMeshDescHandle.IsValid() && ISMInfosView.IsValidIndex(Representation.StaticMeshDescHandle.ToIndex())
 						, TEXT("Invalid handle index %u for ISMInfosView"), Representation.StaticMeshDescHandle.ToIndex()))
 		{
 			continue;
 		}
 
-		FMassInstancedStaticMeshInfo& ISMInfo = ISMInfosView[Representation.StaticMeshDescHandle.ToIndex()];
+		if (bIsSkinnedMeshInstance && !ensureMsgf(Representation.SkinnedMeshDescHandle.IsValid() && MeshInfosView.IsValidIndex(Representation.SkinnedMeshDescHandle.ToIndex())
+						, TEXT("Invalid handle index %u for MeshInfosView"), Representation.SkinnedMeshDescHandle.ToIndex()))
+		{
+			continue;
+		}
+
 		if (const bool bSwitchedAwayFromStaticMesh = (Representation.PrevRepresentation == EMassRepresentationType::StaticMeshInstance
 			&& Representation.CurrentRepresentation != EMassRepresentationType::StaticMeshInstance))
 		{
+			FMassInstancedStaticMeshInfo& ISMInfo = ISMInfosView[Representation.StaticMeshDescHandle.ToIndex()];
+
 			// note that we're using the PrevLODSignificance here, and the reason for it is that the Prev value matches the 
 			// PrevRepresentation - thus we need to remove from the "previously" used LODSignificance range.
 			ISMInfo.RemoveInstance(EntityHandle, Representation.PrevLODSignificance);
@@ -84,7 +100,24 @@ void UMassStationaryISMSwitcherProcessor::ProcessContext(FMassExecutionContext&
 			// consume "prev" data
 			Representation.PrevRepresentation = Representation.CurrentRepresentation;
 
-			if (Representation.PrevRepresentation != EMassRepresentationType::None)
+			if (Representation.PrevRepresentation != EMassRepresentationType::None && !bIsSkinnedMeshInstance)
+			{
+				SignalSubsystem.SignalEntity(UE::Mass::Signals::SwitchedToActor, EntityHandle);
+			}
+		}
+		else if (const bool bSwitchedAwayFromSkinnedMesh = (Representation.PrevRepresentation == EMassRepresentationType::SkinnedMeshInstance
+			&& Representation.CurrentRepresentation != EMassRepresentationType::SkinnedMeshInstance))
+		{
+			FMassInstancedSkinnedMeshInfo& MeshInfo = MeshInfosView[Representation.SkinnedMeshDescHandle.ToIndex()];
+
+			// note that we're using the PrevLODSignificance here, and the reason for it is that the Prev value matches the 
+			// PrevRepresentation - thus we need to remove from the "previously" used LODSignificance range.
+			MeshInfo.RemoveInstance(EntityHandle, Representation.PrevLODSignificance);
+
+			// consume "prev" data
+			Representation.PrevRepresentation = Representation.CurrentRepresentation;
+
+			if (Representation.PrevRepresentation != EMassRepresentationType::None && !bIsStaticMeshInstance)
 			{
 				SignalSubsystem.SignalEntity(UE::Mass::Signals::SwitchedToActor, EntityHandle);
 			}
@@ -92,10 +125,10 @@ void UMassStationaryISMSwitcherProcessor::ProcessContext(FMassExecutionContext&
 		else if (const bool bSwitchedToStaticMesh = (Representation.PrevRepresentation != EMassRepresentationType::StaticMeshInstance
 			&& Representation.CurrentRepresentation == EMassRepresentationType::StaticMeshInstance))
 		{
+			FMassInstancedStaticMeshInfo& ISMInfo = ISMInfosView[Representation.StaticMeshDescHandle.ToIndex()];
+
 			const FTransform& Transform = TransformFragment.GetTransform();
-			const FTransform& PrevTransform = Representation.PrevTransform;
 			const float LODSignificance = RepresentationLOD.LODSignificance;
-			const float PrevLODSignificance = Representation.PrevLODSignificance;
 				
 			if (FMassLODSignificanceRange* NewRange = ISMInfo.GetLODSignificanceRange(RepresentationLOD.LODSignificance))
 			{
@@ -125,28 +158,103 @@ void UMassStationaryISMSwitcherProcessor::ProcessContext(FMassExecutionContext&
 				Representation.PrevRepresentation = Representation.CurrentRepresentation;
 			}
 
-			SignalSubsystem.SignalEntity(UE::Mass::Signals::SwitchedToISM, EntityHandle);
+			if (!bIsSkinnedMeshInstance)
+			{
+				SignalSubsystem.SignalEntity(UE::Mass::Signals::SwitchedToISM, EntityHandle);
+			}		
 		}
-		else if (ISMInfo.GetLODSignificanceRangesNum() > 1 && Representation.PrevLODSignificance != RepresentationLOD.LODSignificance)
+		else if (const bool bSwitchedToSkinnedMesh = (Representation.PrevRepresentation != EMassRepresentationType::SkinnedMeshInstance
+			&& Representation.CurrentRepresentation == EMassRepresentationType::SkinnedMeshInstance))
 		{
-			// we remain in ISM land, but LODSignificance changed and we have multiple LODSignificance ranges for this entity
-			FMassLODSignificanceRange* OldRange = ISMInfo.GetLODSignificanceRange(Representation.PrevLODSignificance);
-			FMassLODSignificanceRange* NewRange = ISMInfo.GetLODSignificanceRange(RepresentationLOD.LODSignificance);
-			if (OldRange != NewRange)
+			FMassInstancedSkinnedMeshInfo& MeshInfo = MeshInfosView[Representation.SkinnedMeshDescHandle.ToIndex()];
+
+			const FTransform& Transform = TransformFragment.GetTransform();
+			const float LODSignificance = RepresentationLOD.LODSignificance;
+
+			if (FMassLODInstancedSkinnedMeshSignificanceRange* NewRange = MeshInfo.GetLODSignificanceRange(RepresentationLOD.LODSignificance))
 			{
-				if (OldRange)
+				if (MeshInfo.ShouldUseTransformOffset())
 				{
-					OldRange->RemoveInstance(EntityHandle);
+					const FTransform& TransformOffset = MeshInfo.GetTransformOffset();
+					const FTransform SMTransform = TransformOffset * Transform;
+					NewRange->AddInstance(EntityHandle, SMTransform);
 				}
-				if (NewRange)
+				else
 				{
-					const FTransform& Transform = TransformFragment.GetTransform();
 					NewRange->AddInstance(EntityHandle, Transform);
 				}
+
+				if (MeshInfo.GetDesc().CustomDataFloats.Num() > 0)
+				{
+					NewRange->AddBatchedCustomDataFloats(MeshInfo.GetDesc().CustomDataFloats, {});
+				}
+			}
+
+			// consume "prev" data
+			// @note crazy hacky, but we don't want to consume if bDoKeepActorExtraFrame is true. In that case 
+			// UMassRepresentationProcessor::UpdateRepresentation expects the "prev" state not to be consumed 
+			// a frame longer so that it can do the consuming (and call "disable actor").
+			if (bDoKeepActorExtraFrame == false)
+			{
+				Representation.PrevRepresentation = Representation.CurrentRepresentation;
+			}
+
+			if (bIsStaticMeshInstance)
+			{
+				SignalSubsystem.SignalEntity(UE::Mass::Signals::SwitchedToISM, EntityHandle);
 			}
 		}
+		else if (bIsStaticMeshInstance)
+		{
+			FMassInstancedStaticMeshInfo& ISMInfo = ISMInfosView[Representation.StaticMeshDescHandle.ToIndex()];
+
+			if (ISMInfo.GetLODSignificanceRangesNum() > 1 && Representation.PrevLODSignificance != RepresentationLOD.LODSignificance)
+			{
+				// we remain in ISM land, but LODSignificance changed and we have multiple LODSignificance ranges for this entity
+				FMassLODSignificanceRange* OldRange = ISMInfo.GetLODSignificanceRange(Representation.PrevLODSignificance);
+				FMassLODSignificanceRange* NewRange = ISMInfo.GetLODSignificanceRange(RepresentationLOD.LODSignificance);
+				if (OldRange != NewRange)
+				{
+					if (OldRange)
+					{
+						OldRange->RemoveInstance(EntityHandle);
+					}
+					if (NewRange)
+					{
+						const FTransform& Transform = TransformFragment.GetTransform();
+						NewRange->AddInstance(EntityHandle, Transform);
+					}
+				}
+		
+				// consume "prev" data
+				Representation.PrevLODSignificance = RepresentationLOD.LODSignificance;
+			}
+		}
+		else if (bIsSkinnedMeshInstance)
+		{
+			FMassInstancedSkinnedMeshInfo& MeshInfo = MeshInfosView[Representation.SkinnedMeshDescHandle.ToIndex()];
 
-		// consume "prev" data
-		Representation.PrevLODSignificance = RepresentationLOD.LODSignificance;
+			if (MeshInfo.GetLODSignificanceRangesNum() > 1 && Representation.PrevLODSignificance != RepresentationLOD.LODSignificance)
+			{
+				// we remain in Skinned Mesh land, but LODSignificance changed and we have multiple LODSignificance ranges for this entity
+				FMassLODInstancedSkinnedMeshSignificanceRange* OldRange = MeshInfo.GetLODSignificanceRange(Representation.PrevLODSignificance);
+				FMassLODInstancedSkinnedMeshSignificanceRange* NewRange = MeshInfo.GetLODSignificanceRange(RepresentationLOD.LODSignificance);
+				if (OldRange != NewRange)
+				{
+					if (OldRange)
+					{
+						OldRange->RemoveInstance(EntityHandle);
+					}
+					if (NewRange)
+					{
+						const FTransform& Transform = TransformFragment.GetTransform();
+						NewRange->AddInstance(EntityHandle, Transform);
+					}
+				}
+
+				// consume "prev" data
+				Representation.PrevLODSignificance = RepresentationLOD.LODSignificance;
+			}
+		}
 	}
-}
+}
\ No newline at end of file
diff --git a/Engine/Plugins/Runtime/MassGameplay/Source/MassRepresentation/Private/MassUpdateSkinnedMeshProcessor.cpp b/Engine/Plugins/Runtime/MassGameplay/Source/MassRepresentation/Private/MassUpdateSkinnedMeshProcessor.cpp
new file mode 100644
--- /dev/null
+++ b/Engine/Plugins/Runtime/MassGameplay/Source/MassRepresentation/Private/MassUpdateSkinnedMeshProcessor.cpp
@@ -0,0 +1,87 @@
+// Copyright Epic Games, Inc. All Rights Reserved.
+
+#include "MassUpdateSkinnedMeshProcessor.h"
+
+#include "Engine/World.h"
+#include "MassCommonFragments.h"
+#include "MassEntityManager.h"
+#include "MassExecutionContext.h"
+#include "MassLODFragments.h"
+#include "MassRepresentationFragments.h"
+#include "MassRepresentationSubsystem.h"
+#include "MassVisualizationComponent.h"
+
+#include UE_INLINE_GENERATED_CPP_BY_NAME(MassUpdateSkinnedMeshProcessor)
+
+UMassUpdateInstancedSkinnedMeshProcessor::UMassUpdateInstancedSkinnedMeshProcessor()
+	: EntityQuery(*this)
+{
+	ExecutionFlags = (int32)(EProcessorExecutionFlags::Client | EProcessorExecutionFlags::Standalone);
+
+	ExecutionOrder.ExecuteInGroup = UE::Mass::ProcessorGroupNames::Representation;
+	bRequiresGameThreadExecution = true;
+}
+
+void UMassUpdateInstancedSkinnedMeshProcessor::UpdateMeshTransform(FMassEntityHandle EntityHandle, FMassInstancedSkinnedMeshInfo& MeshInfo, const FTransform& Transform, const FTransform& PrevTransform, const float LODSignificance, const float PrevLODSignificance /*= -1.0f*/)
+{
+	if (MeshInfo.ShouldUseTransformOffset())
+	{
+		const FTransform& TransformOffset = MeshInfo.GetTransformOffset();
+		const FTransform SMTransform = TransformOffset * Transform;
+		const FTransform SMPrevTransform = TransformOffset * PrevTransform;
+
+		MeshInfo.AddBatchedTransform(EntityHandle, SMTransform, SMPrevTransform, LODSignificance, PrevLODSignificance);
+	}
+	else
+	{
+		MeshInfo.AddBatchedTransform(EntityHandle, Transform, PrevTransform, LODSignificance, PrevLODSignificance);
+	}
+}
+
+void UMassUpdateInstancedSkinnedMeshProcessor::ConfigureQueries(const TSharedRef<FMassEntityManager>& EntityManager)
+{
+	EntityQuery.AddRequirement<FTransformFragment>(EMassFragmentAccess::ReadOnly);
+	EntityQuery.AddRequirement<FMassRepresentationFragment>(EMassFragmentAccess::ReadWrite);
+	EntityQuery.AddRequirement<FMassRepresentationLODFragment>(EMassFragmentAccess::ReadOnly);
+	EntityQuery.AddChunkRequirement<FMassVisualizationChunkFragment>(EMassFragmentAccess::ReadWrite);
+	EntityQuery.SetChunkFilter(&FMassVisualizationChunkFragment::AreAnyEntitiesVisibleInChunk);
+	EntityQuery.AddSharedRequirement<FMassRepresentationSubsystemSharedFragment>(EMassFragmentAccess::ReadWrite);
+	EntityQuery.AddSubsystemRequirement<UMassRepresentationSubsystem>(EMassFragmentAccess::ReadWrite);
+
+	// ignore entities configured to have their representation static (@todo maybe just check if there's not movement fragment?)
+	EntityQuery.AddTagRequirement<FMassStaticRepresentationTag>(EMassFragmentPresence::None);
+}
+
+void UMassUpdateInstancedSkinnedMeshProcessor::Execute(FMassEntityManager& EntityManager, FMassExecutionContext& Context)
+{
+	EntityQuery.ForEachEntityChunk(Context, [](FMassExecutionContext& Context)
+		{
+			UMassRepresentationSubsystem* RepresentationSubsystem = Context.GetSharedFragment<FMassRepresentationSubsystemSharedFragment>().RepresentationSubsystem;
+			check(RepresentationSubsystem);
+
+			FMassInstancedSkinnedMeshInfoArrayView SkinnedMeshInfos = RepresentationSubsystem->GetMutableInstancedSkinnedMeshInfos();
+
+			const TConstArrayView<FTransformFragment> TransformList = Context.GetFragmentView<FTransformFragment>();
+			const TArrayView<FMassRepresentationFragment> RepresentationList = Context.GetMutableFragmentView<FMassRepresentationFragment>();
+			const TConstArrayView<FMassRepresentationLODFragment> RepresentationLODList = Context.GetFragmentView<FMassRepresentationLODFragment>();
+
+			for (FMassExecutionContext::FEntityIterator EntityIt = Context.CreateEntityIterator(); EntityIt; ++EntityIt)
+			{
+				const FTransformFragment& TransformFragment = TransformList[EntityIt];
+				const FMassRepresentationLODFragment& RepresentationLOD = RepresentationLODList[EntityIt];
+
+				FMassRepresentationFragment& Representation = RepresentationList[EntityIt];
+
+				if (Representation.CurrentRepresentation == EMassRepresentationType::SkinnedMeshInstance)
+				{
+					const int32 SkinnedMeshInfoIndex = Representation.SkinnedMeshDescHandle.ToIndex();
+					if (ensureMsgf(SkinnedMeshInfos.IsValidIndex(SkinnedMeshInfoIndex), TEXT("Invalid handle index %u for MeshInfos View"), SkinnedMeshInfoIndex))
+					{
+						UpdateMeshTransform(Context.GetEntity(EntityIt), SkinnedMeshInfos[SkinnedMeshInfoIndex], TransformFragment.GetTransform(), Representation.PrevTransform, RepresentationLOD.LODSignificance, Representation.PrevLODSignificance);
+					}
+				}
+				Representation.PrevTransform = TransformFragment.GetTransform();
+				Representation.PrevLODSignificance = RepresentationLOD.LODSignificance;
+			}
+		});
+}
diff --git a/Engine/Plugins/Runtime/MassGameplay/Source/MassRepresentation/Private/MassVisualizationComponent.cpp b/Engine/Plugins/Runtime/MassGameplay/Source/MassRepresentation/Private/MassVisualizationComponent.cpp
--- a/Engine/Plugins/Runtime/MassGameplay/Source/MassRepresentation/Private/MassVisualizationComponent.cpp
+++ b/Engine/Plugins/Runtime/MassGameplay/Source/MassRepresentation/Private/MassVisualizationComponent.cpp
@@ -357,6 +357,24 @@ void UMassVisualizationComponent::ClearAllVisualInstances()
 	ISMCSharedData.Reset();
 	InstancedSMComponentsRequiringConstructing.Reset();
 	InstancedStaticMeshInfos.Reset();
+
+	UE_MT_SCOPED_WRITE_ACCESS(InstancedSkinnedMeshInfosDetector);
+
+	for (int32 SharedDataIndex = 0; SharedDataIndex < InstancedSkinnedMeshComponentSharedData.Num(); ++SharedDataIndex)
+	{
+		if (UInstancedSkinnedMeshComponent* InstancedSkinnedMeshComponent = InstancedSkinnedMeshComponentSharedData.GetAtIndex(SharedDataIndex).GetMutableInstancedSkinnedMeshComponent())
+		{
+			InstancedSkinnedMeshComponent->ClearInstances();
+			InstancedSkinnedMeshComponent->DestroyComponent();
+		}
+	}
+
+	InstancedSkinnedMeshComponentMap.Reset();
+	MeshDescToInstancedSkinnedMeshComponentMap.Reset();
+	InstancedSkinnedMeshComponentSharedData.Reset();
+	InstancedSkinnedMeshComponentsRequiringConstructing.Reset();
+	InstancedSkinnedMeshInfos.Reset();
+	InstancedSkinnedMeshInfosFreeIndices.Reset();
 }
 
 void UMassVisualizationComponent::DirtyVisuals()
@@ -369,6 +387,15 @@ void UMassVisualizationComponent::DirtyVisuals()
 			InstancedStaticMeshComponent->MarkRenderStateDirty();
 		}
 	}
+
+	UE_MT_SCOPED_WRITE_ACCESS(InstancedSkinnedMeshInfosDetector);
+	for (FMassInstancedSkinnedMeshInfo& Info : InstancedSkinnedMeshInfos)
+	{
+		for (UInstancedSkinnedMeshComponent* InstancedSkinnedMeshComponent : Info.InstancedSkinnedMeshComponents)
+		{
+			InstancedSkinnedMeshComponent->MarkRenderStateDirty();
+		}
+	}
 }
 
 void UMassVisualizationComponent::BeginVisualChanges()
@@ -381,6 +408,13 @@ void UMassVisualizationComponent::BeginVisualChanges()
 		ConstructStaticMeshComponents();
 		InstancedSMComponentsRequiringConstructing.Reset();
 	}
+
+	// Conditionally construct skinned mesh components
+	if (InstancedSkinnedMeshComponentsRequiringConstructing.Num())
+	{
+		ConstructSkinnedMeshComponents();
+		InstancedSkinnedMeshComponentsRequiringConstructing.Reset();
+	}
 }
 
 void UMassVisualizationComponent::HandleChangesWithExternalIDTracking(UInstancedStaticMeshComponent& ISMComponent, FMassISMCSharedData& SharedData)
@@ -639,6 +673,231 @@ void UMassVisualizationComponent::EndVisualChanges()
 		
 		SharedData.ResetAccumulatedData();
 	}
+
+	// Batch update gathered instance transforms
+	for (FMassInstancedSkinnedMeshComponentSharedDataMap::FDirtyIterator It(InstancedSkinnedMeshComponentSharedData); It; ++It)
+	{
+		FMassInstancedSkinnedMeshComponentSharedData& SharedData = *It;
+
+		UInstancedSkinnedMeshComponent* InstancedSkinnedMeshComponent = SharedData.GetMutableInstancedSkinnedMeshComponent();
+		// @todo need to check validity this way since Mass used to rely on the assumption that all the ISM components used were
+		// under its control. That's no longer the case, but the system has not been updated to take that into consideration.
+		// This is a temporary fix. 
+		if (IsValid(InstancedSkinnedMeshComponent))
+		{
+			if (SharedData.RequiresExternalInstanceIDTracking())
+			{
+				HandleChangesWithExternalIDTrackingForComponent(*InstancedSkinnedMeshComponent, SharedData);
+				It.ClearDirtyFlag();
+			}
+			else
+			{
+				// Process all removes.
+				ProcessRemovesForComponent(*InstancedSkinnedMeshComponent, SharedData);
+
+				const int32 NumCustomDataFloats = SharedData.MeshInstanceCustomFloats.Num() / (FMath::Max(1, SharedData.EntitiesRequiringUpdate.Num()));
+
+				// Ensure InstanceCustomData is passed if NumCustomDataFloats > 0. If it is, also make sure
+				// its length is NumCustomDataFloats * InstanceTransforms.Num()
+				ensure(NumCustomDataFloats == 0 || (SharedData.MeshInstanceCustomFloats.Num() == NumCustomDataFloats * SharedData.EntitiesRequiringUpdate.Num()));
+				InstancedSkinnedMeshComponent->SetNumCustomDataFloats(NumCustomDataFloats);
+				TArray<FMassEntityHandle>& EntityHandlesWithPendingUpdates = SharedData.EntitiesRequiringUpdate;
+				{
+					// Loop over all the instances in the update and 
+					// 1. Sort the data such that all Adds are last
+					// 2. Remove any duplicates (unsure if they may exist)
+					FMassInstancedSkinnedMeshComponentSharedData::FEntityToPrimitiveIdMap& SharedIdMap = SharedData.GetMutableEntityToPrimitiveIdMap();
+					FMassInstancedSkinnedMeshComponentSharedData::FEntityToTrackIdMap& TrackMap = SharedData.GetMutableEntityToTrackMap();
+					// Filter out all updates & duplicate adds
+					TBitArray<> Unprocessed;
+					Unprocessed.SetNum(SharedIdMap.GetMaxIndex(), true);
+					// Process interval
+
+					TConstArrayView<FTransform> PrevInstanceTransforms = SharedData.GetMeshInstancePrevTransforms();
+					TConstArrayView<FTransform> InstanceTransforms = SharedData.GetMeshInstanceTransformsArray();
+					TConstArrayView<FAnimSequenceTrackAutoPlayData> InstanceAnimationData = SharedData.GetMeshInstanceAnimationData();
+					TConstArrayView<float> CustomDataFloats = SharedData.GetMeshInstanceCustomFloats();
+
+					struct FAddItem
+					{
+						Experimental::FHashElementId ElementId;
+						int32 IDIndex;
+					};
+					TArray<FAddItem> ToAdd;
+					ToAdd.Reserve(EntityHandlesWithPendingUpdates.Num());
+					for (int32 IDIndex = 0; IDIndex < EntityHandlesWithPendingUpdates.Num(); ++IDIndex)
+					{
+						bool bWasAlreadyInMap = false;
+						Experimental::FHashElementId ElementId = SharedIdMap.FindOrAddId(EntityHandlesWithPendingUpdates[IDIndex], FPrimitiveInstanceId{ INDEX_NONE }, bWasAlreadyInMap);
+
+						// if it was already in the map, it may be a duplicate if we have processed it already
+						bool bIsDuplicate = bWasAlreadyInMap && !Unprocessed[ElementId.GetIndex()];
+						if (bIsDuplicate)
+						{
+							continue;
+						}
+
+						FPrimitiveInstanceId Id = SharedIdMap.GetByElementId(ElementId).Value;
+						if (!Id.IsValid())
+						{
+							check(!bWasAlreadyInMap);
+							ToAdd.Emplace(FAddItem{ ElementId, IDIndex });
+						}
+						else
+						{
+							ensure(InstanceTransforms.IsValidIndex(IDIndex));
+							InstancedSkinnedMeshComponent->SetInstanceTransform(Id, InstanceTransforms[IDIndex]);
+							if (!PrevInstanceTransforms.IsEmpty())
+							{
+								InstancedSkinnedMeshComponent->SetInstancePrevTransform(Id, PrevInstanceTransforms[IDIndex]);
+							}
+							if (!InstanceAnimationData.IsEmpty())
+							{
+								int DesiredAnimIndex = -1;
+								int CurrentAnimIndex = -1;
+								if (UAnimSequenceTransformProviderDataInstance* ASTPDI = Cast<UAnimSequenceTransformProviderDataInstance>(InstancedSkinnedMeshComponent->GetTransformProvider()))
+								{
+									bool bAlreadyPresent = false;
+									auto TrackElementId = TrackMap.FindOrAddId(EntityHandlesWithPendingUpdates[IDIndex], 0, bAlreadyPresent);
+									int32 TrackId = INDEX_NONE;
+									if (bAlreadyPresent)
+									{
+										TrackId = TrackMap.GetByElementId(TrackElementId).Value;
+										if (ASTPDI->GetSequenceIndex(TrackId) != InstanceAnimationData[IDIndex].SequenceIndex)
+										{
+											ASTPDI->UpdateAutoPlayTrack(TrackId, InstanceAnimationData[IDIndex]);
+										}
+									}
+									else
+									{
+										TrackId = TrackMap.GetByElementId(TrackElementId).Value = ASTPDI->AllocateAutoPlayTrack(InstanceAnimationData[IDIndex]);
+									}
+							
+									// Readback resolved animation state for UAF handoff
+									const TArray<FAnimSequenceTransformProviderSequence>& Sequences = ASTPDI->GetSequences();
+									FMassSkinnedMeshResolvedAnimState State = {
+										.SequenceIndex = ASTPDI->GetSequenceIndex(TrackId),
+										.Position = ASTPDI->GetPosition(TrackId),
+										.PlayRate = ASTPDI->GetPlayRate(TrackId),
+										.AnimSequence = Sequences.IsValidIndex(State.SequenceIndex) ? Sequences[State.SequenceIndex].Sequence.Get() : nullptr
+									};
+									bool bAnimStatePresent = false;
+									auto AnimStateElementId = ResolvedAnimStates.FindOrAddId(EntityHandlesWithPendingUpdates[IDIndex], FMassSkinnedMeshResolvedAnimState(), bAnimStatePresent);
+									ResolvedAnimStates.GetByElementId(AnimStateElementId).Value = State;
+							
+									// Set the anim index on the instance to be the trackId
+									DesiredAnimIndex = TrackId;
+								}
+								else
+								{
+									DesiredAnimIndex = InstanceAnimationData[IDIndex].SequenceIndex;
+								}
+							
+								if (InstancedSkinnedMeshComponent->GetInstanceAnimationIndex(Id,CurrentAnimIndex) && CurrentAnimIndex != DesiredAnimIndex)
+								{
+									InstancedSkinnedMeshComponent->SetInstanceAnimationIndex(Id, DesiredAnimIndex);
+								}
+							}
+							if (!CustomDataFloats.IsEmpty())
+							{
+								InstancedSkinnedMeshComponent->SetCustomData(Id, MakeArrayView(CustomDataFloats.GetData() + IDIndex * NumCustomDataFloats, NumCustomDataFloats));
+							}
+						}
+
+						// Make sure we have enough space to track the already processed IDs
+						Unprocessed.SetNum(SharedIdMap.GetMaxIndex(), true);
+						Unprocessed[ElementId.GetIndex()] = false;
+					}
+					// Collect unwanted items & remove
+					TArray<FPrimitiveInstanceId> RemovedInstancedSkinnedMeshInstanceIds;
+					RemovedInstancedSkinnedMeshInstanceIds.Reserve(Unprocessed.Num());
+					{
+						for (TConstSetBitIterator<> BitIt(Unprocessed); BitIt; ++BitIt)
+						{
+							Experimental::FHashElementId ElementId(BitIt.GetIndex());
+							if (SharedIdMap.ContainsElementId(ElementId))
+							{
+								FPrimitiveInstanceId InstanceId = SharedIdMap.GetByElementId(ElementId).Value;
+								check(InstanceId.IsValid());
+								SharedIdMap.RemoveByElementId(ElementId);
+								RemovedInstancedSkinnedMeshInstanceIds.Add(InstanceId);
+							}
+						}
+						InstancedSkinnedMeshComponent->RemoveInstances(RemovedInstancedSkinnedMeshInstanceIds);
+					}
+					// Process deferred adds.
+					for (FAddItem AddItem : ToAdd)
+					{
+						FPrimitiveInstanceId Id = InstancedSkinnedMeshComponent->AddInstance(InstanceTransforms[AddItem.IDIndex], 0);
+						check(!SharedIdMap.GetByElementId(AddItem.ElementId).Value.IsValid());
+						SharedIdMap.GetByElementId(AddItem.ElementId).Value = Id;
+
+						if (!PrevInstanceTransforms.IsEmpty())
+						{
+							InstancedSkinnedMeshComponent->SetInstancePrevTransform(Id, PrevInstanceTransforms[AddItem.IDIndex]);
+						}
+						if (!CustomDataFloats.IsEmpty())
+						{
+							InstancedSkinnedMeshComponent->SetCustomData(Id, MakeArrayView(CustomDataFloats.GetData() + AddItem.IDIndex * NumCustomDataFloats, NumCustomDataFloats));
+						}
+						if (!InstanceAnimationData.IsEmpty())
+						{
+							int DesiredAnimIndex = -1;
+							int CurrentAnimIndex = -1;
+							auto EntityHandle = SharedIdMap.GetByElementId(AddItem.ElementId).Key;
+							if (UAnimSequenceTransformProviderDataInstance* ASTPDI = Cast<UAnimSequenceTransformProviderDataInstance>(InstancedSkinnedMeshComponent->GetTransformProvider()))
+							{
+								bool bAlreadyPresent = false;
+								auto TrackElementId = TrackMap.FindOrAddId(EntityHandle, 0, bAlreadyPresent);
+								int32 TrackId = INDEX_NONE;
+								if (bAlreadyPresent)
+								{
+									TrackId = TrackMap.GetByElementId(TrackElementId).Value;
+									if (ASTPDI->GetSequenceIndex(TrackId) != InstanceAnimationData[AddItem.IDIndex].SequenceIndex)
+									{
+										ASTPDI->UpdateAutoPlayTrack(TrackId, InstanceAnimationData[AddItem.IDIndex]);
+									}
+								}
+								else
+								{
+									TrackId = TrackMap.GetByElementId(TrackElementId).Value = ASTPDI->AllocateAutoPlayTrack(InstanceAnimationData[AddItem.IDIndex]);
+								}
+						
+								// Readback resolved animation state for UAF handoff
+								const TArray<FAnimSequenceTransformProviderSequence>& Sequences = ASTPDI->GetSequences();
+								FMassSkinnedMeshResolvedAnimState State = {
+									.SequenceIndex = ASTPDI->GetSequenceIndex(TrackId),
+									.Position = ASTPDI->GetPosition(TrackId),
+									.PlayRate = ASTPDI->GetPlayRate(TrackId),
+									.AnimSequence = Sequences.IsValidIndex(State.SequenceIndex) ? Sequences[State.SequenceIndex].Sequence.Get() : nullptr
+								};
+								bool bAnimStatePresent = false;
+								auto AnimStateElementId = ResolvedAnimStates.FindOrAddId(EntityHandlesWithPendingUpdates[AddItem.IDIndex], FMassSkinnedMeshResolvedAnimState(), bAnimStatePresent);
+								ResolvedAnimStates.GetByElementId(AnimStateElementId).Value = State;
+						
+								// Set the anim index on the instance to be the trackId
+								DesiredAnimIndex = TrackId;
+							}
+							else
+							{
+								DesiredAnimIndex = InstanceAnimationData[AddItem.IDIndex].SequenceIndex;
+							}
+						
+							if (InstancedSkinnedMeshComponent->GetInstanceAnimationIndex(Id,CurrentAnimIndex) && CurrentAnimIndex != DesiredAnimIndex)
+							{
+								InstancedSkinnedMeshComponent->SetInstanceAnimationIndex(Id, DesiredAnimIndex);
+							}
+						}
+					}
+					// note that we're not clearing the dirty flag on purpose - these components require constant updates
+				}
+			}
+
+			// bump the touch counter so that anyone caching data based on contents of this SharedData can detect the change
+			++SharedData.ComponentInstanceIdTouchCounter;
+		}
+		SharedData.ResetAccumulatedData();
+	}
 }
 
 //---------------------------------------------------------------
@@ -735,3 +994,548 @@ void FMassLODSignificanceRange::WriteCustomDataFloatsAtStartIndex(int32 StaticMe
 		}
 	}
 }
+
+FSkinnedMeshInstanceVisualizationDescHandle UMassVisualizationComponent::FindOrAddSkinnedMeshInstanceVisualDesc(const FSkinnedMeshInstanceVisualizationDesc& Desc)
+{
+	UE_MT_SCOPED_WRITE_ACCESS(InstancedSkinnedMeshInfosDetector);
+	// First check to see if we already have a matching Desc already and reuse / return that
+	// Note: FSkinnedMeshInstanceVisualizationDescHandle(int32) handles the INDEX_NONE case here, generating an invalid handle in this case
+	FSkinnedMeshInstanceVisualizationDescHandle VisualDescHandle(InstancedSkinnedMeshInfos.IndexOfByPredicate([&Desc](const FMassInstancedSkinnedMeshInfo& Info) { return Info.GetDesc() == Desc; }));
+	if (!VisualDescHandle.IsValid())
+	{
+		if (Desc.IsValid())
+		{
+			for (const FMassSkinnedMeshInstanceVisualizationMeshDesc& MeshDesc : Desc.Meshes)
+			{
+				if (MeshDesc.Asset && MeshDesc.InstancedSkinnedMeshComponentClass)
+				{
+					// if we've already encountered MeshDesc in the past MeshDescToInstancedSkinnedMeshComponentMap already contains information
+					// about actual Instanced Mesh Component used to represent it, and at the same time indicates the SharedData data
+					// tied to it. Regardless we need to process all MeshDesc instances here so that we have all the 
+					// data ready when InstancedSkinnedMeshComponentsRequiringConstructing gets processed next time
+					// UMassVisualizationComponent::ConstructSkinnedMeshComponents gets called.
+					MeshDescToInstancedSkinnedMeshComponentMap.FindOrAdd(GetTypeHash(MeshDesc), FInstancedSkinnedMeshComponentSharedDataKey());
+				}
+			}
+
+			VisualDescHandle = AddInstancedSkinnedMeshInfo(Desc);
+			check(VisualDescHandle.IsValid());
+
+			// VisualDescHandle is a valid handle now, but there's initialization pending, performed in ConstructSkinnedMeshComponents
+			InstancedSkinnedMeshComponentsRequiringConstructing.Add(VisualDescHandle);
+		}
+		else
+		{
+			UE_LOG(LogMassRepresentation, Warning, TEXT("%hs: invalid FSkinnedMeshInstanceVisualizationDesc passed in. Check the contained meshes."), __FUNCTION__);
+		}
+	}
+
+	return VisualDescHandle;
+}
+
+FSkinnedMeshInstanceVisualizationDescHandle UMassVisualizationComponent::AddSkinnedMeshInstanceVisualDescWithComponent(const FSkinnedMeshInstanceVisualizationDesc& Desc, UInstancedSkinnedMeshComponent& InstancedSkinnedMeshComponent)
+{
+
+	TObjectPtr<UInstancedSkinnedMeshComponent> AsObjectPtr = &InstancedSkinnedMeshComponent;
+	return AddSkinnedMeshInstanceVisualDescWithComponents(Desc, MakeArrayView(&AsObjectPtr, 1));
+}
+
+FSkinnedMeshInstanceVisualizationDescHandle UMassVisualizationComponent::AddSkinnedMeshInstanceVisualDescWithComponents(const FSkinnedMeshInstanceVisualizationDesc& Desc, TArrayView<TObjectPtr<UInstancedSkinnedMeshComponent>> InstancedSkinnedMeshComponents)
+{
+	check(Desc.Meshes.Num() == InstancedSkinnedMeshComponents.Num());
+
+	UE_MT_SCOPED_WRITE_ACCESS(InstancedSkinnedMeshInfosDetector);
+
+	// 0. Iterate over all meshes in the visualization desc. Each mesh is a descriptor.
+	FSkinnedMeshInstanceVisualizationDescHandle VisualHandle;
+
+	TArray<UInstancedSkinnedMeshComponent*> InstancedSkinnedMeshComponentsUsed;
+	for (int32 EntryIndex = 0; EntryIndex < Desc.Meshes.Num(); ++EntryIndex)
+	{
+		const FMassSkinnedMeshInstanceVisualizationMeshDesc& MeshDesc = Desc.Meshes[EntryIndex];
+		if (MeshDesc.Asset == nullptr || InstancedSkinnedMeshComponents[EntryIndex] == nullptr)
+		{
+			// invalid description, log an continue.
+			UE_VLOG_UELOG(this, LogMassRepresentation, Error, TEXT("Empty mesh at index %d while registering FSkinnedMeshInstanceVisualizationDesc instance"), EntryIndex);
+			continue;
+		}
+
+		// 1. Creates a VisualHandle that will be used for all InstancedSkinnedMeshComponents
+		if (!VisualHandle.IsValid())
+		{
+			VisualHandle = AddInstancedSkinnedMeshInfo(Desc);
+			check(VisualHandle.IsValid());
+		}
+
+		// 2. Stores one FMassInstancedSkinnedMeshComponentSharedData for each InstancedSkinnedMeshComponent in FMassInstancedSkinnedMeshComponentSharedData
+		// NOTE: FMassISMCSharedData stores a reference to an ISMC and instance transform updates performed in the current frame.
+		FMassInstancedSkinnedMeshComponentSharedData& NewData = InstancedSkinnedMeshComponentSharedData.FindOrAdd(InstancedSkinnedMeshComponents[EntryIndex], FMassInstancedSkinnedMeshComponentSharedData(InstancedSkinnedMeshComponents[EntryIndex], /*bInRequiresExternalInstanceIDTracking=*/true, /*InTransformOffset=*/MeshDesc.LocalTransform));
+		// 3. Stores the newly created FMassInstancedSkinnedMeshComponentSharedData in a VisualHandle -> FMassInstancedSkinnedMeshComponentSharedData array
+		// NOTE: This means that one VisualHandle may reference multiple FMassInstancedSkinnedMeshComponentSharedData
+		InstancedSkinnedMeshInfos[VisualHandle.ToIndex()].AddInstancedSkinnedMeshComponent(NewData);
+		// 4. Stores an InstancedMeshComponent -> VisualHandle map, which allows us to later on use an InstancedMeshComponent to query for its FMassInstancedSkinnedMeshComponentSharedData
+		InstancedSkinnedMeshComponentMap.Add(InstancedSkinnedMeshComponents[EntryIndex], VisualHandle);
+
+		InstancedSkinnedMeshComponentsUsed.Add(InstancedSkinnedMeshComponents[EntryIndex]);
+	}
+	if (VisualHandle.IsValid())
+	{
+		BuildLODSignificanceForInstancedSkinnedMeshInfo(InstancedSkinnedMeshInfos[VisualHandle.ToIndex()], InstancedSkinnedMeshComponentsUsed);
+	}
+
+	return VisualHandle;
+}
+
+const FMassInstancedSkinnedMeshComponentSharedData* UMassVisualizationComponent::GetSharedDataForDescriptionIndex(const int32 DescriptionIndex) const
+{
+	return InstancedSkinnedMeshComponentSharedData.GetDataForIndex(DescriptionIndex);
+}
+
+const FMassInstancedSkinnedMeshComponentSharedData* UMassVisualizationComponent::GetSharedDataForInstancedSkinnedMeshComponent(const UInstancedSkinnedMeshComponent* InstancedSkinnedMeshComponent) const
+{
+	return InstancedSkinnedMeshComponentSharedData.GetDataForKey(InstancedSkinnedMeshComponent);
+}
+
+void UMassVisualizationComponent::RemoveSkinnedMeshInstanceVisualDesc(const FSkinnedMeshInstanceVisualizationDescHandle VisualizationHandle)
+{
+	UE_MT_SCOPED_WRITE_ACCESS(InstancedSkinnedMeshInfosDetector);
+
+	if (ensure(InstancedSkinnedMeshInfos.IsValidIndex(VisualizationHandle.ToIndex()))
+		&& ensureMsgf(InstancedSkinnedMeshInfos[VisualizationHandle.ToIndex()].IsValid(), TEXT("Trying to remove visualization data that has already been cleaned")))
+	{
+		for (TObjectPtr<UInstancedSkinnedMeshComponent>& InstancedSkinnedMeshComponent : InstancedSkinnedMeshInfos[VisualizationHandle.ToIndex()].InstancedSkinnedMeshComponents)
+		{
+			const bool bValidKey = InstancedSkinnedMeshComponentMap.Contains(InstancedSkinnedMeshComponent);
+			checkf(bValidKey, TEXT("Failed to find Instanced Mesh in InstancedSkinnedMeshComponentMap, path: %s"), *InstancedSkinnedMeshComponent.GetPathName());
+			if (bValidKey)
+			{
+				const FSkinnedMeshInstanceVisualizationDescHandle StoredVisualizationDescHandle = InstancedSkinnedMeshComponentMap.FindAndRemoveChecked(InstancedSkinnedMeshComponent);
+				ensure(StoredVisualizationDescHandle == VisualizationHandle);
+			}
+
+			InstancedSkinnedMeshComponentSharedData.Remove(InstancedSkinnedMeshComponent);
+		}
+
+		InstancedSkinnedMeshInfos[VisualizationHandle.ToIndex()].Reset();
+		InstancedSkinnedMeshInfosFreeIndices.Add(VisualizationHandle);
+	}
+}
+
+void UMassVisualizationComponent::ProcessRemovesForComponent(UInstancedSkinnedMeshComponent& InstancedSkinnedMeshComponent, FMassInstancedSkinnedMeshComponentSharedData& SharedData)
+{
+	if (!SharedData.GetEntitiesRequiringRemoval().IsEmpty())
+	{
+		FMassInstancedSkinnedMeshComponentSharedData::FEntityToPrimitiveIdMap& SharedIdMap = SharedData.GetMutableEntityToPrimitiveIdMap();
+		FMassInstancedSkinnedMeshComponentSharedData::FEntityToTrackIdMap& SharedTrackMap = SharedData.GetMutableEntityToTrackMap();
+
+		INC_DWORD_STAT_BY(STAT_Mass_VisualizationComponent_InstancesRemovedNum, SharedData.GetEntitiesRequiringRemoval().Num());
+
+		TConstArrayView<FMassEntityHandle> EntityHandles = SharedData.GetEntitiesRequiringRemoval();
+
+		TArray<FPrimitiveInstanceId> InstancedSkinnedMeshInstanceIds;
+		InstancedSkinnedMeshInstanceIds.Reserve(EntityHandles.Num());
+
+		// Translate Mass IDs to Instanced Mesh IDs
+		for (const FMassEntityHandle MassInstanceId : EntityHandles)
+		{
+			Experimental::FHashElementId ElementId = SharedIdMap.FindId(MassInstanceId);
+			if (ElementId.IsValid())
+			{
+				FPrimitiveInstanceId InstanceId = SharedIdMap.GetByElementId(ElementId).Value;
+				check(InstanceId.IsValid());
+				SharedIdMap.RemoveByElementId(ElementId);
+				InstancedSkinnedMeshInstanceIds.Add(InstanceId);
+			}
+			Experimental::FHashElementId TrackElementId = SharedTrackMap.FindId(MassInstanceId);
+			if (TrackElementId.IsValid())
+			{
+				if (UAnimSequenceTransformProviderDataInstance* ASTPDI = Cast<UAnimSequenceTransformProviderDataInstance>(InstancedSkinnedMeshComponent.GetTransformProvider()))
+				{
+					ASTPDI->DeallocateTrack(SharedTrackMap.GetByElementId(TrackElementId).Value);
+				}
+				SharedTrackMap.RemoveByElementId(TrackElementId);
+
+				// Clean up resolved anim state
+				Experimental::FHashElementId AnimStateId = ResolvedAnimStates.FindId(MassInstanceId);
+				if (AnimStateId.IsValid())
+				{
+					ResolvedAnimStates.RemoveByElementId(AnimStateId);
+				}
+			}
+		}
+
+		InstancedSkinnedMeshComponent.RemoveInstances(InstancedSkinnedMeshInstanceIds);
+	}
+}
+
+void UMassVisualizationComponent::HandleChangesWithExternalIDTrackingForComponent(UInstancedSkinnedMeshComponent& InstancedSkinnedMeshComponent, FMassInstancedSkinnedMeshComponentSharedData& SharedData)
+{
+	if (SharedData.HasUpdatesToApply() == false)
+	{
+		// nothing to do here. We most probably were called as the part of the very first tick of this given SharedData
+		// since all the SharedData starts off as `dirty`.
+		return;
+	}
+
+	SCOPE_CYCLE_COUNTER(STAT_Mass_VisualizationComponent_HandleChangesWithExternalIDTracking);
+
+	// removing instances first, since this operation is more resilient to duplicates. Plus we make an arbitrary decision 
+	// that it's better to have redundant things visible than not seeing required things
+	ProcessRemovesForComponent(InstancedSkinnedMeshComponent, SharedData);
+
+	// NOTE: This code path is designed to only perform Adds, never updates so updates are filtered out along with duplicates.
+	TArray<FMassEntityHandle>& EntityHandles = SharedData.EntitiesRequiringUpdate;
+	if (!EntityHandles.IsEmpty())
+	{
+		INC_DWORD_STAT_BY(STAT_Mass_VisualizationComponent_InstancesAddedNum, EntityHandles.Num());
+
+		FMassISMCSharedData::FEntityToPrimitiveIdMap& SharedIdMap = SharedData.GetMutableEntityToPrimitiveIdMap();
+		TArray<Experimental::FHashElementId> ElementIds;
+		ElementIds.SetNumUninitialized(EntityHandles.Num());
+		// Filter out all updates & duplicate adds
+		for (int32 IDIndex = EntityHandles.Num() - 1; IDIndex >= 0; --IDIndex)
+		{
+			bool bWasAlreadyInMap = false;
+			Experimental::FHashElementId ElementId = SharedIdMap.FindOrAddId(EntityHandles[IDIndex], FPrimitiveInstanceId{ INDEX_NONE }, bWasAlreadyInMap);
+
+			if (bWasAlreadyInMap)
+			{
+				SharedData.RemoveUpdatedInstanceIdsAtSwap(IDIndex);
+				ElementIds.RemoveAtSwap(IDIndex);
+			}
+			else
+			{
+				ElementIds[IDIndex] = ElementId;
+			}
+		}
+
+		// it's possible the loop above removed all the data, so we do one last check
+		if (!EntityHandles.IsEmpty())
+		{
+			check(ElementIds.Num() == EntityHandles.Num());
+
+			const TArray<FTransform>& InstanceTransforms = SharedData.GetMeshInstanceTransformsArray();
+			
+			TArray<int32> InstanceAnimationIndices;
+			InstanceAnimationIndices.Init(0, InstanceTransforms.Num());
+
+			TConstArrayView<float> CustomFloatData = SharedData.GetMeshInstanceCustomFloats();
+			const int32 NumCustomDataFloatsPerInstance = CustomFloatData.Num() / EntityHandles.Num();
+
+			// if these are the first entities we're adding we need to set NumCustomDataFloats so that the PerInstanceSMCustomData
+			// gets populated properly by the AddInstancesById call below
+			const int32 StartingCount = InstancedSkinnedMeshComponent.GetInstanceCount();
+			const bool bInitiallyEmpty = (StartingCount == 0);
+			ensure(bInitiallyEmpty || InstancedSkinnedMeshComponent.GetNumCustomDataFloats() == NumCustomDataFloatsPerInstance);
+			if (bInitiallyEmpty && NumCustomDataFloatsPerInstance > 0 && InstancedSkinnedMeshComponent.Mobility != EComponentMobility::Static)
+			{
+				InstancedSkinnedMeshComponent.SetNumCustomDataFloats(NumCustomDataFloatsPerInstance);
+			}
+
+			check(EntityHandles.Num() == InstanceTransforms.Num());
+			TArray<FPrimitiveInstanceId> NewIds = InstancedSkinnedMeshComponent.AddInstances(InstanceTransforms, InstanceAnimationIndices, /*bShouldReturnIds=*/true, /*bWorldSpace=*/true);
+			check(EntityHandles.Num() == NewIds.Num());
+			int32 EntityHandleToUpdateIndex = StartingCount;
+			for (int32 EntityHandleIndex = 0; EntityHandleIndex < EntityHandles.Num(); ++EntityHandleIndex, ++EntityHandleToUpdateIndex)
+			{
+				SharedIdMap.GetByElementId(ElementIds[EntityHandleIndex]).Value = NewIds[EntityHandleIndex];	
+				
+				if (!CustomFloatData.IsEmpty())
+				{
+					InstancedSkinnedMeshComponent.SetCustomData(NewIds[EntityHandleIndex], MakeArrayView(CustomFloatData.GetData() + EntityHandleIndex * NumCustomDataFloatsPerInstance, NumCustomDataFloatsPerInstance));
+				}
+			}
+		}
+	}
+
+	if (bNavigationRelevant && InstancedSkinnedMeshComponent.GetInstanceCount() == 0)
+	{
+		FNavigationSystem::UnregisterComponent(InstancedSkinnedMeshComponent);
+	}
+}
+
+void UMassVisualizationComponent::ConstructSkinnedMeshComponents()
+{
+	AActor* ActorOwner = GetOwner();
+	check(ActorOwner);
+
+	TArray<UInstancedSkinnedMeshComponent*> TransientInstancedSkinnedMeshComponents;
+
+	UE_MT_SCOPED_WRITE_ACCESS(InstancedSkinnedMeshInfosDetector);
+	for (const FSkinnedMeshInstanceVisualizationDescHandle VisualDescHandle : InstancedSkinnedMeshComponentsRequiringConstructing)
+	{
+		if (!ensureMsgf(InstancedSkinnedMeshInfos.IsValidIndex(VisualDescHandle.ToIndex())
+			, TEXT("InstancedSkinnedMeshInfos (size: %d) is never expected to shrink, so VisualDescHandle (value: %u) being invalid indicates it was wrong from the start.")
+			, InstancedSkinnedMeshInfos.Num(), VisualDescHandle.ToIndex()))
+		{
+			continue;
+		}
+
+		FMassInstancedSkinnedMeshInfo& Info = InstancedSkinnedMeshInfos[VisualDescHandle.ToIndex()];
+
+		// Check if it is already created
+		if (!Info.InstancedSkinnedMeshComponents.IsEmpty())
+		{
+			continue;
+		}
+
+		// Check if there are any specified meshes for this visual type
+		if (Info.Desc.Meshes.Num() == 0)
+		{
+			UE_LOG(LogMassRepresentation, Error, TEXT("No associated meshes for this instanced static mesh type"));
+			continue;
+		}
+
+		TransientInstancedSkinnedMeshComponents.Reset();
+		for (const FMassSkinnedMeshInstanceVisualizationMeshDesc& InstancedSkinnedMeshDesc : Info.Desc.Meshes)
+		{
+			// MeshDescToInstancedSkinnedMeshComponentMap here lets us figure out whether for the given InstancedSkinnedMeshDesc we need to create a new Instanced Mesh component
+			// or a one has already been created in the past. Note that we only need this intermediate map for 
+			// FMassSkinnedMeshInstanceVisualizationMeshDesc that has been added to the system without specifying an
+			// ISM component to handle the instances (i.e. added via FindOrAddVisualDesc rather than AddVisualDescWithISMComponents).
+			// This is the only kind of FMassSkinnedMeshInstanceVisualizationMeshDesc were processing here. 
+			FInstancedSkinnedMeshComponentSharedDataKey& InstancedSkinnedMeshComponentKey = MeshDescToInstancedSkinnedMeshComponentMap.FindChecked(GetTypeHash(InstancedSkinnedMeshDesc));
+			FMassInstancedSkinnedMeshComponentSharedData* SharedData = InstancedSkinnedMeshComponentSharedData.Find(InstancedSkinnedMeshComponentKey);
+			UInstancedSkinnedMeshComponent* InstancedSkinnedMeshComponent = SharedData ? SharedData->GetMutableInstancedSkinnedMeshComponent() : nullptr;
+
+			if (InstancedSkinnedMeshComponent == nullptr)
+			{
+				InstancedSkinnedMeshComponent = NewObject<UInstancedSkinnedMeshComponent>(ActorOwner, InstancedSkinnedMeshDesc.InstancedSkinnedMeshComponentClass);
+				CA_ASSUME(InstancedSkinnedMeshComponent);
+				REDIRECT_OBJECT_TO_VLOG(InstancedSkinnedMeshComponent, this);
+
+				// note that InstancedSkinnedMeshComponentKey is a reference, so the assignment below actually sets a value in MeshDescToInstancedSkinnedMeshComponentMap
+				// and all subsequent handling of a given InstancedSkinnedMeshDesc configuration (i.e. containing same values) will 
+				// result in referring to the InstancedSkinnedMeshComponent we just created.
+				InstancedSkinnedMeshComponentKey = InstancedSkinnedMeshComponent;
+
+				InstancedSkinnedMeshComponent->SetSkinnedAssetAndUpdate(InstancedSkinnedMeshDesc.Asset);
+				for (int32 ElementIndex = 0; ElementIndex < InstancedSkinnedMeshDesc.MaterialOverrides.Num(); ++ElementIndex)
+				{
+					if (UMaterialInterface* MaterialOverride = InstancedSkinnedMeshDesc.MaterialOverrides[ElementIndex])
+					{
+						InstancedSkinnedMeshComponent->SetMaterial(ElementIndex, MaterialOverride);
+					}
+				}
+
+				if (UAnimSequenceTransformProviderData* ASTPD = Cast<UAnimSequenceTransformProviderData>(InstancedSkinnedMeshDesc.TransformProvider))
+				{
+					InstancedSkinnedMeshComponent->SetTransformProvider(UAnimSequenceTransformProviderDataInstance::CreateAnimSequenceTransformProviderDataInstance(ASTPD, InstancedSkinnedMeshComponent));
+				}
+				else
+				{
+					InstancedSkinnedMeshComponent->SetTransformProvider(InstancedSkinnedMeshDesc.TransformProvider);
+				}
+
+				InstancedSkinnedMeshComponent->SetCullDistances(0, 1000000); // @todo: Need to figure out what to do here, either LOD or cull distances.
+				InstancedSkinnedMeshComponent->SetupAttachment(ActorOwner->GetRootComponent());
+				InstancedSkinnedMeshComponent->SetCanEverAffectNavigation(false);
+				InstancedSkinnedMeshComponent->SetCollisionProfileName(UCollisionProfile::NoCollision_ProfileName);
+				InstancedSkinnedMeshComponent->SetCastShadow(InstancedSkinnedMeshDesc.bCastShadows);
+				InstancedSkinnedMeshComponent->Mobility = InstancedSkinnedMeshDesc.Mobility;
+				InstancedSkinnedMeshComponent->SetReceivesDecals(false);
+				InstancedSkinnedMeshComponent->RegisterComponent();
+
+				if (SharedData == nullptr)
+				{
+					SharedData = &InstancedSkinnedMeshComponentSharedData.Add(InstancedSkinnedMeshComponent, FMassInstancedSkinnedMeshComponentSharedData(InstancedSkinnedMeshComponent, InstancedSkinnedMeshDesc.bRequiresExternalInstanceIDTracking, InstancedSkinnedMeshDesc.LocalTransform));
+				}
+				else
+				{
+					SharedData->SetInstancedSkinnedMeshComponent(*InstancedSkinnedMeshComponent);
+				}
+
+				ensureMsgf(InstancedSkinnedMeshComponentMap.Find(InstancedSkinnedMeshComponent) == nullptr, TEXT("We've just created the ISMC that's being used here, so this check failing indicates hash-clash."));
+				InstancedSkinnedMeshComponentMap.Add(InstancedSkinnedMeshComponent, VisualDescHandle);
+			}
+
+			TransientInstancedSkinnedMeshComponents.Add(InstancedSkinnedMeshComponent);
+
+			check(SharedData);
+			Info.AddInstancedSkinnedMeshComponent(*SharedData);
+		}
+
+		// Build the LOD significance ranges
+		if (TransientInstancedSkinnedMeshComponents.Num())
+		{
+			check(Info.LODSignificanceRanges.Num() == 0);
+			BuildLODSignificanceForInstancedSkinnedMeshInfo(Info, TransientInstancedSkinnedMeshComponents);
+		}
+	}
+}
+
+void UMassVisualizationComponent::BuildLODSignificanceForInstancedSkinnedMeshInfo(FMassInstancedSkinnedMeshInfo& Info, TConstArrayView<UInstancedSkinnedMeshComponent*> MeshRefKeys)
+{
+	TArray<float> AllLODSignificances;
+	auto UniqueInsertOrdered = [&AllLODSignificances](const float Significance)
+		{
+			int i = 0;
+			for (; i < AllLODSignificances.Num(); ++i)
+			{
+				// I did not use epsilon check here on purpose, because it will make it hard later meshes inside.
+				if (Significance == AllLODSignificances[i])
+				{
+					return;
+				}
+				if (AllLODSignificances[i] > Significance)
+				{
+					break;
+				}
+			}
+			AllLODSignificances.Insert(Significance, i);
+		};
+	for (const FMassSkinnedMeshInstanceVisualizationMeshDesc& MeshDesc : Info.Desc.Meshes)
+	{
+		UniqueInsertOrdered(MeshDesc.MinLODSignificance);
+		UniqueInsertOrdered(MeshDesc.MaxLODSignificance);
+	}
+
+	if (AllLODSignificances.Num() > 1)
+	{
+		Info.LODSignificanceRanges.SetNum(AllLODSignificances.Num() - 1);
+		for (int RangeIndex = 0; RangeIndex < Info.LODSignificanceRanges.Num(); ++RangeIndex)
+		{
+			FMassLODInstancedSkinnedMeshSignificanceRange& Range = Info.LODSignificanceRanges[RangeIndex];
+			Range.MinSignificance = AllLODSignificances[RangeIndex];
+			Range.MaxSignificance = AllLODSignificances[RangeIndex + 1];
+			Range.InstancedSkinnedMeshSharedDataPtr = &InstancedSkinnedMeshComponentSharedData;
+
+			for (int MeshIndex = 0; MeshIndex < Info.Desc.Meshes.Num(); ++MeshIndex)
+			{
+				const FMassSkinnedMeshInstanceVisualizationMeshDesc& SkinnedMeshDesc = Info.Desc.Meshes[MeshIndex];
+				const bool bAddMeshInRange = (Range.MinSignificance >= SkinnedMeshDesc.MinLODSignificance && Range.MinSignificance < SkinnedMeshDesc.MaxLODSignificance);
+				if (bAddMeshInRange)
+				{
+					checkf(MeshRefKeys.IsValidIndex(MeshIndex) && MeshRefKeys[MeshIndex]
+						, TEXT("We don't expect receiving null ISMCs at this point"));
+					Range.SkinnedMeshComponentRefs.Add(MeshRefKeys[MeshIndex]);
+				}
+			}
+		}
+	}
+}
+
+FSkinnedMeshInstanceVisualizationDescHandle UMassVisualizationComponent::AddInstancedSkinnedMeshInfo(const FSkinnedMeshInstanceVisualizationDesc& Desc)
+{
+	FSkinnedMeshInstanceVisualizationDescHandle Handle;
+	if (InstancedSkinnedMeshInfosFreeIndices.Num() > 0)
+	{
+		Handle = InstancedSkinnedMeshInfosFreeIndices.Pop(EAllowShrinking::No);
+		new(&InstancedSkinnedMeshInfos[Handle.ToIndex()]) FMassInstancedSkinnedMeshInfo(Desc);
+	}
+	else
+	{
+		int32 AddedInfoIndex = InstancedSkinnedMeshInfos.Emplace(Desc);
+		Handle = FSkinnedMeshInstanceVisualizationDescHandle(AddedInfoIndex);
+	}
+
+	return Handle;
+}
+			
+void FMassLODInstancedSkinnedMeshSignificanceRange::AddBatchedTransform(const FMassEntityHandle EntityHandle, const FTransform& Transform, const FTransform& PrevTransform, TConstArrayView<FInstancedSkinnedMeshComponentSharedDataKey> ExcludeMeshRefs)
+{
+	check(InstancedSkinnedMeshSharedDataPtr);
+	for (int32 MeshIndex = 0; MeshIndex < SkinnedMeshComponentRefs.Num(); ++MeshIndex)
+	{
+		if (ExcludeMeshRefs.Contains(SkinnedMeshComponentRefs[MeshIndex]))
+		{
+			continue;
+		}
+
+		if (FMassInstancedSkinnedMeshComponentSharedData* SharedData = InstancedSkinnedMeshSharedDataPtr->GetAndMarkDirty(SkinnedMeshComponentRefs[MeshIndex]))
+		{
+			SharedData->EntitiesRequiringUpdate.Add(EntityHandle);
+
+			SharedData->MeshInstanceTransforms.Add(SharedData->LocalTransform * Transform);
+			SharedData->MeshInstancePrevTransforms.Add(SharedData->LocalTransform * PrevTransform);
+		}
+	}
+}
+
+void FMassLODInstancedSkinnedMeshSignificanceRange::AddBatchedAnimationData(const FMassEntityHandle EntityHandle, FAnimSequenceTrackAutoPlayData AnimationData, TConstArrayView<FInstancedSkinnedMeshComponentSharedDataKey> ExcludeMeshRefs)
+{
+	check(InstancedSkinnedMeshSharedDataPtr);
+	for (int32 MeshIndex = 0; MeshIndex < SkinnedMeshComponentRefs.Num(); ++MeshIndex)
+	{
+		if (ExcludeMeshRefs.Contains(SkinnedMeshComponentRefs[MeshIndex]))
+		{
+			continue;
+		}
+
+		if (FMassInstancedSkinnedMeshComponentSharedData* SharedData = InstancedSkinnedMeshSharedDataPtr->GetAndMarkDirty(SkinnedMeshComponentRefs[MeshIndex]))
+		{
+			SharedData->MeshInstanceAnimationData.Add(AnimationData);
+		}
+	}
+}
+
+void FMassLODInstancedSkinnedMeshSignificanceRange::AddBatchedCustomDataFloats(const TArray<float>& CustomFloats, const TArray<FInstancedSkinnedMeshComponentSharedDataKey>& ExcludeMeshRefs)
+{
+	check(InstancedSkinnedMeshSharedDataPtr);
+	for (int32 MeshIndex = 0; MeshIndex < SkinnedMeshComponentRefs.Num(); ++MeshIndex)
+	{
+		if (ExcludeMeshRefs.Contains(SkinnedMeshComponentRefs[MeshIndex]))
+		{
+			continue;
+		}
+
+		if (FMassInstancedSkinnedMeshComponentSharedData* SharedData = InstancedSkinnedMeshSharedDataPtr->GetAndMarkDirty(SkinnedMeshComponentRefs[MeshIndex]))
+		{
+			SharedData->MeshInstanceCustomFloats.Append(CustomFloats);
+		}
+	}
+}
+
+void FMassLODInstancedSkinnedMeshSignificanceRange::AddInstance(const FMassEntityHandle EntityHandle, const FTransform& Transform)
+{
+	check(InstancedSkinnedMeshSharedDataPtr);
+	for (int32 MeshIndex = 0; MeshIndex < SkinnedMeshComponentRefs.Num(); ++MeshIndex)
+	{
+		if (FMassInstancedSkinnedMeshComponentSharedData* SharedData = InstancedSkinnedMeshSharedDataPtr->GetAndMarkDirty(SkinnedMeshComponentRefs[MeshIndex]))
+		{
+			SharedData->EntitiesRequiringUpdate.Add(EntityHandle);
+
+			FTransform AdjustedTransform{ SharedData->LocalTransform * Transform };
+
+			SharedData->MeshInstanceTransforms.Add(AdjustedTransform);
+			SharedData->MeshInstancePrevTransforms.Add(AdjustedTransform);
+		}
+	}
+}
+
+void FMassLODInstancedSkinnedMeshSignificanceRange::RemoveInstance(const FMassEntityHandle EntityHandle)
+{
+	check(InstancedSkinnedMeshSharedDataPtr);
+	for (int32 MeshIndex = 0; MeshIndex < SkinnedMeshComponentRefs.Num(); ++MeshIndex)
+	{
+		if (FMassInstancedSkinnedMeshComponentSharedData* SharedData = InstancedSkinnedMeshSharedDataPtr->GetAndMarkDirty(SkinnedMeshComponentRefs[MeshIndex]))
+		{
+			SharedData->EntitiesRequiringRemoval.Add(EntityHandle);
+		}
+	}
+}
+
+void FMassLODInstancedSkinnedMeshSignificanceRange::WriteCustomDataFloatsAtStartIndex(int32 MeshIndex, const TArrayView<float>& CustomFloats, const int32 FloatsPerInstance, const int32 StartFloatIndex, const TArray<FInstancedSkinnedMeshComponentSharedDataKey>& ExcludeMeshRefs)
+{
+	check(InstancedSkinnedMeshSharedDataPtr);
+	if (SkinnedMeshComponentRefs.IsValidIndex(MeshIndex))
+	{
+		if (ExcludeMeshRefs.Contains(SkinnedMeshComponentRefs[MeshIndex]))
+		{
+			return;
+		}
+
+		if (FMassInstancedSkinnedMeshComponentSharedData* SharedData = InstancedSkinnedMeshSharedDataPtr->GetAndMarkDirty(SkinnedMeshComponentRefs[MeshIndex]))
+		{
+			const int32 StartIndex = FloatsPerInstance * SharedData->WriteIterator + StartFloatIndex;
+
+			ensure(SharedData->MeshInstanceCustomFloats.Num() >= StartIndex + CustomFloats.Num());
+
+			for (int CustomFloatIdx = 0; CustomFloatIdx < CustomFloats.Num(); CustomFloatIdx++)
+			{
+				SharedData->MeshInstanceCustomFloats[StartIndex + CustomFloatIdx] = CustomFloats[CustomFloatIdx];
+			}
+			SharedData->WriteIterator++;
+		}
+	}
+}
\ No newline at end of file
diff --git a/Engine/Plugins/Runtime/MassGameplay/Source/MassRepresentation/Private/MassVisualizationTrait.cpp b/Engine/Plugins/Runtime/MassGameplay/Source/MassRepresentation/Private/MassVisualizationTrait.cpp
--- a/Engine/Plugins/Runtime/MassGameplay/Source/MassRepresentation/Private/MassVisualizationTrait.cpp
+++ b/Engine/Plugins/Runtime/MassGameplay/Source/MassRepresentation/Private/MassVisualizationTrait.cpp
@@ -109,15 +109,31 @@ void UMassVisualizationTrait::BuildTemplate(FMassEntityTemplateBuildContext& Bui
 		}
 	}
 
+	bool bSkinnedMeshDescriptionValid = SkinnedMeshInstanceDesc.IsValid();
+	if (bSkinnedMeshDescriptionValid)
+	{
+		if (!BuildContext.IsInspectingData())
+		{
+			RepresentationFragment.SkinnedMeshDescHandle = RepresentationSubsystem->FindOrAddSkinnedMeshVisualizationDesc(SkinnedMeshInstanceDesc);
+			ensureMsgf(RepresentationFragment.SkinnedMeshDescHandle.IsValid()
+				, TEXT("Expected to get a valid SkinnedMeshDescHandle since we already checked that SkinnedMeshInstanceDesc is valid"));
+			// if the unexpected happens and SkinnedMeshDescHandle is not valid we're going to treat it as if SkinnedMeshInstanceDesc
+			// was not valid in the first place and handle it accordingly in a moment
+			bSkinnedMeshDescriptionValid = RepresentationFragment.SkinnedMeshDescHandle.IsValid();
+		}
+	}
+
+
+
 	FConstSharedStruct ParamsFragment;
-	if (bStaticMeshDescriptionValid)
+	if (bStaticMeshDescriptionValid && bSkinnedMeshDescriptionValid)
 	{
 		ParamsFragment = EntityManager.GetOrCreateConstSharedFragment(Params);
 	}
 	else
 	{
 		FMassRepresentationParameters ParamsCopy = Params;
-		SanitizeParams(ParamsCopy, bStaticMeshDescriptionValid);
+		SanitizeMeshParams(ParamsCopy, bStaticMeshDescriptionValid, bSkinnedMeshDescriptionValid);
 		ParamsFragment = EntityManager.GetOrCreateConstSharedFragment(ParamsCopy);
 	}
 	ParamsFragment.Get<const FMassRepresentationParameters>().ComputeCachedValues();
@@ -138,6 +154,11 @@ void UMassVisualizationTrait::BuildTemplate(FMassEntityTemplateBuildContext& Bui
 }
 
 void UMassVisualizationTrait::SanitizeParams(FMassRepresentationParameters& InOutParams, const bool bStaticMeshDeterminedInvalid) const
+{
+	SanitizeMeshParams(InOutParams, bStaticMeshDeterminedInvalid);
+}
+
+void UMassVisualizationTrait::SanitizeMeshParams(FMassRepresentationParameters& InOutParams, const bool bStaticMeshDeterminedInvalid /*= false*/, const bool bSkinnedMeshDeterminedInvalid /*= false*/) const
 {
 	if (bStaticMeshDeterminedInvalid || (StaticMeshInstanceDesc.IsValid() == false))
 	{
@@ -149,6 +170,17 @@ void UMassVisualizationTrait::SanitizeParams(FMassRepresentationParameters& InOu
 			}
 		}
 	}
+
+	if (bSkinnedMeshDeterminedInvalid || (SkinnedMeshInstanceDesc.IsValid() == false))
+	{
+		for (int32 LODIndex = 0; LODIndex < EMassLOD::Max; ++LODIndex)
+		{
+			if (InOutParams.LODRepresentation[LODIndex] == EMassRepresentationType::SkinnedMeshInstance)
+			{
+				InOutParams.LODRepresentation[LODIndex] = EMassRepresentationType::None;
+			}
+		}
+	}
 }
 
 void UMassVisualizationTrait::Serialize(FArchive& Ar)
diff --git a/Engine/Plugins/Runtime/MassGameplay/Source/MassRepresentation/Public/MassRepresentationAnimationProcessor.h b/Engine/Plugins/Runtime/MassGameplay/Source/MassRepresentation/Public/MassRepresentationAnimationProcessor.h
new file mode 100644
--- /dev/null
+++ b/Engine/Plugins/Runtime/MassGameplay/Source/MassRepresentation/Public/MassRepresentationAnimationProcessor.h
@@ -0,0 +1,39 @@
+// Copyright Epic Games, Inc. All Rights Reserved.
+
+#pragma once
+
+#include "MassEntityQuery.h"
+#include "MassProcessor.h"
+#include "MassRepresentationAnimationProcessor.generated.h"
+
+#define UE_API MASSREPRESENTATION_API
+
+struct FMassInstancedSkinnedMeshInfo;
+struct FAnimSequenceTrackAutoPlayData;
+
+/**
+ * This processor's responsibility is to consume Animation Data generated by end user processors and forward it onto the Skinned Mesh
+ * this data will be batch updated in the Instanced Skinned Mesh Component
+ */
+UCLASS(MinimalAPI)
+class UMassConsumeInstancedSkinnedMeshAnimationProcessor : public UMassProcessor
+{
+	GENERATED_BODY()
+
+public:
+	UE_API UMassConsumeInstancedSkinnedMeshAnimationProcessor();
+
+	static UE_API void UpdateMeshAnimation(FMassEntityHandle EntityHandle, FMassInstancedSkinnedMeshInfo& MeshInfo, FAnimSequenceTrackAutoPlayData AnimationData, const float LODSignificance, const float PrevLODSignificance = -1.0f);
+
+protected:
+
+	// UMassProcessor overrides begin
+	UE_API virtual void ConfigureQueries(const TSharedRef<FMassEntityManager>& EntityManager) override;
+
+	UE_API virtual void Execute(FMassEntityManager& EntityManager, FMassExecutionContext& Context) override;
+	// UMassProcessor overrides end
+
+	FMassEntityQuery EntityQuery;
+};
+
+#undef UE_API
\ No newline at end of file
diff --git a/Engine/Plugins/Runtime/MassGameplay/Source/MassRepresentation/Public/MassRepresentationFragments.h b/Engine/Plugins/Runtime/MassGameplay/Source/MassRepresentation/Public/MassRepresentationFragments.h
--- a/Engine/Plugins/Runtime/MassGameplay/Source/MassRepresentation/Public/MassRepresentationFragments.h
+++ b/Engine/Plugins/Runtime/MassGameplay/Source/MassRepresentation/Public/MassRepresentationFragments.h
@@ -9,6 +9,7 @@
 #endif // UE_ENABLE_INCLUDE_ORDER_DEPRECATED_IN_5_6
 #include "MassActorSpawnerSubsystem.h"
 #include "MassRepresentationTypes.h"
+#include "MassSkinnedMeshRepresentationTypes.h"
 #include "MassRepresentationActorManagement.h"
 #include "MassLODCalculator.h"
 #include "MassRepresentationFragments.generated.h"
@@ -67,6 +68,9 @@ struct FMassRepresentationFragment : public FMassFragment
     UPROPERTY()
     FStaticMeshInstanceVisualizationDescHandle StaticMeshDescHandle;
 
+	UPROPERTY()
+    FSkinnedMeshInstanceVisualizationDescHandle SkinnedMeshDescHandle;
+
     /** Value scaling from 0 to 3, 0 highest LOD we support and 3 being completely off LOD */
     UPROPERTY()
     float PrevLODSignificance = -1.0f;
@@ -81,6 +85,15 @@ struct FMassRepresentationFragment : public FMassFragment
     FTransform PrevTransform;
 };
 
+USTRUCT()
+struct FMassRepresentationAnimationFragment : public FMassFragment
+{
+    GENERATED_BODY()
+
+	UPROPERTY(Transient)
+	FAnimSequenceTrackAutoPlayData AnimData;
+};
+
 USTRUCT()
 struct FMassRepresentationSubsystemSharedFragment : public FMassSharedFragment
 {
diff --git a/Engine/Plugins/Runtime/MassGameplay/Source/MassRepresentation/Public/MassRepresentationSubsystem.h b/Engine/Plugins/Runtime/MassGameplay/Source/MassRepresentation/Public/MassRepresentationSubsystem.h
--- a/Engine/Plugins/Runtime/MassGameplay/Source/MassRepresentation/Public/MassRepresentationSubsystem.h
+++ b/Engine/Plugins/Runtime/MassGameplay/Source/MassRepresentation/Public/MassRepresentationSubsystem.h
@@ -8,6 +8,7 @@
 #include "MassCommonTypes.h"
 #endif // UE_ENABLE_INCLUDE_ORDER_DEPRECATED_IN_5_6
 #include "MassRepresentationTypes.h"
+#include "MassSkinnedMeshRepresentationTypes.h"
 #include "MassActorSpawnerSubsystem.h"
 #include "MassSubsystemBase.h"
 #include "MassExternalSubsystemTraits.h"
@@ -154,6 +155,9 @@ public:
 
 	UMassActorSpawnerSubsystem* GetActorSpawnerSubsystem() const { return ActorSpawnerSubsystem; }
 
+	/** Access to the visualization component for animation state readback */
+	UMassVisualizationComponent* GetVisualizationComponent() const { return VisualizationComponent; }
+
 protected:
 	// USubsystem BEGIN
 	MASSREPRESENTATION_API virtual void Initialize(FSubsystemCollectionBase& Collection) override;
@@ -173,6 +177,20 @@ protected:
 
 	static MASSREPRESENTATION_API void AddReferencedObjects(UObject* InThis, FReferenceCollector& Collector);
 
+public:
+
+	/**
+	 * @return the array of all the skinned mesh instance component information
+	 */
+	MASSREPRESENTATION_API FMassInstancedSkinnedMeshInfoArrayView GetMutableInstancedSkinnedMeshInfos();
+
+	/** 
+	 * Get the index of the skinned mesh visual type, will add a new one if does not exist
+	 * @param Desc is the information for the mesh that will be instantiated later when adding a skinned mesh instance
+	 * @return The index of the mesh type
+	 */
+	MASSREPRESENTATION_API FSkinnedMeshInstanceVisualizationDescHandle FindOrAddSkinnedMeshVisualizationDesc(const FSkinnedMeshInstanceVisualizationDesc& Desc);
+
 protected:
 
 	struct FTemplateActorData
@@ -228,4 +246,7 @@ private:
 
 	/** Used for default return value from GetMutableInstancedStaticMeshInfos */
 	UE_MT_DECLARE_RW_ACCESS_DETECTOR(InstancedStaticMeshInfosDetector);
+
+	/** Used for default return value from GetMutableInstancedSkinedMeshInfos */
+	UE_MT_DECLARE_RW_ACCESS_DETECTOR(InstancedSkinnedMeshInfosDetector);
 };
diff --git a/Engine/Plugins/Runtime/MassGameplay/Source/MassRepresentation/Public/MassRepresentationTypes.h b/Engine/Plugins/Runtime/MassGameplay/Source/MassRepresentation/Public/MassRepresentationTypes.h
--- a/Engine/Plugins/Runtime/MassGameplay/Source/MassRepresentation/Public/MassRepresentationTypes.h
+++ b/Engine/Plugins/Runtime/MassGameplay/Source/MassRepresentation/Public/MassRepresentationTypes.h
@@ -34,10 +34,25 @@ enum class EMassRepresentationType : uint8
 {
 	HighResSpawnedActor,
 	LowResSpawnedActor,
+	SkinnedMeshInstance,
 	StaticMeshInstance,
 	None,
 };
 
+
+namespace UE::Mass::Representation
+{
+	UE_API inline bool IsValidMeshRepresentation(EMassRepresentationType RepresentationType)
+	{
+		return RepresentationType == EMassRepresentationType::SkinnedMeshInstance || RepresentationType == EMassRepresentationType::StaticMeshInstance;
+	}
+
+	UE_API inline bool IsValidActorRepresentation(EMassRepresentationType RepresentationType)
+	{
+		return RepresentationType == EMassRepresentationType::HighResSpawnedActor || RepresentationType == EMassRepresentationType::LowResSpawnedActor;
+	}
+}
+
 enum class EMassActorEnabledType : uint8
 {
 	Disabled,
@@ -162,7 +177,10 @@ struct FStaticMeshInstanceVisualizationDesc : public FTableRowBase
 
 	void Reset()
 	{
-		new(this)FStaticMeshInstanceVisualizationDesc();
+		Meshes.Empty();
+		bUseTransformOffset = false;
+		TransformOffset = FTransform::Identity;
+		CustomDataFloats.Empty();
 	}
 
 	/** @return whether any of descriptions in Meshes is valid. This implies that empty Meshes will be treated as not valid.*/
@@ -562,8 +580,8 @@ public:
 	UE_API void WriteCustomDataFloatsAtStartIndex(int32 StaticMeshIndex, const TArrayView<float>& CustomFloats, const int32 FloatsPerInstance, const int32 StartIndex, const TArray<FISMCSharedDataKey>& ExcludeStaticMeshRefs);
 
 	/** LOD Significance range */
-	float MinSignificance;
-	float MaxSignificance;
+	float MinSignificance = 0.0f;
+	float MaxSignificance = 0.0f;
 
 	/** The component handling these instances */
 	TArray<FISMCSharedDataKey> StaticMeshRefs;
@@ -771,14 +789,14 @@ private:
 	const FRWAccessDetector* AccessDetector;
 };
 
-typedef FMassInstancedStaticMeshInfoArrayViewAccessDetector FMassInstancedStaticMeshInfoArrayView;
+using FMassInstancedStaticMeshInfoArrayView = FMassInstancedStaticMeshInfoArrayViewAccessDetector;
 
 #else // ENABLE_MT_DETECTOR
 
 #define MAKE_MASS_INSTANCED_STATIC_MESH_INFO_ARRAY_VIEW(ArrayView, AccessDetector) ArrayView
 
-typedef TArrayView<FMassInstancedStaticMeshInfo> FMassInstancedStaticMeshInfoArrayView;
+using FMassInstancedStaticMeshInfoArrayView = TArrayView<FMassInstancedStaticMeshInfo>;
 
 #endif // ENABLE_MT_DETECTOR
 
-#undef UE_API
\ No newline at end of file
+#undef UE_API
diff --git a/Engine/Plugins/Runtime/MassGameplay/Source/MassRepresentation/Public/MassSkinnedMeshRepresentationTypes.h b/Engine/Plugins/Runtime/MassGameplay/Source/MassRepresentation/Public/MassSkinnedMeshRepresentationTypes.h
new file mode 100644
--- /dev/null
+++ b/Engine/Plugins/Runtime/MassGameplay/Source/MassRepresentation/Public/MassSkinnedMeshRepresentationTypes.h
@@ -0,0 +1,830 @@
+// Copyright Epic Games, Inc. All Rights Reserved.
+
+#pragma once
+
+#include "Animation/AnimSequence.h"
+#include "Components/InstancedSkinnedMeshComponent.h"
+#include "Engine/DataTable.h"
+#include "Experimental/Containers/RobinHoodHashTable.h"
+#include "InstanceDataTypes.h"
+#include "MassEntityHandle.h"
+#include "MassLODTypes.h"
+#include "MassRepresentationTypes.h"
+#include "Math/MathFwd.h"
+#include "Misc/MTAccessDetector.h"
+#include "UObject/ObjectKey.h"
+
+#include "MassSkinnedMeshRepresentationTypes.generated.h"
+
+#define UE_API MASSREPRESENTATION_API
+
+class UMaterialInterface;
+struct FMassLODSignificanceRange;
+class UMassVisualizationComponent;
+class USkinnedAsset;
+class UTransformProviderData;
+
+/**
+// * Resolved animation state read back from ISKM tracks.
+// * Populated during EndVisualChanges() where the ASTPD track data is already hot in cache.
+// * Consumed by processors that need to push animation state to UAF.
+// */
+struct FMassSkinnedMeshResolvedAnimState
+{
+	int32 SequenceIndex = INDEX_NONE;
+	float Position = 0.0f;
+	float PlayRate = 1.0f;
+	TWeakObjectPtr<UAnimSequence> AnimSequence = nullptr;
+};
+
+using FInstancedSkinnedMeshComponentSharedDataKey = TObjectKey<UInstancedSkinnedMeshComponent>;
+
+USTRUCT()
+struct FMassSkinnedMeshInstanceVisualizationMeshDesc
+{
+	GENERATED_BODY()
+
+	UE_API FMassSkinnedMeshInstanceVisualizationMeshDesc();
+
+	/** A relative Transform from the parent, passed along to FMassInstancedSkinnedMeshComponentSharedData::LocalTransform */
+	UPROPERTY(EditAnywhere, Category = "Mass|Visual")
+	FTransform LocalTransform = FTransform::Identity;
+	
+	/** The mesh visual representation */
+	UPROPERTY(EditAnywhere, Category = "Mass|Visual")
+	TObjectPtr<USkinnedAsset> Asset = nullptr;
+
+	/** The Animation Transform Provider when using this mesh*/
+	UPROPERTY(EditAnywhere, Category = "Mass|Visual")
+	TObjectPtr<UTransformProviderData> TransformProvider = nullptr;
+
+	/**
+	 * Material overrides for the mesh visual representation. 
+	 * 
+	 * Array indices correspond to material slot indices on the mesh.
+	 */
+	UPROPERTY(EditAnywhere, Category = "Mass|Visual")
+	TArray<TObjectPtr<UMaterialInterface>> MaterialOverrides;
+
+	/** The minimum inclusive LOD significance to start using this mesh */
+	UPROPERTY(EditAnywhere, Category = "Mass|Visual")
+	float MinLODSignificance = float(EMassLOD::High);
+
+	/** The maximum exclusive LOD significance to stop using this mesh */
+	UPROPERTY(EditAnywhere, Category = "Mass|Visual")
+	float MaxLODSignificance = float(EMassLOD::Max);
+
+	/** Controls whether the ISM can cast shadow or not */
+	UPROPERTY(EditAnywhere, Category = "Mass|Visual")
+	bool bCastShadows = false;
+
+	/**
+	 * Determines the value of corresponding FMassISMCSharedData.bRequiresExternalInstanceIDTracking.
+	 * Recommended enabled for stationary meshes.
+	 * @see FMassISMCSharedData.bRequiresExternalInstanceIDTracking for details.
+	 */
+	bool bRequiresExternalInstanceIDTracking = false;
+	
+	/** Controls the mobility of the ISM */
+	EComponentMobility::Type Mobility = EComponentMobility::Movable;
+
+	/** InstancedSkinnedMeshComponent class to use to manage instances described by this struct instance */
+	UPROPERTY(EditAnywhere, Category = "Mass|Visual")
+	TSubclassOf<UInstancedSkinnedMeshComponent> InstancedSkinnedMeshComponentClass;
+
+	bool operator==(const FMassSkinnedMeshInstanceVisualizationMeshDesc& Other) const
+	{		
+		return Asset == Other.Asset && 
+			TransformProvider == Other.TransformProvider &&
+			InstancedSkinnedMeshComponentClass == Other.InstancedSkinnedMeshComponentClass &&
+			MaterialOverrides == Other.MaterialOverrides && 
+			FMath::IsNearlyEqual(MinLODSignificance, Other.MinLODSignificance, KINDA_SMALL_NUMBER) &&
+			FMath::IsNearlyEqual(MaxLODSignificance, Other.MaxLODSignificance, KINDA_SMALL_NUMBER) &&
+			bCastShadows == Other.bCastShadows && 
+			Mobility == Other.Mobility &&
+			bRequiresExternalInstanceIDTracking == Other.bRequiresExternalInstanceIDTracking &&
+			LocalTransform.Equals(Other.LocalTransform);
+	}
+	
+	friend inline uint32 GetTypeHash(const FMassSkinnedMeshInstanceVisualizationMeshDesc& MeshDesc)
+	{
+		uint32 Hash = 0x0;
+		Hash = PointerHash(MeshDesc.Asset, Hash);
+		Hash = HashCombine(GetTypeHash(MeshDesc.TransformProvider), Hash);
+		Hash = HashCombine(GetTypeHash(MeshDesc.InstancedSkinnedMeshComponentClass), Hash);
+		Hash = HashCombine(GetTypeHash(MeshDesc.bCastShadows), Hash);
+		Hash = HashCombine(GetTypeHash(MeshDesc.Mobility), Hash);
+		Hash = HashCombine(GetTypeHash(MeshDesc.LocalTransform), Hash);
+		Hash = HashCombine(GetTypeHash(MeshDesc.bRequiresExternalInstanceIDTracking), Hash);
+		for (UMaterialInterface* MaterialOverride : MeshDesc.MaterialOverrides)
+		{
+			if (MaterialOverride)
+			{
+				Hash = PointerHash(MaterialOverride, Hash);
+			}
+		}
+		return Hash;
+	}
+
+	// convenience function for setting MinLODSinificance and MaxLODSinificance based on EMassLOD values
+	void SetSignificanceRange(const EMassLOD::Type MinLOD, const EMassLOD::Type MaxLOD)
+	{
+		checkSlow(MinLOD <= MaxLOD);
+		MinLODSignificance = float(MinLOD);
+		MaxLODSignificance = float(MaxLOD);
+	}
+};
+
+USTRUCT()
+struct FSkinnedMeshInstanceVisualizationDesc : public FTableRowBase
+{
+	GENERATED_BODY()
+
+	/** 
+	 * Mesh descriptions. These will be instanced together using the same transform for each, to 
+	 * visualize the given agent.
+	 */
+	UPROPERTY(EditAnywhere, Category = "Mass|Visual", meta=(TitleProperty="{Mesh}"))
+	TArray<FMassSkinnedMeshInstanceVisualizationMeshDesc> Meshes;
+
+	/** Boolean to enable code to transform the meshes if not align the mass agent transform */
+	UPROPERTY(EditAnywhere, Category = "Mass|Visual")
+	bool bUseTransformOffset = false;
+
+	/** Transform to offset the meshes if not align the mass agent transform */
+	UPROPERTY(EditAnywhere, Category = "Mass|Visual", meta=(EditCondition="bUseTransformOffset"))
+	FTransform TransformOffset = FTransform::Identity;
+
+	/** Custom data that can eventually be propagated to UInstancedSkinnedMeshComponent::PerInstanceCustomData */
+	UPROPERTY(EditAnywhere, Category = "Mass|Visual")
+	TArray<float> CustomDataFloats;
+
+	bool operator==(const FSkinnedMeshInstanceVisualizationDesc& Other) const
+	{
+		return Meshes == Other.Meshes && bUseTransformOffset == Other.bUseTransformOffset && TransformOffset.Equals(Other.TransformOffset);	//Intentionally not comparing CustomDataFloats
+	}
+
+	void Reset()
+	{
+		Meshes.Empty();
+		bUseTransformOffset = false;
+		TransformOffset = FTransform::Identity;
+		CustomDataFloats.Empty();
+	}
+
+	/** @return whether any of descriptions in Meshes is valid. This implies that empty Meshes will be treated as not valid.*/
+	UE_API bool IsValid() const;
+};
+
+/** Handle for FSkinnedMeshInstanceVisualizationDesc's registered with UMassRepresentationSubsystem */
+USTRUCT()
+struct alignas(2) FSkinnedMeshInstanceVisualizationDescHandle 
+{
+	GENERATED_BODY()
+
+	static constexpr uint16 InvalidIndex = TNumericLimits<uint16>::Max();
+
+	FSkinnedMeshInstanceVisualizationDescHandle () = default;
+
+	explicit FSkinnedMeshInstanceVisualizationDescHandle (uint16 InIndex)
+	: Index(InIndex)
+	{}
+
+	explicit FSkinnedMeshInstanceVisualizationDescHandle (int32 InIndex) 
+	{
+		// Handle special case INDEX_NONE = InvalidIndex
+		if (InIndex == INDEX_NONE)
+		{
+			Index = InvalidIndex;
+		}
+		else
+		{
+			checkf(InIndex < static_cast<int32>(InvalidIndex), TEXT("Visualization description index InIndex %d is out of expected bounds (< %u)"), InIndex, InvalidIndex);
+			Index = static_cast<uint16>(InIndex);
+		}
+	}
+
+	inline int32 ToIndex() const
+	{
+		return IsValid() ? Index : INDEX_NONE;
+	}
+
+	bool IsValid() const
+	{
+		return Index != InvalidIndex;
+	}
+
+	bool operator==(const FSkinnedMeshInstanceVisualizationDescHandle & Other) const = default;
+
+private:
+
+	UPROPERTY()
+	uint16 Index = InvalidIndex;
+
+	// @todo: Add a version / serial number to protect against recycled handle reuse. Leaving this out for now to keep size down due to 
+	// prevalent use in FMassRepresentationFragment. Perhaps serial number could be formed from the referenced 
+	// FStaticMeshInstanceVisualizationDesc's hash.
+};
+static_assert(sizeof(FSkinnedMeshInstanceVisualizationDescHandle ) == sizeof(uint16), "FSkinnedMeshInstanceVisualizationDescHandle must be uint16 sized to ensure FMassRepresentationFragment memory isn't unexpectedly bloated");
+
+struct FMassInstancedSkinnedMeshComponentSharedData
+{
+	FMassInstancedSkinnedMeshComponentSharedData()
+	: bRequiresExternalInstanceIDTracking(false)
+	{
+	}
+
+	explicit FMassInstancedSkinnedMeshComponentSharedData(UInstancedSkinnedMeshComponent* InInstancedSkinnedMeshComponent, bool bInRequiresExternalInstanceIDTracking = false, const FTransform& InTransformOffset = FTransform::Identity)
+		: LocalTransform(InTransformOffset), InstancedSkinnedMeshComponent(InInstancedSkinnedMeshComponent), bRequiresExternalInstanceIDTracking(bInRequiresExternalInstanceIDTracking)
+	{
+	}
+
+	FMassInstancedSkinnedMeshComponentSharedData(const FMassInstancedSkinnedMeshComponentSharedData& Other) = default;
+	FMassInstancedSkinnedMeshComponentSharedData& operator=(const FMassInstancedSkinnedMeshComponentSharedData& Other) = default;
+
+
+	void SetInstancedSkinnedMeshComponent(UInstancedSkinnedMeshComponent& InInstancedSkinnedMeshComponent)
+	{
+		check(InstancedSkinnedMeshComponent == nullptr && InstancedSkinnedMeshComponentReferencesCount == 0);
+		InstancedSkinnedMeshComponent = &InInstancedSkinnedMeshComponent;
+	}
+
+
+	UInstancedSkinnedMeshComponent* GetMutableInstancedSkinnedMeshComponent()	{ return InstancedSkinnedMeshComponent; }
+	const UInstancedSkinnedMeshComponent* GetInstancedSkinnedMeshComponent() const	{ return InstancedSkinnedMeshComponent; }
+	int32 OnInstancedSkinnedMeshComponentReferenceStored() { return ++InstancedSkinnedMeshComponentReferencesCount; }
+	int32 OnInstancedSkinnedMeshComponentReferenceReleased() 
+	{ 
+		ensure(InstancedSkinnedMeshComponentReferencesCount >= 0); 
+		return --InstancedSkinnedMeshComponentReferencesCount;
+	}
+
+	void ResetAccumulatedData()
+	{
+		EntitiesRequiringUpdate.Reset();
+		MeshInstanceCustomFloats.Reset();
+		MeshInstanceTransforms.Reset();
+		MeshInstancePrevTransforms.Reset();
+		MeshInstanceAnimationData.Reset();
+		EntitiesRequiringRemoval.Reset();
+		WriteIterator = 0;
+	}
+
+	void RemoveUpdatedInstanceIdsAtSwap(const int32 InstanceIDIndex)
+	{
+		EntitiesRequiringUpdate.RemoveAtSwap(InstanceIDIndex, EAllowShrinking::No);
+		MeshInstanceTransforms.RemoveAtSwap(InstanceIDIndex, EAllowShrinking::No);
+		MeshInstancePrevTransforms.RemoveAtSwap(InstanceIDIndex, EAllowShrinking::No);
+		MeshInstanceAnimationData.RemoveAtSwap(InstanceIDIndex, EAllowShrinking::No);
+		if (MeshInstanceCustomFloats.Num())
+		{
+			MeshInstanceCustomFloats.RemoveAtSwap(InstanceIDIndex, EAllowShrinking::No);
+		}
+	}
+
+	bool HasUpdatesToApply() const { return EntitiesRequiringUpdate.Num() || EntitiesRequiringRemoval.Num(); }
+	TConstArrayView<FMassEntityHandle> GetEntitiesRequiringUpdate() const { return EntitiesRequiringUpdate; }
+	TConstArrayView<FTransform> GetMeshInstanceTransforms() const { return MeshInstanceTransforms; }
+	TConstArrayView<FAnimSequenceTrackAutoPlayData> GetMeshInstanceAnimationData() const { return MeshInstanceAnimationData; }
+	/**
+	 * this function is a flavor we need to interact with older engine API that's using TArray references.
+	 * Use GetMeshInstanceTransforms instead whenever possible.
+	 */
+	const TArray<FTransform>& GetMeshInstanceTransformsArray() const { return MeshInstanceTransforms; }
+	const TArray<FAnimSequenceTrackAutoPlayData>& GetMeshInstanceAnimatioDataArray() const { return MeshInstanceAnimationData; }
+
+	TConstArrayView<FTransform> GetMeshInstancePrevTransforms() const { return MeshInstancePrevTransforms; }
+	TConstArrayView<FMassEntityHandle> GetEntitiesRequiringRemoval() const { return EntitiesRequiringRemoval; }
+	TConstArrayView<float> GetMeshInstanceCustomFloats() const { return MeshInstanceCustomFloats; }
+
+	bool RequiresExternalInstanceIDTracking() const { return bRequiresExternalInstanceIDTracking; }
+
+	void Reset()
+	{
+		*this = FMassInstancedSkinnedMeshComponentSharedData();
+	}
+
+	using FEntityToPrimitiveIdMap = Experimental::TRobinHoodHashMap<FMassEntityHandle, FPrimitiveInstanceId>;
+	using FEntityToTrackIdMap = Experimental::TRobinHoodHashMap<FMassEntityHandle, int32>;
+
+	FEntityToPrimitiveIdMap& GetMutableEntityPrimitiveToIdMap() { return EntityHandleToPrimitiveIdMap; }
+	const FEntityToPrimitiveIdMap& GetEntityPrimitiveToIdMap() const { return EntityHandleToPrimitiveIdMap; }
+
+	FEntityToTrackIdMap& GetMutableEntityToTrackMap() { return EntityHandleToTrackIdMap; }
+	const FEntityToTrackIdMap& GetEntityToTrackMap() const { return EntityHandleToTrackIdMap; }
+
+	FEntityToPrimitiveIdMap& GetMutableEntityToPrimitiveIdMap() { return EntityHandleToPrimitiveIdMap; }
+	const FEntityToPrimitiveIdMap& GetEntityToPrimitiveIdMap() const { return EntityHandleToPrimitiveIdMap; }
+
+	int16 GetComponentInstanceIdTouchCounter() const { return ComponentInstanceIdTouchCounter; }
+
+protected:
+	friend FMassLODInstancedSkinnedMeshSignificanceRange;
+	friend UMassVisualizationComponent;
+
+	FTransform LocalTransform = FTransform::Identity;
+
+	/** Buffer holding current frame transforms for the mesh instances, used to batch update the transforms */
+	TArray<FMassEntityHandle> EntitiesRequiringUpdate;
+	TArray<FTransform> MeshInstanceTransforms;
+	TArray<FTransform> MeshInstancePrevTransforms;
+	/** Buffer holding animation data */
+	TArray<FAnimSequenceTrackAutoPlayData> MeshInstanceAnimationData;
+	TArray<FMassEntityHandle> EntitiesRequiringRemoval;
+
+	/** Buffer holding current frame custom floats for the mesh instances, used to batch update the ISMs custom data */
+	TArray<float> MeshInstanceCustomFloats;
+
+	// When initially adding to MeshInstanceCustomFloats, can use the size as the write iterator, but on subsequent processors, we need to know where to start writing
+	int32 WriteIterator = 0;
+
+	UInstancedSkinnedMeshComponent* InstancedSkinnedMeshComponent = nullptr;
+
+	int32 InstancedSkinnedMeshComponentReferencesCount = 0;
+
+	/**
+	 * When set to true will result in MassVisualizationComponent manually perform Instance ID-related operations
+	 * instead of relying on ISMComponent's internal ID operations.
+	 * @note this mechanism has been added in preparation of changes to ISM component to change access to its internal
+	 *	instance ID logic. WIP as of Jun 17th 2023
+	 */
+	uint8 bRequiresExternalInstanceIDTracking : 1;
+
+private:
+	/** Indicates that mutating changes, that can affect MassInstanceIdToComponentInstanceIdMap, have been performed.
+	 *	Can be used to validate whether cached data stored in other placed needs to be re-cached. */
+	uint16 ComponentInstanceIdTouchCounter = 0;
+
+protected:
+	FEntityToPrimitiveIdMap EntityHandleToPrimitiveIdMap;
+	FEntityToTrackIdMap EntityHandleToTrackIdMap;
+};
+
+
+/**
+ * The container type hosting FMassInstancedSkinnedMeshSharedData instances and supplying functionality of marking entries that require
+ * instance-related operations (adding, removing).
+ *
+ * To get a FMassInstance instance to add operations to it call GetAndMarkDirty.
+ *
+ * Use FDirtyIterator to iterate over just the data that needs processing.
+ *
+ * @see UMassVisualizationComponent::EndVisualChanges for iteration
+ * @see FMassLODSignificanceRange methods for performing dirtying operations
+ */
+struct FMassInstancedSkinnedMeshComponentSharedDataMap
+{
+	struct FDirtyIterator
+	{
+		friend FMassInstancedSkinnedMeshComponentSharedDataMap;
+		explicit FDirtyIterator(FMassInstancedSkinnedMeshComponentSharedDataMap& InContainer)
+			: Container(InContainer), It(InContainer.GetDirtyArray())
+		{
+			if (It && It.GetValue() != bValueToCheck)
+			{
+				// will result in either setting IT to the first bInValue, or making bool(It) == false
+				++(*this);
+			}
+		}
+	public:
+		operator bool() const { return bool(It); }
+
+		FDirtyIterator& operator++()
+		{
+			while (++It)
+			{
+				if (It.GetValue() == bValueToCheck)
+				{
+					break;
+				}
+			}
+			return *this;
+		}
+
+		FMassInstancedSkinnedMeshComponentSharedData& operator*() const
+		{
+			return Container.GetAtIndex(It.GetIndex());
+		}
+
+		void ClearDirtyFlag()
+		{
+			It.GetValue() = false;
+		}
+
+	private:
+		FMassInstancedSkinnedMeshComponentSharedDataMap& Container;
+		TBitArray<>::FIterator It;
+		static constexpr bool bValueToCheck = true;
+	};
+
+	FMassInstancedSkinnedMeshComponentSharedData& GetAndMarkDirtyChecked(const FInstancedSkinnedMeshComponentSharedDataKey OwnerKey)
+	{
+		const int32 DataIndex = Map[OwnerKey];
+		DirtyData[DataIndex] = true;
+		return Data[DataIndex];
+	}
+
+	FMassInstancedSkinnedMeshComponentSharedData* GetAndMarkDirty(const FInstancedSkinnedMeshComponentSharedDataKey OwnerKey)
+	{
+		const int32* DataIndex = Map.Find(OwnerKey);
+		if (ensureMsgf(DataIndex, TEXT("%hs Failed to find OwnerKey %s"), __FUNCTION__, *GetNameSafe(OwnerKey.ResolveObjectPtrEvenIfGarbage())))
+		{
+			DirtyData[*DataIndex] = true;
+			return &Data[*DataIndex];
+		}
+		return nullptr;
+	}
+
+	template<typename... TArgs>
+	FMassInstancedSkinnedMeshComponentSharedData& FindOrAdd(const FInstancedSkinnedMeshComponentSharedDataKey OwnerKey, TArgs&&... InNewInstanceArgs)
+	{
+		const int32* DataIndex = Map.Find(OwnerKey);
+		if (DataIndex == nullptr)
+		{
+			return Add(OwnerKey, Forward<TArgs>(InNewInstanceArgs)...);
+		}
+		check(Data.IsValidIndex(*DataIndex));
+		return Data[*DataIndex];
+	}
+
+	FMassInstancedSkinnedMeshComponentSharedData* Find(const FInstancedSkinnedMeshComponentSharedDataKey OwnerKey)
+	{
+		int32* DataIndex = Map.Find(OwnerKey);
+		return (DataIndex == nullptr || *DataIndex == INDEX_NONE) ? (FMassInstancedSkinnedMeshComponentSharedData*)nullptr : &Data[*DataIndex];
+	}
+
+	template<typename... TArgs>
+	FMassInstancedSkinnedMeshComponentSharedData& Add(const FInstancedSkinnedMeshComponentSharedDataKey OwnerKey, TArgs&&... InNewInstanceArgs)
+	{
+		const int32 DataIndex = FreeIndices.Num() ? FreeIndices.Pop() : Data.Num();
+		Map.Add(OwnerKey, DataIndex);
+
+		if (DataIndex == Data.Num())
+		{
+			DirtyData.Add(false, DataIndex - DirtyData.Num() + 1);
+			DirtyData[DataIndex] = true;
+			return Data.Add_GetRef(FMassInstancedSkinnedMeshComponentSharedData(Forward<TArgs>(InNewInstanceArgs)...));
+		}
+		else
+		{
+			DirtyData[DataIndex] = true;
+			Data[DataIndex] = FMassInstancedSkinnedMeshComponentSharedData(Forward<TArgs>(InNewInstanceArgs)...);
+			return Data[DataIndex];
+		}
+	}
+
+	void Remove(const FInstancedSkinnedMeshComponentSharedDataKey OwnerKey)
+	{
+		int32 DataIndex = INDEX_NONE;
+		if (ensure(Map.RemoveAndCopyValue(OwnerKey, DataIndex)))
+		{
+			DirtyData[DataIndex] = false;
+			Data[DataIndex].Reset();
+			FreeIndices.Add(DataIndex);
+		}
+	}
+
+	FMassInstancedSkinnedMeshComponentSharedData& GetAtIndex(const int32 DataIndex)
+	{
+		return Data[DataIndex];
+	}
+
+	TBitArray<>& GetDirtyArray()
+	{
+		return DirtyData;
+	}
+
+	/** @return total number of entries in Data array. Note that some or all entries could be empty (i.e. already freed) */
+	int32 Num() const
+	{
+		return Data.Num();
+	}
+
+	/** @return number of non-empty entries in Data. */
+	int32 NumValid() const
+	{
+		return Data.Num() - FreeIndices.Num();
+	}
+
+	bool IsDirty(const int32 DataIndex) const
+	{
+		return DirtyData[DataIndex];
+	}
+
+	bool IsEmpty() const
+	{
+		return NumValid() == 0;
+	}
+
+	void Reset()
+	{
+		*this = FMassInstancedSkinnedMeshComponentSharedDataMap();
+	}
+
+	const FMassInstancedSkinnedMeshComponentSharedData* GetDataForIndex(const int32 Index) const
+	{
+		return Data.IsValidIndex(Index) ? &Data[Index] : nullptr;
+	}
+
+	const FMassInstancedSkinnedMeshComponentSharedData* GetDataForKey(const FInstancedSkinnedMeshComponentSharedDataKey Key) const
+	{
+		const int32* Index = Map.Find(Key);
+		return (Index && Data.IsValidIndex(*Index))
+			? &Data[*Index]
+			: nullptr;
+	}
+
+protected:
+	TArray<FMassInstancedSkinnedMeshComponentSharedData> Data;
+	/** Mapping from Owner (as FObjectKey) of data represented by FMassISMCSharedData to an index to Data */
+	TMap<FInstancedSkinnedMeshComponentSharedDataKey, int32> Map;
+	/** Indicates whether corresponding Data entry has any instance work assigned to it (instance addition or removal) */
+	TBitArray<> DirtyData;
+	/** Indices to Data that are available for reuse */
+	TArray<int32> FreeIndices;
+};
+
+
+
+USTRUCT()
+struct FMassLODInstancedSkinnedMeshSignificanceRange
+{
+	GENERATED_BODY()
+public:
+
+	UE_API void AddBatchedTransform(const FMassEntityHandle EntityHandle, const FTransform& Transform, const FTransform& PrevTransform, TConstArrayView<FInstancedSkinnedMeshComponentSharedDataKey> ExcludeMeshRefs);
+	UE_API void AddBatchedAnimationData(const FMassEntityHandle EntityHandle, FAnimSequenceTrackAutoPlayData AnimationData, TConstArrayView<FInstancedSkinnedMeshComponentSharedDataKey> ExcludeMeshRefs);
+
+	// Adds the specified struct reinterpreted as custom floats to our custom data. Individual members of the specified struct should always fit into a float.
+	// When adding any custom data, the custom data must be added for every instance.
+	template<typename InCustomDataType>
+	void AddBatchedCustomData(InCustomDataType InCustomData, const TArray<FInstancedSkinnedMeshComponentSharedDataKey>& ExcludeMeshRefs, int32 NumFloatsToPad = 0)
+	{
+		check(InstancedSkinnedMeshSharedDataPtr);
+		static_assert((sizeof(InCustomDataType) % sizeof(float)) == 0, "AddBatchedCustomData: InCustomDataType should have a total size multiple of sizeof(float), and have members that fit in a float's boundaries");
+		const size_t StructSize = sizeof(InCustomDataType);
+		const size_t StructSizeInFloats = StructSize / sizeof(float);
+		for (int i = 0; i < SkinnedMeshComponentRefs.Num(); i++)
+		{
+			if (ExcludeMeshRefs.Contains(SkinnedMeshComponentRefs[i]))
+			{
+				continue;
+			}
+
+			FMassInstancedSkinnedMeshComponentSharedData& SharedData = (*InstancedSkinnedMeshSharedDataPtr).GetAndMarkDirtyChecked(SkinnedMeshComponentRefs[i]);
+			const int32 StartIndex = SharedData.MeshInstanceCustomFloats.AddDefaulted(StructSizeInFloats + NumFloatsToPad);
+			InCustomDataType* CustomData = reinterpret_cast<InCustomDataType*>(&SharedData.MeshInstanceCustomFloats[StartIndex]);
+			*CustomData = InCustomData;
+		}
+	}
+
+	UE_API void AddBatchedCustomDataFloats(const TArray<float>& CustomFloats, const TArray<FInstancedSkinnedMeshComponentSharedDataKey>& ExcludeMeshRefs);
+
+	/** Single-instance version of AddBatchedCustomData when called to add entities (as opposed to modify existing ones).*/
+	UE_API void AddInstance(const FMassEntityHandle EntityHandle, const FTransform& Transform);
+
+	UE_API void RemoveInstance(const FMassEntityHandle EntityHandle);
+
+	UE_API void WriteCustomDataFloatsAtStartIndex(int32 MeshIndex, const TArrayView<float>& CustomFloats, const int32 FloatsPerInstance, const int32 StartIndex, const TArray<FInstancedSkinnedMeshComponentSharedDataKey>& ExcludeMeshRefs);
+
+	/** LOD Significance range */
+	float MinSignificance = 0.0f;
+	float MaxSignificance = 0.0f;
+
+	/** The component handling these instances */
+	TArray<FInstancedSkinnedMeshComponentSharedDataKey> SkinnedMeshComponentRefs;
+
+	FMassInstancedSkinnedMeshComponentSharedDataMap* InstancedSkinnedMeshSharedDataPtr = nullptr;
+};
+
+
+USTRUCT()
+struct FMassInstancedSkinnedMeshInfo
+{
+	GENERATED_BODY()
+public:
+
+	FMassInstancedSkinnedMeshInfo() = default;
+
+	explicit FMassInstancedSkinnedMeshInfo(const FSkinnedMeshInstanceVisualizationDesc& InDesc)
+		: Desc(InDesc)
+	{
+	}
+
+	/** Clears out contents so that a given FMassInstancedSkinnedMeshInfo instance can be reused */
+	UE_API void Reset();
+
+	const FSkinnedMeshInstanceVisualizationDesc& GetDesc() const
+	{
+		return Desc;
+	}
+
+	/** Whether or not to transform the  meshes if not align the mass agent transform */
+	bool ShouldUseTransformOffset() const { return Desc.bUseTransformOffset; }
+	FTransform GetTransformOffset() const { return Desc.TransformOffset; }
+
+	inline FMassLODInstancedSkinnedMeshSignificanceRange* GetLODSignificanceRange(float LODSignificance)
+	{
+		for (FMassLODInstancedSkinnedMeshSignificanceRange& Range : LODSignificanceRanges)
+		{
+			if (LODSignificance >= Range.MinSignificance && LODSignificance < Range.MaxSignificance)
+			{
+				return &Range;
+			}
+		}
+		return nullptr;
+	}
+
+	void AddBatchedTransform(const FMassEntityHandle EntityHandle, const FTransform& Transform, const FTransform& PrevTransform, const float LODSignificance, const float PrevLODSignificance = -1.0f)
+	{
+		if (FMassLODInstancedSkinnedMeshSignificanceRange* Range = GetLODSignificanceRange(LODSignificance))
+		{
+			Range->AddBatchedTransform(EntityHandle, Transform, PrevTransform, {});
+			if (PrevLODSignificance >= 0.0f)
+			{
+				FMassLODInstancedSkinnedMeshSignificanceRange* PrevRange = GetLODSignificanceRange(PrevLODSignificance);
+				if (ensureMsgf(PrevRange, TEXT("Couldn't find a valid LODSignificanceRange for PrevLODSignificance %f"), PrevLODSignificance)
+					&& PrevRange != Range)
+				{
+					PrevRange->AddBatchedTransform(EntityHandle, Transform, PrevTransform, Range->SkinnedMeshComponentRefs);
+				}
+			}
+		}
+	}
+
+	void AddBatchedAnimationData(const FMassEntityHandle EntityHandle, FAnimSequenceTrackAutoPlayData AnimationData, const float LODSignificance, const float PrevLODSignificance = -1.0f)
+	{
+		if (FMassLODInstancedSkinnedMeshSignificanceRange* Range = GetLODSignificanceRange(LODSignificance))
+		{
+			Range->AddBatchedAnimationData(EntityHandle, AnimationData, {});
+			if (PrevLODSignificance >= 0.0f)
+			{
+				FMassLODInstancedSkinnedMeshSignificanceRange* PrevRange = GetLODSignificanceRange(PrevLODSignificance);
+				if (ensureMsgf(PrevRange, TEXT("Couldn't find a valid LODSignificanceRange for PrevLODSignificance %f"), PrevLODSignificance)
+					&& PrevRange != Range)
+				{
+					PrevRange->AddBatchedAnimationData(EntityHandle, AnimationData, Range->SkinnedMeshComponentRefs);
+				}
+			}
+		}
+	}
+
+
+	inline void RemoveInstance(const FMassEntityHandle EntityHandle, const float LODSignificance)
+	{
+		if (FMassLODInstancedSkinnedMeshSignificanceRange* Range = GetLODSignificanceRange(LODSignificance))
+		{
+			Range->RemoveInstance(EntityHandle);
+		}
+	}
+
+	// Adds the specified struct reinterpreted as custom floats to our custom data. Individual members of the specified struct should always fit into a float.
+	// When adding any custom data, the custom data must be added for every instance.
+	template<typename InCustomDataType>
+	void AddBatchedCustomData(InCustomDataType InCustomData, const float LODSignificance, const float PrevLODSignificance = -1.0f, int32 NumFloatsToPad = 0)
+	{
+		if (FMassLODInstancedSkinnedMeshSignificanceRange* Range = GetLODSignificanceRange(LODSignificance))
+		{
+			Range->AddBatchedCustomData(InCustomData, {}, NumFloatsToPad);
+			if (PrevLODSignificance >= 0.0f)
+			{
+				FMassLODInstancedSkinnedMeshSignificanceRange* PrevRange = GetLODSignificanceRange(PrevLODSignificance);
+				if (ensureMsgf(PrevRange, TEXT("Couldn't find a valid LODSignificanceRange for PrevLODSignificance %f"), PrevLODSignificance)
+					&& PrevRange != Range)
+				{
+					PrevRange->AddBatchedCustomData(InCustomData, Range->SkinnedMeshComponentRefs, NumFloatsToPad);
+				}
+			}
+		}
+	}
+
+	inline void AddBatchedCustomDataFloats(const TArray<float>& CustomFloats, const float LODSignificance, const float PrevLODSignificance = -1.0f)
+	{
+		if (FMassLODInstancedSkinnedMeshSignificanceRange* Range = GetLODSignificanceRange(LODSignificance))
+		{
+			Range->AddBatchedCustomDataFloats(CustomFloats, {});
+			if (PrevLODSignificance >= 0.0f)
+			{
+				FMassLODInstancedSkinnedMeshSignificanceRange* PrevRange = GetLODSignificanceRange(PrevLODSignificance);
+				if (ensureMsgf(PrevRange, TEXT("Couldn't find a valid LODSignificanceRange for PrevLODSignificance %f"), PrevLODSignificance)
+					&& PrevRange != Range)
+				{
+					PrevRange->AddBatchedCustomDataFloats(CustomFloats, Range->SkinnedMeshComponentRefs);
+				}
+			}
+		}
+	}
+
+	void WriteCustomDataFloatsAtStartIndex(int32 MeshIndex, const TArrayView<float>& CustomFloats, const float LODSignificance, const int32 FloatsPerInstance, const int32 FloatStartIndex, const float PrevLODSignificance = -1.0f)
+	{
+		if (FMassLODInstancedSkinnedMeshSignificanceRange* Range = GetLODSignificanceRange(LODSignificance))
+		{
+			Range->WriteCustomDataFloatsAtStartIndex(MeshIndex, CustomFloats, FloatsPerInstance, FloatStartIndex, {});
+			if (PrevLODSignificance >= 0.0f)
+			{
+				FMassLODInstancedSkinnedMeshSignificanceRange* PrevRange = GetLODSignificanceRange(PrevLODSignificance);
+				if (ensureMsgf(PrevRange, TEXT("Couldn't find a valid LODSignificanceRange for PrevLODSignificance %f"), PrevLODSignificance)
+					&& PrevRange != Range)
+				{
+					PrevRange->WriteCustomDataFloatsAtStartIndex(MeshIndex, CustomFloats, FloatsPerInstance, FloatStartIndex, Range->SkinnedMeshComponentRefs);
+				}
+			}
+		}
+	}
+
+	void AddInstancedSkinnedMeshComponent(FMassInstancedSkinnedMeshComponentSharedData& SharedData)
+	{
+		if (ensure(SharedData.GetInstancedSkinnedMeshComponent()))
+		{
+			InstancedSkinnedMeshComponents.Add(SharedData.GetMutableInstancedSkinnedMeshComponent());
+			SharedData.OnInstancedSkinnedMeshComponentReferenceStored();
+		}
+	}
+
+	int32 GetLODSignificanceRangesNum() const { return LODSignificanceRanges.Num(); }
+
+	bool IsValid() const
+	{
+		return Desc.Meshes.Num() && InstancedSkinnedMeshComponents.Num() && LODSignificanceRanges.Num();
+	}
+
+protected:
+
+	/** Destroy the visual instance */
+	UE_API void ClearVisualInstance(UInstancedSkinnedMeshComponent& ISMComponent);
+
+	/** Information about this  mesh which will represent all instances */
+	UPROPERTY(VisibleAnywhere, Category = "Mass/Debug")
+	FSkinnedMeshInstanceVisualizationDesc Desc;
+
+	/** The components handling these instances */
+	UPROPERTY(VisibleAnywhere, Category = "Mass/Debug")
+	TArray<TObjectPtr<UInstancedSkinnedMeshComponent>> InstancedSkinnedMeshComponents;
+
+	UPROPERTY(VisibleAnywhere, Category = "Mass/Debug")
+	TArray<FMassLODInstancedSkinnedMeshSignificanceRange> LODSignificanceRanges;
+
+	friend class UMassVisualizationComponent;
+};
+
+
+#if ENABLE_MT_DETECTOR
+
+#define MAKE_MASS_INSTANCED_SKINNED_MESH_INFO_ARRAY_VIEW(ArrayView, AccessDetector) FMassInstancedSkinnedMeshInfoArrayView(ArrayView, AccessDetector)
+struct FMassInstancedSkinnedMeshInfoArrayViewAccessDetector
+{
+	FMassInstancedSkinnedMeshInfoArrayViewAccessDetector(TArrayView<FMassInstancedSkinnedMeshInfo> InInstancedSkinnedMeshInfos, const FRWAccessDetector& InAccessDetector)
+		: InstancedSkinnedMeshInfos(InInstancedSkinnedMeshInfos)
+		, AccessDetector(&InAccessDetector)
+	{
+		UE_MT_ACQUIRE_WRITE_ACCESS(*AccessDetector);
+	}
+
+	FMassInstancedSkinnedMeshInfoArrayViewAccessDetector(FMassInstancedSkinnedMeshInfoArrayViewAccessDetector&& Other)
+		: InstancedSkinnedMeshInfos(Other.InstancedSkinnedMeshInfos)
+		, AccessDetector(Other.AccessDetector)
+	{
+		Other.AccessDetector = nullptr;
+	}
+	FMassInstancedSkinnedMeshInfoArrayViewAccessDetector(const FMassInstancedSkinnedMeshInfoArrayViewAccessDetector& Other) = delete;
+	void operator=(const FMassInstancedSkinnedMeshInfoArrayViewAccessDetector& Other) = delete;
+
+	~FMassInstancedSkinnedMeshInfoArrayViewAccessDetector()
+	{
+		if (AccessDetector)
+		{
+			UE_MT_RELEASE_WRITE_ACCESS(*AccessDetector);
+		}
+	}
+
+	inline FMassInstancedSkinnedMeshInfo& operator[](int32 Index) const
+	{
+		return InstancedSkinnedMeshInfos[Index];
+	}
+
+	bool IsValidIndex(const int32 Index) const
+	{
+		return InstancedSkinnedMeshInfos.IsValidIndex(Index);
+	}
+
+	int32 Num() const
+	{
+		return InstancedSkinnedMeshInfos.Num();
+	}
+
+private:
+	TArrayView<FMassInstancedSkinnedMeshInfo> InstancedSkinnedMeshInfos;
+	const FRWAccessDetector* AccessDetector;
+};
+
+using FMassInstancedSkinnedMeshInfoArrayView = FMassInstancedSkinnedMeshInfoArrayViewAccessDetector;
+
+
+#else // ENABLE_MT_DETECTOR
+
+#define MAKE_MASS_INSTANCED_SKINNED_MESH_INFO_ARRAY_VIEW(ArrayView, AccessDetector) ArrayView
+using FMassInstancedSkinnedMeshInfoArrayView = TArrayView<FMassInstancedSkinnedMeshInfo>;
+
+#endif // ENABLE_MT_DETECTOR
+
+#undef UE_API
\ No newline at end of file
diff --git a/Engine/Plugins/Runtime/MassGameplay/Source/MassRepresentation/Public/MassStationaryISMSwitcherProcessor.h b/Engine/Plugins/Runtime/MassGameplay/Source/MassRepresentation/Public/MassStationaryISMSwitcherProcessor.h
--- a/Engine/Plugins/Runtime/MassGameplay/Source/MassRepresentation/Public/MassStationaryISMSwitcherProcessor.h
+++ b/Engine/Plugins/Runtime/MassGameplay/Source/MassRepresentation/Public/MassStationaryISMSwitcherProcessor.h
@@ -29,8 +29,8 @@ struct FMassStationaryISMSwitcherProcessorTag : public FMassTag
 
 /** 
  * This processor's sole responsibility is to process all entities tagged with FMassStaticRepresentationTag
- * and check if they've switched to or away from EMassRepresentationType::StaticMeshInstance; and acordingly add or remove 
- * the entity from the appropriate FMassInstancedStaticMeshInfoArrayView.
+ * and check if they've switched to or away from EMassRepresentationType::StaticMeshInstance or EMassRepresentationType::SkinnedMeshInstance ; and acordingly add or remove 
+ * the entity from the appropriate ArrayView.
  */
 UCLASS(MinimalAPI)
 class UMassStationaryISMSwitcherProcessor : public UMassProcessor
diff --git a/Engine/Plugins/Runtime/MassGameplay/Source/MassRepresentation/Public/MassUpdateSkinnedMeshProcessor.h b/Engine/Plugins/Runtime/MassGameplay/Source/MassRepresentation/Public/MassUpdateSkinnedMeshProcessor.h
new file mode 100644
--- /dev/null
+++ b/Engine/Plugins/Runtime/MassGameplay/Source/MassRepresentation/Public/MassUpdateSkinnedMeshProcessor.h
@@ -0,0 +1,38 @@
+// Copyright Epic Games, Inc. All Rights Reserved.
+
+#pragma once
+
+#include "MassEntityQuery.h"
+#include "MassProcessor.h"
+#include "MassUpdateSkinnedMeshProcessor.generated.h"
+
+#define UE_API MASSREPRESENTATION_API
+
+class UMassRepresentationSubsystem;
+struct FMassInstancedSkinnedMeshInfo;
+
+UCLASS(MinimalAPI)
+class UMassUpdateInstancedSkinnedMeshProcessor : public UMassProcessor
+{
+	GENERATED_BODY()
+
+public:
+	UE_API UMassUpdateInstancedSkinnedMeshProcessor();
+
+	static UE_API void UpdateMeshTransform(FMassEntityHandle EntityHandle, FMassInstancedSkinnedMeshInfo& MeshInfo, const FTransform& Transform, const FTransform& PrevTransform, const float LODSignificance, const float PrevLODSignificance = -1.0f);
+
+protected:
+
+	/** Configure the owned FMassEntityQuery instances to express processor's requirements */
+	UE_API virtual void ConfigureQueries(const TSharedRef<FMassEntityManager>& EntityManager) override;
+
+	/**
+	 * Execution method for this processor
+	 * @param EntitySubsystem is the system to execute the lambdas on each entity chunk
+	 * @param Context is the execution context to be passed when executing the lambdas */
+	UE_API virtual void Execute(FMassEntityManager& EntityManager, FMassExecutionContext& Context) override;
+
+	FMassEntityQuery EntityQuery;
+};
+
+#undef UE_API
\ No newline at end of file
diff --git a/Engine/Plugins/Runtime/MassGameplay/Source/MassRepresentation/Public/MassVisualizationComponent.h b/Engine/Plugins/Runtime/MassGameplay/Source/MassRepresentation/Public/MassVisualizationComponent.h
--- a/Engine/Plugins/Runtime/MassGameplay/Source/MassRepresentation/Public/MassVisualizationComponent.h
+++ b/Engine/Plugins/Runtime/MassGameplay/Source/MassRepresentation/Public/MassVisualizationComponent.h
@@ -3,6 +3,7 @@
 #pragma once
 
 #include "MassRepresentationTypes.h"
+#include "MassSkinnedMeshRepresentationTypes.h"
 #include "Components/ActorComponent.h"
 #include "Misc/MTAccessDetector.h"
 #if UE_ENABLE_INCLUDE_ORDER_DEPRECATED_IN_5_6
@@ -10,7 +11,6 @@
 #endif // UE_ENABLE_INCLUDE_ORDER_DEPRECATED_IN_5_6
 #include "MassVisualizationComponent.generated.h"
 
-
 class UInstancedStaticMeshComponent;
 
 /** 
@@ -129,4 +129,108 @@ protected:
 
 	/** Indices to InstancedStaticMeshInfos that need their SMComponent constructed */
 	TArray<FStaticMeshInstanceVisualizationDescHandle> InstancedSMComponentsRequiringConstructing;
-};
+
+public:
+	/**
+	 * Get the index of the visual type, will add a new one if does not exist
+	 * @param Desc is the information for the visual that will be instantiated later via AddVisualInstance()
+	 * @return The index of the visual type
+	 */
+	MASSREPRESENTATION_API FSkinnedMeshInstanceVisualizationDescHandle FindOrAddSkinnedMeshInstanceVisualDesc(const FSkinnedMeshInstanceVisualizationDesc& Desc);
+
+	/**
+	 * Creates a dedicated visual type described by host Desc and ties InstancedSkinnedMeshComponent to it.
+	 * @note this is a helper function for a common "single ISMComponent" case. Calls AddVisualDescWithComponents under the hood.
+	 * @return The index of the visual type
+	 */
+	MASSREPRESENTATION_API FSkinnedMeshInstanceVisualizationDescHandle AddSkinnedMeshInstanceVisualDescWithComponent(const FSkinnedMeshInstanceVisualizationDesc& Desc, UInstancedSkinnedMeshComponent& InstancedSkinnedMeshComponent);
+
+	/**
+	 * Creates a dedicated visual type described by host Desc and ties given InstancedSkinnedMeshComponents to it.
+	 * @return The index of the visual type
+	 */
+	MASSREPRESENTATION_API FSkinnedMeshInstanceVisualizationDescHandle AddSkinnedMeshInstanceVisualDescWithComponents(const FSkinnedMeshInstanceVisualizationDesc& Desc, TArrayView<TObjectPtr<UInstancedSkinnedMeshComponent>> InstancedSkinnedMeshComponents);
+	
+	/**
+	 * Fetches FMassInstancedSkinnedMeshComponentSharedData indicated by DescriptionIndex, or nullptr if it's not a valid index
+	 */
+	MASSREPRESENTATION_API const FMassInstancedSkinnedMeshComponentSharedData* GetSharedDataForDescriptionIndex(const int32 DescriptionIndex) const;
+
+	/**
+	 * Fetches FMassInstancedSkinnedMeshComponentSharedData indicated by an InstancedSkinnedMeshComponent, or nullptr if the InstancedSkinnedMeshComponent is not represented by any shared data.
+	 */
+	MASSREPRESENTATION_API const FMassInstancedSkinnedMeshComponentSharedData* GetSharedDataForInstancedSkinnedMeshComponent(const UInstancedSkinnedMeshComponent* InstancedSkinnedMeshComponent) const;
+
+	/**
+	 * Removes all data associated with a given VisualizationIndex. Note that this is safe to do only if there are no
+	 * entities relying on this index. No entity data patching will take place.
+	 */
+	MASSREPRESENTATION_API void RemoveSkinnedMeshInstanceVisualDesc(const FSkinnedMeshInstanceVisualizationDescHandle VisualizationHandle);
+
+	/** Get the array of all visual instance informations */
+	FMassInstancedSkinnedMeshInfoArrayView GetMutableInstanceVisualInfos()
+	{
+		FMassInstancedSkinnedMeshInfoArrayView View = MAKE_MASS_INSTANCED_SKINNED_MESH_INFO_ARRAY_VIEW(MakeArrayView(InstancedSkinnedMeshInfos), InstancedSkinnedMeshInfosDetector);
+		return MoveTemp(View);
+	}
+
+	const FMassSkinnedMeshResolvedAnimState* FindResolvedAnimState(const FMassEntityHandle& EntityHandle) const
+	{
+		return ResolvedAnimStates.Find(EntityHandle);
+	}
+
+protected:
+	/**
+	 * Process all removed IDs in FMassInstancedSkinnedMeshComponentSharedData and apply to the Instanced Skinned Mesh component.
+	 */
+	MASSREPRESENTATION_API void ProcessRemovesForComponent(UInstancedSkinnedMeshComponent& InstancedSkinnedMeshComponent, FMassInstancedSkinnedMeshComponentSharedData& SharedData);
+
+	/**
+	 * Applies changes accumulated in SharedData while manually updating the Instance ID mapping. This approach is done in preparation
+	 * to upcoming Instanced Mesh changes to keep the mapping management more secure (by making mapping private and fully component-owned).
+	 */
+	MASSREPRESENTATION_API void HandleChangesWithExternalIDTrackingForComponent(UInstancedSkinnedMeshComponent& InstancedSkinnedMeshComponent, FMassInstancedSkinnedMeshComponentSharedData& SharedData);
+
+	/** Recreate all the instanced skinned mesh components from the InstancedSkinnedMeshInfos */
+	MASSREPRESENTATION_API void ConstructSkinnedMeshComponents();
+
+	/**
+	 * Creates LODSignificance ranges for all the meshes indicated by Info
+	 * @param MeshRefKeys if not empty will be used when adding individual FMassSkinnedMeshInstanceVisualizationMeshDesc
+	 *	instances to LOD significance ranges.
+	 */
+	MASSREPRESENTATION_API void BuildLODSignificanceForInstancedSkinnedMeshInfo(FMassInstancedSkinnedMeshInfo& Info, TConstArrayView<UInstancedSkinnedMeshComponent*> MeshRefKeys);
+
+	/** Either adds an element to InstancedSkinnedMeshInfos or reuses an existing entry based on InstancedSkinnedMeshInfosFreeIndices*/
+	MASSREPRESENTATION_API FSkinnedMeshInstanceVisualizationDescHandle AddInstancedSkinnedMeshInfo(const FSkinnedMeshInstanceVisualizationDesc& Desc);
+
+	/** The information of all the instanced skinned meshes. Make sure to use AddInstancedSkinnedMeshInfo to add elements to it */
+	UPROPERTY(Transient)
+	TArray<FMassInstancedSkinnedMeshInfo> InstancedSkinnedMeshInfos;
+	UE_MT_DECLARE_RW_ACCESS_DETECTOR(InstancedSkinnedMeshInfosDetector);
+
+	/** Indices to InstancedSkinnedMeshInfos that have been released and can be reused */
+	TArray<FSkinnedMeshInstanceVisualizationDescHandle> InstancedSkinnedMeshInfosFreeIndices;
+
+	/** Mapping from InstancedSkinnedMeshComponent (indicated by FInstancedSkinnedMeshComponentSharedDataKey) to corresponding SkinnedMeshInstanceVisualizationDescHandle */
+	TMap<FInstancedSkinnedMeshComponentSharedDataKey, FSkinnedMeshInstanceVisualizationDescHandle> InstancedSkinnedMeshComponentMap;
+
+	FMassInstancedSkinnedMeshComponentSharedDataMap InstancedSkinnedMeshComponentSharedData;
+
+	/** 
+	 * Mapping FMassSkinnedMeshInstanceVisualizationMeshDesc hash to FMassInstancedSkinnedMeshSharedData entries for all FMassSkinnedMeshInstanceVisualizationMeshDesc
+	 * that didn't come with Instanced Skinned Mesh Component explicitly provided. Used only for initialization.
+	 * Note that FMassSkinnedMeshInstanceVisualizationMeshDesc that were added with InstancedSkinnedMeshComponents provided directly
+	 * (via AddVisualDescWithComponents call) will never make it to this map.
+	 */
+	TMap<uint32, FInstancedSkinnedMeshComponentSharedDataKey> MeshDescToInstancedSkinnedMeshComponentMap;
+
+	/** Indices to InstancedSkinnedMeshInfos that need their Skinned Mesh Component constructed */
+	TArray<FSkinnedMeshInstanceVisualizationDescHandle> InstancedSkinnedMeshComponentsRequiringConstructing;
+
+	/** Resolved animation state per entity, populated during EndVisualChanges().
+	 * Keyed by entity handle — one entry per entity regardless of how many ISKMs it has.
+	 * Read by processors that need to push animation state to UAF during representation handoff. */
+	Experimental::TRobinHoodHashMap<FMassEntityHandle, FMassSkinnedMeshResolvedAnimState> ResolvedAnimStates;
+
+};
\ No newline at end of file
diff --git a/Engine/Plugins/Runtime/MassGameplay/Source/MassRepresentation/Public/MassVisualizationTrait.h b/Engine/Plugins/Runtime/MassGameplay/Source/MassRepresentation/Public/MassVisualizationTrait.h
--- a/Engine/Plugins/Runtime/MassGameplay/Source/MassRepresentation/Public/MassVisualizationTrait.h
+++ b/Engine/Plugins/Runtime/MassGameplay/Source/MassRepresentation/Public/MassVisualizationTrait.h
@@ -4,6 +4,7 @@
 
 #include "MassEntityTraitBase.h"
 #include "MassRepresentationTypes.h"
+#include "MassSkinnedMeshRepresentationTypes.h"
 #include "MassRepresentationFragments.h"
 #include "GameFramework/Actor.h"
 #include "MassVisualizationTrait.generated.h"
@@ -26,6 +27,10 @@ public:
 	UPROPERTY(EditAnywhere, Category = "Mass|Visual")
 	mutable FStaticMeshInstanceVisualizationDesc StaticMeshInstanceDesc;
 
+	/** Instanced skinned mesh information for this agent */
+	UPROPERTY(EditAnywhere, Category = "Mass|Visual")
+	mutable FSkinnedMeshInstanceVisualizationDesc SkinnedMeshInstanceDesc;
+
 	/** Actor class of this agent when spawned in high resolution*/
 	UPROPERTY(EditAnywhere, Category = "Mass|Visual")
 	TSubclassOf<AActor> HighResTemplateActor;
@@ -85,8 +90,18 @@ protected:
 	 * @param bStaticMeshDeterminedInvalid if StaticMeshInstanceDesc has already been determined invalid then bStaticMeshDeterminedInvalid
 	 *	can be set to `true` to skip the redundant check.
 	 */
+	UE_DEPRECATED(5.8, "Please use SanitizeMeshParams version for ISKM Support")
 	UE_API virtual void SanitizeParams(FMassRepresentationParameters& InOutParams, const bool bStaticMeshDeterminedInvalid = false) const;
 
+	/**
+	 * Tests whether StaticMeshInstanceDesc and SkinnedMeshInstanceDesc are valid and if not cleans up InOutParams for those occurrences.
+	 * @param bStaticMeshDeterminedInvalid if StaticMeshInstanceDesc has already been determined invalid then bStaticMeshDeterminedInvalid
+	 *	can be set to `true` to skip the redundant check.
+	 * @param bSkinnedMeshDeterminedInvalid if SkinnedMeshInstanceDesc has already been determined invalid then bSkinnedMeshDeterminedInvalid
+	 *	can be set to `true` to skip the redundant check.
+	 */
+	UE_API virtual void SanitizeMeshParams(FMassRepresentationParameters& InOutParams, const bool bStaticMeshDeterminedInvalid = false, const bool bSkinnedMeshDeterminedInvalid = false) const;
+
 	UE_API virtual void Serialize(FArchive& Ar) override;
 #if WITH_EDITOR
 	UE_API virtual void PostEditChangeProperty(struct FPropertyChangedEvent& PropertyChangedEvent) override;
