[Mass] Rename GetOrSpawnActorFromTemplate to GetOrRequestSpawnActorFromTemplate per [REDACTED_HANDLE]

The function requests a spawn rather than spawning synchronously. Renamed to reflect the actual behavior. Updated declaration, definition, and all callers.

#Synced-CL 52096882

diff --git a/Engine/Plugins/Runtime/MassGameplay/Source/MassRepresentation/Private/MassRepresentationActorManagement.cpp b/Engine/Plugins/Runtime/MassGameplay/Source/MassRepresentation/Private/MassRepresentationActorManagement.cpp
--- a/Engine/Plugins/Runtime/MassGameplay/Source/MassRepresentation/Private/MassRepresentationActorManagement.cpp
+++ b/Engine/Plugins/Runtime/MassGameplay/Source/MassRepresentation/Private/MassRepresentationActorManagement.cpp
@@ -21,7 +21,7 @@ AActor* UMassRepresentationActorManagement::GetOrSpawnActor(UMassRepresentationS
 {
 	TSharedRef<FMassEntityManager> SharedEntityManager = EntityManager.AsShared();
 
-	return RepresentationSubsystem.GetOrSpawnActorFromTemplate(MassAgent, Transform, TemplateActorIndex, InOutSpawnRequestHandle, Priority,
+	return RepresentationSubsystem.GetOrRequestSpawnActorFromTemplate(MassAgent, Transform, TemplateActorIndex, InOutSpawnRequestHandle, Priority,
 		FMassActorPreSpawnDelegate::CreateUObject(this, &UMassRepresentationActorManagement::OnPreActorSpawn, SharedEntityManager),
 		FMassActorPostSpawnDelegate::CreateUObject(this, &UMassRepresentationActorManagement::OnPostActorSpawn, SharedEntityManager));
 }
@@ -104,7 +104,7 @@ EMassActorSpawnRequestAction UMassRepresentationActorManagement::OnPostActorSpaw
 		}
 
 		// by default we keep the spawn request since there's code in other places that will consume this request by 
-		// calling RemoveActorSpawnRequest (like UMassRepresentationSubsystem::GetOrSpawnActorFromTemplate)
+		// calling RemoveActorSpawnRequest (like UMassRepresentationSubsystem::GetOrRequestSpawnActorFromTemplate)
 		return EMassActorSpawnRequestAction::Keep;
 	}
 
diff --git a/Engine/Plugins/Runtime/MassGameplay/Source/MassRepresentation/Private/MassRepresentationSubsystem.cpp b/Engine/Plugins/Runtime/MassGameplay/Source/MassRepresentation/Private/MassRepresentationSubsystem.cpp
--- a/Engine/Plugins/Runtime/MassGameplay/Source/MassRepresentation/Private/MassRepresentationSubsystem.cpp
+++ b/Engine/Plugins/Runtime/MassGameplay/Source/MassRepresentation/Private/MassRepresentationSubsystem.cpp
@@ -136,7 +136,7 @@ int16 UMassRepresentationSubsystem::FindOrAddTemplateActor(const TSubclassOf<AAc
 	return (int16)VisualIndex;
 }
 
-AActor* UMassRepresentationSubsystem::GetOrSpawnActorFromTemplate(const FMassEntityHandle MassAgent, const FTransform& Transform
+AActor* UMassRepresentationSubsystem::GetOrRequestSpawnActorFromTemplate(const FMassEntityHandle MassAgent, const FTransform& Transform
 	, const int16 TemplateActorIndex, FMassActorSpawnRequestHandle& InOutSpawnRequestHandle, float Priority
 	, FMassActorPreSpawnDelegate ActorPreSpawnDelegate, FMassActorPostSpawnDelegate ActorPostSpawnDelegate)
 {
diff --git a/Engine/Plugins/Runtime/MassGameplay/Source/MassRepresentation/Public/MassRepresentationSubsystem.h b/Engine/Plugins/Runtime/MassGameplay/Source/MassRepresentation/Public/MassRepresentationSubsystem.h
--- a/Engine/Plugins/Runtime/MassGameplay/Source/MassRepresentation/Public/MassRepresentationSubsystem.h
+++ b/Engine/Plugins/Runtime/MassGameplay/Source/MassRepresentation/Public/MassRepresentationSubsystem.h
@@ -95,9 +95,8 @@ public:
 	 * @param ActorPreSpawnDelegate is an optional delegate called before the spawning of an actor
 	 * @param ActorPostSpawnDelegate is an optional delegate called once the actor is spawned
 	 * @return The spawned actor from the template actor type if ready
-	 * @todo should be renamed to GetOrRequestSpawnActorFromTemplate
 	 */
-	MASSREPRESENTATION_API AActor* GetOrSpawnActorFromTemplate(const FMassEntityHandle MassAgent, const FTransform& Transform, const int16 TemplateActorIndex, FMassActorSpawnRequestHandle& InOutSpawnRequestHandle, float Priority = MAX_FLT,
+	MASSREPRESENTATION_API AActor* GetOrRequestSpawnActorFromTemplate(const FMassEntityHandle MassAgent, const FTransform& Transform, const int16 TemplateActorIndex, FMassActorSpawnRequestHandle& InOutSpawnRequestHandle, float Priority = MAX_FLT,
 		FMassActorPreSpawnDelegate ActorPreSpawnDelegate = FMassActorPreSpawnDelegate(), FMassActorPostSpawnDelegate ActorPostSpawnDelegate = FMassActorPostSpawnDelegate());
 
 	/**
@@ -240,6 +239,17 @@ protected:
 	/** Keeping track of all the mass agent this subsystem is responsible for spawning actors */
 	TMap<FMassEntityHandle, int32> HandledMassAgents;
 
+	//-----------------------------------------------------------------------------
+	// DEPRECATED
+	//-----------------------------------------------------------------------------
+public:
+	UE_DEPRECATED(5.8, "Use GetOrRequestSpawnActorFromTemplate instead.")
+	AActor* GetOrSpawnActorFromTemplate(const FMassEntityHandle MassAgent, const FTransform& Transform, const int16 TemplateActorIndex, FMassActorSpawnRequestHandle& InOutSpawnRequestHandle, float Priority = MAX_FLT,
+		FMassActorPreSpawnDelegate ActorPreSpawnDelegate = FMassActorPreSpawnDelegate(), FMassActorPostSpawnDelegate ActorPostSpawnDelegate = FMassActorPostSpawnDelegate())
+	{
+		return GetOrRequestSpawnActorFromTemplate(MassAgent, Transform, TemplateActorIndex, InOutSpawnRequestHandle, Priority, MoveTemp(ActorPreSpawnDelegate), MoveTemp(ActorPostSpawnDelegate));
+	}
+
 private:
 	UFUNCTION()
 	MASSREPRESENTATION_API void HandleVisualizerEndPlay(AActor* Actor, const EEndPlayReason::Type EndPlayReason);
