[Mass] fixed inverted booleans in `SanitizeMeshParams` call in MassVisualizationTrait
`bStaticMeshDescriptionValid` (true=valid) was passed where `bStaticMeshDeterminedInvalid` (true=invalid) was expected. When static mesh was valid but skinned mesh wasn't, we'd enter the else branch and accidentally tell `SanitizeMeshParams` to wipe `StaticMeshInstance` entries from `LODRepresentation`, breaking ISM rendering on the client for replicated crowd entities.

#rnx
#Synced-CL 52124835

diff --git a/Engine/Plugins/Runtime/MassGameplay/Source/MassRepresentation/Private/MassVisualizationTrait.cpp b/Engine/Plugins/Runtime/MassGameplay/Source/MassRepresentation/Private/MassVisualizationTrait.cpp
--- a/Engine/Plugins/Runtime/MassGameplay/Source/MassRepresentation/Private/MassVisualizationTrait.cpp
+++ b/Engine/Plugins/Runtime/MassGameplay/Source/MassRepresentation/Private/MassVisualizationTrait.cpp
@@ -133,7 +133,7 @@ void UMassVisualizationTrait::BuildTemplate(FMassEntityTemplateBuildContext& Bui
 	else
 	{
 		FMassRepresentationParameters ParamsCopy = Params;
-		SanitizeMeshParams(ParamsCopy, bStaticMeshDescriptionValid, bSkinnedMeshDescriptionValid);
+		SanitizeMeshParams(ParamsCopy, !bStaticMeshDescriptionValid, !bSkinnedMeshDescriptionValid);
 		ParamsFragment = EntityManager.GetOrCreateConstSharedFragment(ParamsCopy);
 	}
 	ParamsFragment.Get<const FMassRepresentationParameters>().ComputeCachedValues();
