[Mass][ISKM] Add support for setting an AnimationMinScreenSize override on an ISKM description

#Synced-CL 52216836

diff --git a/Engine/Plugins/Runtime/MassGameplay/Source/MassRepresentation/Private/MassVisualizationComponent.cpp b/Engine/Plugins/Runtime/MassGameplay/Source/MassRepresentation/Private/MassVisualizationComponent.cpp
--- a/Engine/Plugins/Runtime/MassGameplay/Source/MassRepresentation/Private/MassVisualizationComponent.cpp
+++ b/Engine/Plugins/Runtime/MassGameplay/Source/MassRepresentation/Private/MassVisualizationComponent.cpp
@@ -1328,6 +1328,7 @@ void UMassVisualizationComponent::ConstructSkinnedMeshComponents()
 				}
 
 				InstancedSkinnedMeshComponent->SetCullDistances(0, 1000000); // @todo: Need to figure out what to do here, either LOD or cull distances.
+				InstancedSkinnedMeshComponent->SetAnimationMinScreenSize(InstancedSkinnedMeshDesc.AnimationMinScreenSize);
 				InstancedSkinnedMeshComponent->SetupAttachment(ActorOwner->GetRootComponent());
 				InstancedSkinnedMeshComponent->SetCanEverAffectNavigation(false);
 				InstancedSkinnedMeshComponent->SetCollisionProfileName(UCollisionProfile::NoCollision_ProfileName);
diff --git a/Engine/Plugins/Runtime/MassGameplay/Source/MassRepresentation/Public/MassSkinnedMeshRepresentationTypes.h b/Engine/Plugins/Runtime/MassGameplay/Source/MassRepresentation/Public/MassSkinnedMeshRepresentationTypes.h
--- a/Engine/Plugins/Runtime/MassGameplay/Source/MassRepresentation/Public/MassSkinnedMeshRepresentationTypes.h
+++ b/Engine/Plugins/Runtime/MassGameplay/Source/MassRepresentation/Public/MassSkinnedMeshRepresentationTypes.h
@@ -92,6 +92,14 @@ struct FMassSkinnedMeshInstanceVisualizationMeshDesc
 	UPROPERTY(EditAnywhere, Category = "Mass|Visual")
 	TSubclassOf<UInstancedSkinnedMeshComponent> InstancedSkinnedMeshComponentClass;
 
+	/**
+	 * Passed to UInstancedSkinnedMeshComponent::SetAnimationMinScreenSize.
+	 * Values: < 0 = always animate, 0 = use r.Skinning.DefaultAnimationMinScreenSize, > 0 = custom threshold.
+	 * @see UInstancedSkinnedMeshComponent::SetAnimationMinScreenSize
+	 */
+	UPROPERTY(EditAnywhere, Category = "Mass|Visual")
+	float AnimationMinScreenSize = 0.0f;
+
 	bool operator==(const FMassSkinnedMeshInstanceVisualizationMeshDesc& Other) const
 	{		
 		return Asset == Other.Asset && 
@@ -103,7 +111,8 @@ struct FMassSkinnedMeshInstanceVisualizationMeshDesc
 			bCastShadows == Other.bCastShadows && 
 			Mobility == Other.Mobility &&
 			bRequiresExternalInstanceIDTracking == Other.bRequiresExternalInstanceIDTracking &&
-			LocalTransform.Equals(Other.LocalTransform);
+			LocalTransform.Equals(Other.LocalTransform) &&
+			AnimationMinScreenSize == Other.AnimationMinScreenSize;
 	}
 	
 	friend inline uint32 GetTypeHash(const FMassSkinnedMeshInstanceVisualizationMeshDesc& MeshDesc)
@@ -116,6 +125,7 @@ struct FMassSkinnedMeshInstanceVisualizationMeshDesc
 		Hash = HashCombine(GetTypeHash(MeshDesc.Mobility), Hash);
 		Hash = HashCombine(GetTypeHash(MeshDesc.LocalTransform), Hash);
 		Hash = HashCombine(GetTypeHash(MeshDesc.bRequiresExternalInstanceIDTracking), Hash);
+		Hash = HashCombine(GetTypeHash(MeshDesc.AnimationMinScreenSize), Hash);
 		for (UMaterialInterface* MaterialOverride : MeshDesc.MaterialOverrides)
 		{
 			if (MaterialOverride)
