[Mass] Fix MassAgentComponent assert when Iris destroys puppet in PuppetPendingReplication

`UnregisterAgentComponent` had no handler for `PuppetPendingReplication`. When Iris destroys a replicated actor before its NetID gets replicated (e.g. relevancy change), `MassAgentComponent` is stuck in `PuppetPendingReplication` with no entity handle and no NetID - `ShutdownAgentComponent` just did nothing for this state, leaving the assert to fire. Added `PuppetReplicationAborted()` to cleanly bail out to `None`.

#Synced-CL 52124835

diff --git a/Engine/Plugins/Runtime/MassGameplay/Source/MassActors/Private/MassAgentComponent.cpp b/Engine/Plugins/Runtime/MassGameplay/Source/MassActors/Private/MassAgentComponent.cpp
--- a/Engine/Plugins/Runtime/MassGameplay/Source/MassActors/Private/MassAgentComponent.cpp
+++ b/Engine/Plugins/Runtime/MassGameplay/Source/MassActors/Private/MassAgentComponent.cpp
@@ -309,6 +309,14 @@ void UMassAgentComponent::PuppetInitializationAborted()
 	SwitchToState(EAgentComponentState::PuppetPaused);
 }
 
+void UMassAgentComponent::PuppetReplicationAborted()
+{
+	MASSAGENT_CHECK(State == EAgentComponentState::PuppetPendingReplication,
+		TEXT("%s is expecting to be in state[PuppetPendingReplication] but is in %s"), ANSI_TO_TCHAR(__FUNCTION__), *UEnum::GetValueAsString(State));
+	// No Mass entity exists yet in this state (no handle, no NetID). Just reset.
+	SwitchToState(EAgentComponentState::None);
+}
+
 void UMassAgentComponent::ClearEntityHandle()
 {
 	MASSAGENT_CHECK(State == EAgentComponentState::EntityCreated,
diff --git a/Engine/Plugins/Runtime/MassGameplay/Source/MassActors/Private/MassAgentSubsystem.cpp b/Engine/Plugins/Runtime/MassGameplay/Source/MassActors/Private/MassAgentSubsystem.cpp
--- a/Engine/Plugins/Runtime/MassGameplay/Source/MassActors/Private/MassAgentSubsystem.cpp
+++ b/Engine/Plugins/Runtime/MassGameplay/Source/MassActors/Private/MassAgentSubsystem.cpp
@@ -172,7 +172,14 @@ void UMassAgentSubsystem::UnregisterAgentComponent(UMassAgentComponent& AgentCom
 		return;
 	}
 
-	if (AgentComp.IsPuppet())
+	if (AgentComp.IsPuppetReplicationPending())
+	{
+		// The actor was destroyed (e.g., by Iris relevancy) before its NetID was replicated
+		// via the old bubble system. No Mass entity exists yet, no cleanup needed —
+		// just transition to None.
+		AgentComp.PuppetReplicationAborted();
+	}
+	else if (AgentComp.IsPuppet())
 	{
 		if (AgentComp.IsPuppetPendingInitialization())
 		{
diff --git a/Engine/Plugins/Runtime/MassGameplay/Source/MassActors/Public/MassAgentComponent.h b/Engine/Plugins/Runtime/MassGameplay/Source/MassActors/Public/MassAgentComponent.h
--- a/Engine/Plugins/Runtime/MassGameplay/Source/MassActors/Public/MassAgentComponent.h
+++ b/Engine/Plugins/Runtime/MassGameplay/Source/MassActors/Public/MassAgentComponent.h
@@ -80,6 +80,8 @@ public:
 	UE_API void PuppetInitializationPending();
 	UE_API void PuppetInitializationDone();
 	UE_API void PuppetInitializationAborted();
+	/** Called when a puppet actor is destroyed before its NetID was replicated (e.g., Iris relevancy). */
+	UE_API void PuppetReplicationAborted();
 	UE_API void PuppetUnregistrationDone();
 	bool IsPuppetPendingInitialization() const { return (State == EAgentComponentState::PuppetPendingInitialization); }
 	bool IsPuppetReplicationPending() const { return (State == EAgentComponentState::PuppetPendingReplication); }
