[Mass] Remove redundant validity check

These checks are already done as part of the initialization of the boolean
parameters, and checking them again here means that overrides can not change
this behaviour.

#Synced-CL 52345207

diff --git a/Engine/Plugins/Runtime/MassGameplay/Source/MassRepresentation/Private/MassVisualizationTrait.cpp b/Engine/Plugins/Runtime/MassGameplay/Source/MassRepresentation/Private/MassVisualizationTrait.cpp
--- a/Engine/Plugins/Runtime/MassGameplay/Source/MassRepresentation/Private/MassVisualizationTrait.cpp
+++ b/Engine/Plugins/Runtime/MassGameplay/Source/MassRepresentation/Private/MassVisualizationTrait.cpp
@@ -160,7 +160,7 @@ void UMassVisualizationTrait::SanitizeParams(FMassRepresentationParameters& InOu
 
 void UMassVisualizationTrait::SanitizeMeshParams(FMassRepresentationParameters& InOutParams, const bool bStaticMeshDeterminedInvalid /*= false*/, const bool bSkinnedMeshDeterminedInvalid /*= false*/) const
 {
-	if (bStaticMeshDeterminedInvalid || (StaticMeshInstanceDesc.IsValid() == false))
+	if (bStaticMeshDeterminedInvalid)
 	{
 		for (int32 LODIndex = 0; LODIndex < EMassLOD::Max; ++LODIndex)
 		{
@@ -171,7 +171,7 @@ void UMassVisualizationTrait::SanitizeMeshParams(FMassRepresentationParameters&
 		}
 	}
 
-	if (bSkinnedMeshDeterminedInvalid || (SkinnedMeshInstanceDesc.IsValid() == false))
+	if (bSkinnedMeshDeterminedInvalid)
 	{
 		for (int32 LODIndex = 0; LODIndex < EMassLOD::Max; ++LODIndex)
 		{
