[Mass][Animation] ISKM Support For ISKMs with Valid Transforms but not Animation Data

#rnx
#Synced-CL 52452884

diff --git a/Engine/Plugins/Runtime/MassGameplay/Source/MassRepresentation/Private/MassVisualizationComponent.cpp b/Engine/Plugins/Runtime/MassGameplay/Source/MassRepresentation/Private/MassVisualizationComponent.cpp
--- a/Engine/Plugins/Runtime/MassGameplay/Source/MassRepresentation/Private/MassVisualizationComponent.cpp
+++ b/Engine/Plugins/Runtime/MassGameplay/Source/MassRepresentation/Private/MassVisualizationComponent.cpp
@@ -751,7 +751,7 @@ void UMassVisualizationComponent::EndVisualChanges()
 							{
 								InstancedSkinnedMeshComponent->SetInstancePrevTransform(Id, PrevInstanceTransforms[IDIndex]);
 							}
-							if (!InstanceAnimationData.IsEmpty())
+							if (InstanceAnimationData.IsValidIndex(IDIndex))
 							{
 								int DesiredAnimIndex = -1;
 								int CurrentAnimIndex = -1;
@@ -840,7 +840,7 @@ void UMassVisualizationComponent::EndVisualChanges()
 						{
 							InstancedSkinnedMeshComponent->SetCustomData(Id, MakeArrayView(CustomDataFloats.GetData() + AddItem.IDIndex * NumCustomDataFloats, NumCustomDataFloats));
 						}
-						if (!InstanceAnimationData.IsEmpty())
+						if (InstanceAnimationData.IsValidIndex(AddItem.IDIndex))
 						{
 							int DesiredAnimIndex = -1;
 							int CurrentAnimIndex = -1;
