[PLAY-115634] Fix stale ISM handle crash in MassStationaryISMRepresentationFragmentDestructor

- Add early-out when ISM/skinned mesh info arrays are empty
- Replace ensureMsgf with UE_LOGF warning on invalid handle indices
so cleanup continues gracefully instead of firing ensure reports

https://epicgames.atlassian.net/browse/PLAY-115634

#rnx
#Synced-CL 52581472

diff --git a/Engine/Plugins/Runtime/MassGameplay/Source/MassRepresentation/Private/MassStationaryISMRepresentationFragmentDestructor.cpp b/Engine/Plugins/Runtime/MassGameplay/Source/MassRepresentation/Private/MassStationaryISMRepresentationFragmentDestructor.cpp
--- a/Engine/Plugins/Runtime/MassGameplay/Source/MassRepresentation/Private/MassStationaryISMRepresentationFragmentDestructor.cpp
+++ b/Engine/Plugins/Runtime/MassGameplay/Source/MassRepresentation/Private/MassStationaryISMRepresentationFragmentDestructor.cpp
@@ -37,6 +37,11 @@ void UMassStationaryISMRepresentationFragmentDestructor::Execute(FMassEntityMana
 		FMassInstancedStaticMeshInfoArrayView ISMInfosView = RepresentationSubsystem->GetMutableInstancedStaticMeshInfos();
 		FMassInstancedSkinnedMeshInfoArrayView MeshInfosView = RepresentationSubsystem->GetMutableInstancedSkinnedMeshInfos();
 
+		if (ISMInfosView.Num() == 0 && MeshInfosView.Num() == 0)
+		{
+			return;
+		}
+
 		const TArrayView<FMassRepresentationFragment> RepresentationList = Context.GetMutableFragmentView<FMassRepresentationFragment>();
 
 		for (FMassExecutionContext::FEntityIterator EntityIt = Context.CreateEntityIterator(); EntityIt; ++EntityIt)
@@ -45,16 +50,19 @@ void UMassStationaryISMRepresentationFragmentDestructor::Execute(FMassEntityMana
 			if (Representation.CurrentRepresentation == EMassRepresentationType::StaticMeshInstance)
 			{
 				const int32 ISMInfoIndex = Representation.StaticMeshDescHandle.ToIndex();
-				if (ensureMsgf(ISMInfosView.IsValidIndex(ISMInfoIndex), TEXT("Invalid handle index %u for ISMInfosView"), ISMInfoIndex))
+				if (!ISMInfosView.IsValidIndex(ISMInfoIndex))
 				{
-					FMassInstancedStaticMeshInfo& ISMInfo = ISMInfosView[ISMInfoIndex];
-					if (FMassLODSignificanceRange* OldRange = ISMInfo.GetLODSignificanceRange(Representation.PrevLODSignificance))
+					UE_LOGF(LogMassRepresentation, Warning, "Invalid handle index %u for ISMInfosView", ISMInfoIndex);
+					Representation.CurrentRepresentation = EMassRepresentationType::None;
+					continue;
+				}
+				FMassInstancedStaticMeshInfo& ISMInfo = ISMInfosView[ISMInfoIndex];
+				if (FMassLODSignificanceRange* OldRange = ISMInfo.GetLODSignificanceRange(Representation.PrevLODSignificance))
+				{
+					const FMassEntityHandle EntityHandle = Context.GetEntity(EntityIt);
+					if (OldRange)
 					{
-						const FMassEntityHandle EntityHandle = Context.GetEntity(EntityIt);
-						if (OldRange)
-						{
-							OldRange->RemoveInstance(EntityHandle);
-						}
+						OldRange->RemoveInstance(EntityHandle);
 					}
 				}
 				Representation.CurrentRepresentation = EMassRepresentationType::None;
@@ -62,16 +70,19 @@ void UMassStationaryISMRepresentationFragmentDestructor::Execute(FMassEntityMana
 			else if (Representation.CurrentRepresentation == EMassRepresentationType::SkinnedMeshInstance)
 			{
 				const int32 MeshInfoIndex = Representation.SkinnedMeshDescHandle.ToIndex();
-				if (ensureMsgf(MeshInfosView.IsValidIndex(MeshInfoIndex), TEXT("Invalid handle index %u for MeshInfosView"), MeshInfoIndex))
+				if (!MeshInfosView.IsValidIndex(MeshInfoIndex))
+				{
+					UE_LOGF(LogMassRepresentation, Warning, "Invalid handle index %u for MeshInfosView", MeshInfoIndex);
+					Representation.CurrentRepresentation = EMassRepresentationType::None;
+					continue;
+				}
+				FMassInstancedSkinnedMeshInfo& MeshInfo = MeshInfosView[MeshInfoIndex];
+				if (FMassLODInstancedSkinnedMeshSignificanceRange* OldRange = MeshInfo.GetLODSignificanceRange(Representation.PrevLODSignificance))
 				{
-					FMassInstancedSkinnedMeshInfo& MeshInfo = MeshInfosView[MeshInfoIndex];
-					if (FMassLODInstancedSkinnedMeshSignificanceRange* OldRange = MeshInfo.GetLODSignificanceRange(Representation.PrevLODSignificance))
+					const FMassEntityHandle EntityHandle = Context.GetEntity(EntityIt);
+					if (OldRange)
 					{
-						const FMassEntityHandle EntityHandle = Context.GetEntity(EntityIt);
-						if (OldRange)
-						{
-							OldRange->RemoveInstance(EntityHandle);
-						}
+						OldRange->RemoveInstance(EntityHandle);
 					}
 				}
 				Representation.CurrentRepresentation = EMassRepresentationType::None;
