[Backout] - CL52693199
#rnx
Original CL Desc
-----------------------------------------------------------------
Move trajectory and spring movement code to engine

Creates two new modules within the MassGameplay plugin:
- MassCharacterTrajectory: trajectory generation, trajectory movement, debug visualization
- Spring movement processors and fragments added to MassMovement

New processor architecture:
- USpringUpdateProcessor: runs spring damper, updates FSpringMovementRuntime (no movement)
- USpringMovementApplyProcessor: applies spring state to transform (skipped when FMassCustomMovementTag present)
- UMassSpringMovementToTrajectoryProcessor: generates trajectory from spring state
- UMassMovementToTrajectoryProcessor: generates trajectory from desired velocity (no spring)
- UTrajectoryToMovementProcessor: moves entity along generated trajectory
- UMassPoseTrajectoryDebugProcessor: debug visualization

Execution order: SpringUpdate -> Trajectory -> Movement/SpringApply -> TrajectoryMovement
Trajectory generation always runs before movement for consistent start-of-frame sampling.

Mass Traits:
- USpringMovementTrait: adds spring-damped movement
- UTrajectoryTrait: adds trajectory generation
- UTrajectoryMovementTrait: overrides movement to follow trajectory, sets FMassCustomMovementTag
- ValidateTemplate checks: TrajectoryMovement requires TrajectoryTrait, warns on conflicting movement

Tests: 9 automated tests in MassGameplayTestSuite covering sample layout, variable
frame rate history, processor exclusion, spring trajectory, zero DT safety, and
trajectory movement.

diff --git a/Engine/Plugins/Runtime/MassGameplay/MassGameplay.uplugin b/Engine/Plugins/Runtime/MassGameplay/MassGameplay.uplugin
--- a/Engine/Plugins/Runtime/MassGameplay/MassGameplay.uplugin
+++ b/Engine/Plugins/Runtime/MassGameplay/MassGameplay.uplugin
@@ -48,10 +48,6 @@
 			"Name": "MassMovement",
 			"Type": "Runtime"
 		},
-		{
-			"Name": "MassCharacterTrajectory",
-			"Type": "Runtime"
-		},
 		{
 			"Name": "MassReplication",
 			"Type": "Runtime"
@@ -88,10 +84,6 @@
 			"Name": "ZoneGraph",
 			"Enabled": true
 		},
-		{
-			"Name": "PoseSearch",
-			"Enabled": true
-		},
 		{
 			"Name": "ZoneGraphAnnotations",
 			"Enabled": true
diff --git a/Engine/Plugins/Runtime/MassGameplay/Source/MassCharacterTrajectory/MassCharacterTrajectory.Build.cs b/Engine/Plugins/Runtime/MassGameplay/Source/MassCharacterTrajectory/MassCharacterTrajectory.Build.cs
deleted file mode 100644
diff --git a/Engine/Plugins/Runtime/MassGameplay/Source/MassCharacterTrajectory/Private/MassCharacterTrajectoryDebugProcessor.cpp b/Engine/Plugins/Runtime/MassGameplay/Source/MassCharacterTrajectory/Private/MassCharacterTrajectoryDebugProcessor.cpp
deleted file mode 100644
diff --git a/Engine/Plugins/Runtime/MassGameplay/Source/MassCharacterTrajectory/Private/MassCharacterTrajectoryFragments.cpp b/Engine/Plugins/Runtime/MassGameplay/Source/MassCharacterTrajectory/Private/MassCharacterTrajectoryFragments.cpp
deleted file mode 100644
diff --git a/Engine/Plugins/Runtime/MassGameplay/Source/MassCharacterTrajectory/Private/MassCharacterTrajectoryGenerationProcessors.cpp b/Engine/Plugins/Runtime/MassGameplay/Source/MassCharacterTrajectory/Private/MassCharacterTrajectoryGenerationProcessors.cpp
deleted file mode 100644
diff --git a/Engine/Plugins/Runtime/MassGameplay/Source/MassCharacterTrajectory/Private/MassCharacterTrajectoryModule.cpp b/Engine/Plugins/Runtime/MassGameplay/Source/MassCharacterTrajectory/Private/MassCharacterTrajectoryModule.cpp
deleted file mode 100644
diff --git a/Engine/Plugins/Runtime/MassGameplay/Source/MassCharacterTrajectory/Private/MassCharacterTrajectoryMovementProcessor.cpp b/Engine/Plugins/Runtime/MassGameplay/Source/MassCharacterTrajectory/Private/MassCharacterTrajectoryMovementProcessor.cpp
deleted file mode 100644
diff --git a/Engine/Plugins/Runtime/MassGameplay/Source/MassCharacterTrajectory/Private/MassCharacterTrajectoryTrait.cpp b/Engine/Plugins/Runtime/MassGameplay/Source/MassCharacterTrajectory/Private/MassCharacterTrajectoryTrait.cpp
deleted file mode 100644
diff --git a/Engine/Plugins/Runtime/MassGameplay/Source/MassCharacterTrajectory/Private/MassCharacterTrajectoryTypes.cpp b/Engine/Plugins/Runtime/MassGameplay/Source/MassCharacterTrajectory/Private/MassCharacterTrajectoryTypes.cpp
deleted file mode 100644
diff --git a/Engine/Plugins/Runtime/MassGameplay/Source/MassCharacterTrajectory/Public/IMassCharacterTrajectoryModule.h b/Engine/Plugins/Runtime/MassGameplay/Source/MassCharacterTrajectory/Public/IMassCharacterTrajectoryModule.h
deleted file mode 100644
diff --git a/Engine/Plugins/Runtime/MassGameplay/Source/MassCharacterTrajectory/Public/MassCharacterTrajectoryDebugProcessor.h b/Engine/Plugins/Runtime/MassGameplay/Source/MassCharacterTrajectory/Public/MassCharacterTrajectoryDebugProcessor.h
deleted file mode 100644
diff --git a/Engine/Plugins/Runtime/MassGameplay/Source/MassCharacterTrajectory/Public/MassCharacterTrajectoryFragments.h b/Engine/Plugins/Runtime/MassGameplay/Source/MassCharacterTrajectory/Public/MassCharacterTrajectoryFragments.h
deleted file mode 100644
diff --git a/Engine/Plugins/Runtime/MassGameplay/Source/MassCharacterTrajectory/Public/MassCharacterTrajectoryGenerationProcessors.h b/Engine/Plugins/Runtime/MassGameplay/Source/MassCharacterTrajectory/Public/MassCharacterTrajectoryGenerationProcessors.h
deleted file mode 100644
diff --git a/Engine/Plugins/Runtime/MassGameplay/Source/MassCharacterTrajectory/Public/MassCharacterTrajectoryMovementProcessor.h b/Engine/Plugins/Runtime/MassGameplay/Source/MassCharacterTrajectory/Public/MassCharacterTrajectoryMovementProcessor.h
deleted file mode 100644
diff --git a/Engine/Plugins/Runtime/MassGameplay/Source/MassCharacterTrajectory/Public/MassCharacterTrajectoryTrait.h b/Engine/Plugins/Runtime/MassGameplay/Source/MassCharacterTrajectory/Public/MassCharacterTrajectoryTrait.h
deleted file mode 100644
diff --git a/Engine/Plugins/Runtime/MassGameplay/Source/MassCharacterTrajectory/Public/MassCharacterTrajectoryTypes.h b/Engine/Plugins/Runtime/MassGameplay/Source/MassCharacterTrajectory/Public/MassCharacterTrajectoryTypes.h
deleted file mode 100644
diff --git a/Engine/Plugins/Runtime/MassGameplay/Source/MassGameplayTestSuite/MassGameplayTestSuite.Build.cs b/Engine/Plugins/Runtime/MassGameplay/Source/MassGameplayTestSuite/MassGameplayTestSuite.Build.cs
--- a/Engine/Plugins/Runtime/MassGameplay/Source/MassGameplayTestSuite/MassGameplayTestSuite.Build.cs
+++ b/Engine/Plugins/Runtime/MassGameplay/Source/MassGameplayTestSuite/MassGameplayTestSuite.Build.cs
@@ -24,11 +24,7 @@ namespace UnrealBuildTool.Rules
 					"MassCore",
 					"MassEntity",
 					"MassEntityTestSuite",
-					"MassLOD",
-					"MassMovement",
 					"MassSpawner",
-					"MassCharacterTrajectory",
-					"PoseSearch",
 				}
 			);
 		}
diff --git a/Engine/Plugins/Runtime/MassGameplay/Source/MassGameplayTestSuite/Private/MassCharacterTrajectoryGenerationTest.cpp b/Engine/Plugins/Runtime/MassGameplay/Source/MassGameplayTestSuite/Private/MassCharacterTrajectoryGenerationTest.cpp
deleted file mode 100644
diff --git a/Engine/Plugins/Runtime/MassGameplay/Source/MassMovement/Private/MassMovementProcessors.cpp b/Engine/Plugins/Runtime/MassGameplay/Source/MassMovement/Private/MassMovementProcessors.cpp
--- a/Engine/Plugins/Runtime/MassGameplay/Source/MassMovement/Private/MassMovementProcessors.cpp
+++ b/Engine/Plugins/Runtime/MassGameplay/Source/MassMovement/Private/MassMovementProcessors.cpp
@@ -14,8 +14,6 @@
 //  UMassApplyForceProcessor
 //----------------------------------------------------------------------//
 
-#include "MassSpringMovementFragments.h"
-
 #include UE_INLINE_GENERATED_CPP_BY_NAME(MassMovementProcessors)
 
 UMassApplyForceProcessor::UMassApplyForceProcessor()
@@ -82,8 +80,6 @@ void UMassApplyMovementProcessor::ConfigureQueries(const TSharedRef<FMassEntityM
 	EntityQuery.AddRequirement<FMassVelocityFragment>(EMassFragmentAccess::ReadWrite);
 	EntityQuery.AddRequirement<FTransformFragment>(EMassFragmentAccess::ReadWrite);
 	EntityQuery.AddRequirement<FMassDesiredMovementFragment>(EMassFragmentAccess::ReadOnly);
-	EntityQuery.AddConstSharedRequirement<FSpringMovementSettings>(EMassFragmentPresence::None);	// Do not apply simple movement processor if spring is being used
-	EntityQuery.AddTagRequirement<FMassCustomMovementTag>(EMassFragmentPresence::None);	// Do not apply if a custom movement processor handles this entity
 	EntityQuery.AddTagRequirement<FMassOffLODTag>(EMassFragmentPresence::None);
 	EntityQuery.AddTagRequirement<FMassCodeDrivenMovementTag>(EMassFragmentPresence::All);
 
diff --git a/Engine/Plugins/Runtime/MassGameplay/Source/MassMovement/Private/MassMovementTypes.cpp b/Engine/Plugins/Runtime/MassGameplay/Source/MassMovement/Private/MassMovementTypes.cpp
--- a/Engine/Plugins/Runtime/MassGameplay/Source/MassMovement/Private/MassMovementTypes.cpp
+++ b/Engine/Plugins/Runtime/MassGameplay/Source/MassMovement/Private/MassMovementTypes.cpp
@@ -5,11 +5,6 @@
 
 #include UE_INLINE_GENERATED_CPP_BY_NAME(MassMovementTypes)
 
-namespace UE::Mass::ProcessorGroupNames
-{
-	const FName SpringMovement = TEXT("SpringMovement");
-}
-
 namespace UE::MassMovement
 {
 	int32 bFreezeMovement = 0;
diff --git a/Engine/Plugins/Runtime/MassGameplay/Source/MassMovement/Private/MassSpringMovementProcessors.cpp b/Engine/Plugins/Runtime/MassGameplay/Source/MassMovement/Private/MassSpringMovementProcessors.cpp
deleted file mode 100644
diff --git a/Engine/Plugins/Runtime/MassGameplay/Source/MassMovement/Private/MassSpringMovementTrait.cpp b/Engine/Plugins/Runtime/MassGameplay/Source/MassMovement/Private/MassSpringMovementTrait.cpp
deleted file mode 100644
diff --git a/Engine/Plugins/Runtime/MassGameplay/Source/MassMovement/Public/MassMovementFragments.h b/Engine/Plugins/Runtime/MassGameplay/Source/MassMovement/Public/MassMovementFragments.h
--- a/Engine/Plugins/Runtime/MassGameplay/Source/MassMovement/Public/MassMovementFragments.h
+++ b/Engine/Plugins/Runtime/MassGameplay/Source/MassMovement/Public/MassMovementFragments.h
@@ -66,17 +66,6 @@ struct FMassCodeDrivenMovementTag : public FMassTag
 	GENERATED_BODY()
 };
 
-/**
- * The presence of this tag indicates that a custom movement processor handles
- * this entity's position updates. Default movement processors (UMassApplyMovementProcessor,
- * USpringApplyMovementProcessor) will skip entities with this tag.
- */
-USTRUCT()
-struct FMassCustomMovementTag : public FMassTag
-{
-	GENERATED_BODY()
-};
-
 /** Parameters describing how this mass agent should move */
 USTRUCT()
 struct FMassMovementParameters : public FMassConstSharedFragment
diff --git a/Engine/Plugins/Runtime/MassGameplay/Source/MassMovement/Public/MassMovementTypes.h b/Engine/Plugins/Runtime/MassGameplay/Source/MassMovement/Public/MassMovementTypes.h
--- a/Engine/Plugins/Runtime/MassGameplay/Source/MassMovement/Public/MassMovementTypes.h
+++ b/Engine/Plugins/Runtime/MassGameplay/Source/MassMovement/Public/MassMovementTypes.h
@@ -3,14 +3,8 @@
 #pragma once
 
 #include "CoreTypes.h"
-#include "UObject/NameTypes.h"
 #include "MassMovementTypes.generated.h"
 
-namespace UE::Mass::ProcessorGroupNames
-{
-	extern MASSMOVEMENT_API const FName SpringMovement;
-}
-
 namespace UE::MassMovement
 {
 	extern MASSMOVEMENT_API int32 bFreezeMovement;
diff --git a/Engine/Plugins/Runtime/MassGameplay/Source/MassMovement/Public/MassSpringMovementFragments.h b/Engine/Plugins/Runtime/MassGameplay/Source/MassMovement/Public/MassSpringMovementFragments.h
deleted file mode 100644
diff --git a/Engine/Plugins/Runtime/MassGameplay/Source/MassMovement/Public/MassSpringMovementProcessors.h b/Engine/Plugins/Runtime/MassGameplay/Source/MassMovement/Public/MassSpringMovementProcessors.h
deleted file mode 100644
diff --git a/Engine/Plugins/Runtime/MassGameplay/Source/MassMovement/Public/MassSpringMovementTrait.h b/Engine/Plugins/Runtime/MassGameplay/Source/MassMovement/Public/MassSpringMovementTrait.h
deleted file mode 100644
